Add MaterialSubSurfaceType to Substrate's Slab node.

This CL changes how subsurface is configured per slab. The user now explicit selects which SSS technique is used.
Follow up to not clutter the CL: rename SSS_TYPE_INVALID into SSS_TYPE_NONE.

diff --git a/Engine/Shaders/Private/DebugProbes.usf b/Engine/Shaders/Private/DebugProbes.usf
--- a/Engine/Shaders/Private/DebugProbes.usf
+++ b/Engine/Shaders/Private/DebugProbes.usf
@@ -129,7 +129,7 @@ void WriteOutMaterial(
 		/*float3 SSSMFP							*/0.0f,
 		/*float SSSMFPScale						*/0.0f,
 		/*float SSSPhaseAniso					*/0.0f,
-		/*float UseSSSDiffuse					*/0.0f,
+		/*float SSSType							*/SSS_TYPE_INVALID,
 		/*float3 Emissive						*/0.0f,
 		/*float SecondRoughness					*/0.0f,
 		/*float SecondRoughnessWeight			*/0.0f,
diff --git a/Engine/Shaders/Private/DefaultSSRTiles.usf b/Engine/Shaders/Private/DefaultSSRTiles.usf
--- a/Engine/Shaders/Private/DefaultSSRTiles.usf
+++ b/Engine/Shaders/Private/DefaultSSRTiles.usf
@@ -64,7 +64,7 @@ bool DoesPixelContainSSR(uint2 PixelPos)
 			&& !SubstratePixelHeader.IsHair()
 			&& !SubstratePixelHeader.IsSingleLayerWater();
 
-		bContainValidPixels = bContainValidPixels && (bEnableTwoSidedFoliage ? true : !BSDF_GETISTHIN(BSDF));
+		bContainValidPixels = bContainValidPixels && (bEnableTwoSidedFoliage ? true : (BSDF_GETSSSTYPE(BSDF) != SSS_TYPE_TWO_SIDED_WRAP));
 
 		const uint BSDFType = BSDF_GETTYPE(BSDF);
 		switch (BSDFType)
diff --git a/Engine/Shaders/Private/PathTracing/PathTracingMaterialHitShader.usf b/Engine/Shaders/Private/PathTracing/PathTracingMaterialHitShader.usf
--- a/Engine/Shaders/Private/PathTracing/PathTracingMaterialHitShader.usf
+++ b/Engine/Shaders/Private/PathTracing/PathTracingMaterialHitShader.usf
@@ -612,11 +612,11 @@ RAY_TRACING_ENTRY_CLOSEST_HIT(PathTracingMaterialCHS,
 						Payload.PhaseG = 0.0;
 #endif
 					}
-					else if (BSDF_GETSSSTYPE(BSDF) == SSS_TYPE_WRAP)
+					else if (BSDF_GETSSSTYPE(BSDF) == SSS_TYPE_WRAP || BSDF_GETSSSTYPE(BSDF) == SSS_TYPE_TWO_SIDED_WRAP)
 					{
-						const bool bIsThin = BSDF_GETISTHIN(BSDF);
+						const bool bFoliage = BSDF_GETSSSTYPE(BSDF) == SSS_TYPE_TWO_SIDED_WRAP;
 						const float TransmittanceNoL = 1.0f;
-						const float3 SlabDiffuseColor = bIsThin ? Payload.DiffuseColor : float3(1, 1, 1);
+						const float3 SlabDiffuseColor = bFoliage ? Payload.DiffuseColor : float3(1, 1, 1);
 						const FParticipatingMedia PM = SubstrateSlabCreateParticipatingMedia(Payload.DiffuseColor, SLAB_SSSMFP(BSDF));
 						float3 SubsurfaceColor = IsotropicMediumSlabTransmittance(PM, SUBSTRATE_SIMPLEVOLUME_THICKNESS_M, TransmittanceNoL);
 						const float Opacity = 1.f - abs(SLAB_SSSPHASEANISOTROPY(BSDF));
@@ -632,7 +632,7 @@ RAY_TRACING_ENTRY_CLOSEST_HIT(PathTracingMaterialCHS,
 							Payload.DiffuseColor = Front / S;
 							SubsurfaceColor = Back / S;
 						}
-						if (bIsThin)
+						if (bFoliage)
 						{
 							// foliage case, only the color matters
 							Payload.ShadingModelID = SHADINGMODELID_TWOSIDED_FOLIAGE;
diff --git a/Engine/Shaders/Private/RayTracing/RayTracingOcclusionRGS.usf b/Engine/Shaders/Private/RayTracing/RayTracingOcclusionRGS.usf
--- a/Engine/Shaders/Private/RayTracing/RayTracingOcclusionRGS.usf
+++ b/Engine/Shaders/Private/RayTracing/RayTracingOcclusionRGS.usf
@@ -688,9 +688,9 @@ RAY_TRACING_ENTRY_RAYGEN_OR_INLINE(Occlusion)
 
 			ShadingParameters.bIsValid = SubstratePixelHeader.IsSubstrateMaterial();
 			ShadingParameters.bIsHair = BSDFType == SUBSTRATE_BSDF_TYPE_HAIR;
-			ShadingParameters.bIsSubsurface = bHasSSS ? SSSType == SSS_TYPE_WRAP : false;
+			ShadingParameters.bIsSubsurface = bHasSSS ? SSSType == SSS_TYPE_WRAP || SSSType == SSS_TYPE_TWO_SIDED_WRAP : false;
 			ShadingParameters.bHasSubsurfaceProfile = bHasSSS ? (SSSType == SSS_TYPE_DIFFUSION_PROFILE || SSSType == SSS_TYPE_DIFFUSION) : false;
-			ShadingParameters.bBackfaceCulling = !(BSDFType == SUBSTRATE_BSDF_TYPE_EYE || (bHasSSS && SSSType == SSS_TYPE_WRAP)); // isthin?
+			ShadingParameters.bBackfaceCulling = !(BSDFType == SUBSTRATE_BSDF_TYPE_EYE || (bHasSSS && SSSType == SSS_TYPE_WRAP || SSSType == SSS_TYPE_TWO_SIDED_WRAP));
 
 			// The SSDenoiser reads the top layer normal in order to know how to process the data.
 			// RTShadow traces using the top layer normal when NoL>0 only. So if normals for slabs below do not match, light leaking will occur.
diff --git a/Engine/Shaders/Private/Substrate/Substrate.ush b/Engine/Shaders/Private/Substrate/Substrate.ush
--- a/Engine/Shaders/Private/Substrate/Substrate.ush
+++ b/Engine/Shaders/Private/Substrate/Substrate.ush
@@ -152,9 +152,6 @@
 #define SUBSTRATE_EPSILON						1e-10f
 #define SUBSTRATE_EYE_DEFAULT_F0				0.028f
 
-// Disable SSS diffusion based when the pixel's footprint is larger than x times the MFPs
-#define SUBSTRATE_DISABLE_SSS_DIFFUSION_BASED_ON_PIXEL_FOOTPRINT 0
-
 ///////////////////////////////////////////////////////////////////////////////
 // Sub-surface scattering
 //
@@ -165,31 +162,53 @@
 //  * Top layer/Bottom layer (on Material topology)
 //
 //  Material Parameters (Slab)
-//  -------------------
+//  --------------------------
 //  * Diffuse Albedo
 //  * SSS MFP
 //  * SSS MFP scale
 //  * SSS Phase aniso
-// 
-//  SSS techniques
-//  --------------
-//  * Diffusion         
-//  * Diffusion Profile     
-//  * Wrap (old Foliage)    
-//  * Wrap (ols Subsurface) 
-//  * Simple Volume         
-// 
-//  Techniques used for each configuration 
+//  * SSS modes 
+//     * None
+//     * Wrap
+//     * Two-Sided Wrap
+//     * Diffusion
+//     * Diffusion Profile
+//     * Simple Volume
 //
-// 	|-------------------------------------------------------------------------------------------|-------------------------------------------------------------------------------------------|
-// 	| Diffusion                                                                                 | !Diffusion                                                                                | 
-// 	|---------------------------------------------|---------------------------------------------|---------------------------------------------|---------------------------------------------|
-// 	| !IsThin                                     | IsThin                                      | !IsThin                                     | IsThin                                      |
-// 	|---------------------------------------------|---------------------------------------------|---------------------------------------------|---------------------------------------------|
-// 	| Top                  | Bottom               | Top                  | Bottom               | Top                  | Bottom               | Top                  | Bottom               |
-// 	|----------------------|----------------------|----------------------|----------------------|----------------------|----------------------|----------------------|----------------------|
-// 	| Simple Volume        | Diffusion            | Simple Volume        | Diffusion            | Simple Volume        | Wrap (old Subsurface)| Simple Volume        | Wrap (old Foliage)   |
+//  SSS techniques (shading implementation)
+//  ---------------------------------------
+//  * None
+//  * Simple Volume
+//  * Wrap (ols Subsurface)
+//  * Two-Sided Wrap (old Foliage)
+//  * Diffusion
+//  * Diffusion Profile
 // 
+//  Usable techniques for each configuration 
+//
+//  * Deferred case
+// 	|---------------------------------------------|---------------------------------------------|
+// 	| !IsThin                                     | IsThin                                      |
+// 	|---------------------------------------------|---------------------------------------------|
+// 	| !Bottom              | Bottom               | !Bottom              | Bottom               |
+// 	|----------------------|----------------------|----------------------|----------------------|
+//	| None                 | None                 | None                 | None                 |
+//	| Simple Volume        | Simple Volume        | Simple Volume        | Simple Volume        |
+//	|                      | Wrap                 |                      | Wrap                 |
+//	|                      | Two-Sided Wrap       |                      | Two-Sided Wrap       |
+//	|                      | Diffusion            |                      | Diffusion            |
+//	|                      | Diffusion Profile    |                      | Diffusion Profile    |
+//
+//  * Forward case
+// 	|---------------------------------------------|---------------------------------------------|
+// 	| !IsThin                                     | IsThin                                      |
+// 	|---------------------------------------------|---------------------------------------------|
+// 	| !Bottom              | Bottom               | !Bottom              | Bottom               |
+// 	|----------------------|----------------------|----------------------|----------------------|
+//	| None                 | None                 | None                 | None                 |
+//	| Simple Volume        | Simple Volume        | Simple Volume        | Simple Volume        |
+//	|                      | Wrap                 |                      | Wrap                 |
+//	|                      | Two-Sided Wrap       |                      | Two-Sided Wrap       |
 //
 // How are stored the MFP data?
 // ---------------------------
@@ -1429,7 +1448,7 @@ FSubstrateData InternalGetSubstrateSlabBSDF(
 	FSubstratePixelFootprint InPixelFootprint,
 	float3 DiffuseAlbedo, float3 F0, float3 F90,
 	float Roughness, float Anisotropy,
-	float SSSProfileID, float3 SSSMFP, float SSSMFPScale, float SSSPhaseAnisotropy, float UseSSSDiffusion,
+	float SSSProfileID, float3 SSSMFP, float SSSMFPScale, float SSSPhaseAnisotropy, uint SSSType,
 	float3 Emissive, 
 	float SecondRoughness, float SecondRoughnessWeight, float SecondRoughnessAsSimpleClearCoat,
 	float FuzzAmount, float3 FuzzColor, float FuzzRoughness,
@@ -1444,22 +1463,40 @@ FSubstrateData InternalGetSubstrateSlabBSDF(
 	Roughness		= max(Roughness, InPixelFootprint.NormalCurvatureRoughness);
 	SecondRoughness = max(SecondRoughness, InPixelFootprint.NormalCurvatureRoughness);
 	FuzzRoughness	= max(FuzzRoughness, InPixelFootprint.NormalCurvatureRoughness);
-	
 
 	FSubstrateData SubstrateData = GetInitialisedSubstrateData();
 
 	SSSMFPScale = saturate(SSSMFPScale);
 	const bool bSSSProfilePlugged = SSSProfileID > 0.0f;
 	const bool bMFPPlugged = any(SSSMFP > 0);
-	bool bIsThin		= bIsAtTheBottomOfTopology && (bMFPPlugged && bIsThinSurface); // Enable thin only if MFP is plugged, as otherwise there will be no back-tramission
-	bool bHasSSSWrap	= bIsAtTheBottomOfTopology && (bMFPPlugged && !UseSSSDiffusion);
-	bool bHasSSSProfile	= bIsAtTheBottomOfTopology && (bSSSProfilePlugged && SSSMFPScale > 0.f);
-	bool bHasSSS		= bMFPPlugged || (bIsAtTheBottomOfTopology && (bHasSSSProfile || bHasSSSWrap)); // When MFP is plugged, that can be diffusion (bottom) or Simple volume (any position)
+	const bool bIsThin = bIsAtTheBottomOfTopology && (bMFPPlugged && bIsThinSurface); // Enable thin only if MFP is plugged, as otherwise there will be no back-tramission
+
+	// In forward mode (blend mode != Opaque/Masked), diffusion is not supported. 
+	// In such a case we fallback onto wrap mode, as look-wise, this is the closest one.
+	#if SUBSTRATE_OPAQUE_MATERIAL == 0
+	if (SSSType == SSS_TYPE_DIFFUSION || SSSType == SSS_TYPE_DIFFUSION_PROFILE)
+	{
+		SSSType = SSS_TYPE_WRAP;
+	}
+	#endif
+
+	if (!bIsAtTheBottomOfTopology)
+	{
+		SSSType = SSS_TYPE_SIMPLEVOLUME;
+	}
+	else if (SSSType == SSS_TYPE_DIFFUSION && bSSSProfilePlugged && SSSMFPScale > 0.f)
+	{
+		SSSType = SSS_TYPE_DIFFUSION_PROFILE;
+	}
+	else if (!bMFPPlugged)
+	{
+		SSSType = SSS_TYPE_INVALID;
+	}
 
 	FHaziness Haziness = InitialiseHaziness();
 
 	// Derive roughness and haziness from profile data
-	if (!SUBSTRATE_FASTPATH && bSSSProfilePlugged)
+	if (!SUBSTRATE_FASTPATH && SSSType == SSS_TYPE_DIFFUSION_PROFILE)
 	{
 		// Average roughness for dual specular.
 		const uint SubsurfaceProfileUInt = SubstrateSubsurfaceProfileIdTo8bits(SSSProfileID);
@@ -1475,22 +1512,8 @@ FSubstrateData InternalGetSubstrateSlabBSDF(
 		Haziness.bSimpleClearCoat = SecondRoughnessAsSimpleClearCoat > 0.0f;
 	}
 
-	#if SUBSTRATE_DISABLE_SSS_DIFFUSION_BASED_ON_PIXEL_FOOTPRINT
-	// If MFP is smaller than the pixel's footprint, remove SSS attributes
-	// The scaling factor 'SSSFootprintScale' is based on empiric tests
-	const float SSSFootprintScale = 3.f;
-	if (!SUBSTRATE_FASTPATH && all(SSSMFP * SSSFootprintScale < InPixelFootprint.PixelRadiusInWorldSpace) && !bHasSSSWrap)
-	{
-		bHasSSSProfile = false;
-		bHasSSS = false;
-
-		// MFPPlugged is not cleared as it is used for thin surface and thin volumetric
-		// bMFPPlugged = false;
-	}
-	#endif
-
 	// Convert MFP be expressed relative to simple volume thickness
-	if (!SUBSTRATE_FASTPATH && bHasSSSWrap)
+	if (!SUBSTRATE_FASTPATH && (SSSType == SSS_TYPE_WRAP || SSSType == SSS_TYPE_TWO_SIDED_WRAP || SSSType == SSS_TYPE_SIMPLEVOLUME))
 	{
 		SSSMFP = RescaleMFPToComputationSpace(SSSMFP, ThicknessCm, SUBSTRATE_SIMPLEVOLUME_THICKNESS_CM);
 		ThicknessCm = SUBSTRATE_SIMPLEVOLUME_THICKNESS_CM;
@@ -1515,11 +1538,11 @@ FSubstrateData InternalGetSubstrateSlabBSDF(
 	//  - It will only be ran for the bottom layer only. Otherwise, the simple volume should be used. As such, Two-sided lighting can also be disabled later if part of the bottom layer.
 
 	BSDF_SETTYPE			(SubstrateData.InlinedBSDF,		SUBSTRATE_BSDF_TYPE_SLAB);
-	BSDF_SETSHAREDLOCALBASISID(SubstrateData.InlinedBSDF,		SharedLocalBasisIndex);
+	BSDF_SETSHAREDLOCALBASISID(SubstrateData.InlinedBSDF,	SharedLocalBasisIndex);
 	BSDF_SETEMISSIVE		(SubstrateData.InlinedBSDF,		Emissive);
 
 #if (SUBSTRATE_FASTPATH==0 || MATERIALBLENDING_ANY_TRANSLUCENT) // !Fast path or rendering translucent materials
-	BSDF_SETSSSTYPE			(SubstrateData.InlinedBSDF,		bHasSSS ? (bHasSSSWrap ? SSS_TYPE_WRAP : (bHasSSSProfile ? SSS_TYPE_DIFFUSION_PROFILE : SSS_TYPE_DIFFUSION)) : SSS_TYPE_INVALID);
+	BSDF_SETSSSTYPE			(SubstrateData.InlinedBSDF,		SSSType);
 	BSDF_SETHASANISOTROPY	(SubstrateData.InlinedBSDF,		bHasAnisotropy);
 	BSDF_SETISTOPLAYER		(SubstrateData.InlinedBSDF,		0);
 	BSDF_SETISTHIN			(SubstrateData.InlinedBSDF,		bIsThin ? 1 : 0);
@@ -1563,7 +1586,7 @@ FSubstrateData InternalGetSubstrateSlabBSDF(
 
 	SLAB_HAZINESS(SubstrateData.InlinedBSDF)					= PackHaziness(Haziness);
 
-	if (bHasSSSProfile)
+	if (SSSType == SSS_TYPE_DIFFUSION_PROFILE)
 	{
 		SLAB_SSSPROFILEID	(SubstrateData.InlinedBSDF)		= SSSProfileID;
 		SLAB_SSSPROFILERADIUSSCALE(SubstrateData.InlinedBSDF)	= SSSMFPScale;
@@ -1612,7 +1635,7 @@ FSubstrateData GetSubstrateSlabBSDF(
 	FSubstratePixelFootprint InPixelFootprint,
 	float3 DiffuseAlbedo, float3 F0, float3 F90,
 	float Roughness, float Anisotropy,
-	float SSSProfileID, float3 SSSMFP, float SSSMFPScale, float SSSPhaseAnisotropy, float UseSSSDiffusion,
+	float SSSProfileID, float3 SSSMFP, float SSSMFPScale, float SSSPhaseAnisotropy, uint SSSType,
 	float3 Emissive,
 	float SecondRoughness, float SecondRoughnessWeight, float SecondRoughnessAsSimpleClearCoat,
 	float FuzzAmount, float3 FuzzColor, float FuzzRoughness,
@@ -1627,7 +1650,7 @@ FSubstrateData GetSubstrateSlabBSDF(
 		InPixelFootprint,
 		DiffuseAlbedo, F0, F90,
 		Roughness, Anisotropy,
-		SSSProfileID, SSSMFP, SSSMFPScale, SSSPhaseAnisotropy, UseSSSDiffusion,
+		SSSProfileID, SSSMFP, SSSMFPScale, SSSPhaseAnisotropy, SSSType,
 		Emissive,
 		SecondRoughness, SecondRoughnessWeight, SecondRoughnessAsSimpleClearCoat,
 		FuzzAmount, FuzzColor, FuzzRoughness,
@@ -1644,7 +1667,7 @@ FSubstrateData GetSubstrateSlabBSDF(
 	FSubstratePixelFootprint InPixelFootprint,
 	float3 DiffuseAlbedo, float3 F0, float3 F90,
 	float Roughness, float Anisotropy,
-	float SSSProfileID, float3 SSSMFP, float SSSMFPScale, float SSSPhaseAnisotropy, float UseSSSDiffusion,
+	float SSSProfileID, float3 SSSMFP, float SSSMFPScale, float SSSPhaseAnisotropy, uint SSSType,
 	float3 Emissive,
 	float SecondRoughness, float SecondRoughnessWeight, float SecondRoughnessAsSimpleClearCoat,
 	float FuzzAmount, float3 FuzzColor, float FuzzRoughness,
@@ -1659,7 +1682,7 @@ FSubstrateData GetSubstrateSlabBSDF(
 		InPixelFootprint,
 		DiffuseAlbedo, F0, F90,
 		Roughness, Anisotropy,
-		SSSProfileID, SSSMFP, SSSMFPScale, SSSPhaseAnisotropy, UseSSSDiffusion,
+		SSSProfileID, SSSMFP, SSSMFPScale, SSSPhaseAnisotropy, SSSType,
 		Emissive,
 		SecondRoughness, SecondRoughnessWeight, SecondRoughnessAsSimpleClearCoat,
 		FuzzAmount, FuzzColor, FuzzRoughness,
@@ -2383,7 +2406,7 @@ void FSubstratePixelHeader::SubstrateConvertFromDBuffer(
 		0.f,				// SSSMFP
 		0.f,				// SSSMFPScale
 		0.f,				// SSSPhaseAniso
-		0.f,				// UseSSSDiffusion
+		SSS_TYPE_INVALID,	// UseSSSDiffusion
 		0.f,				// Emissive
 		0.f,				// SecondRoughness
 		0.f,				// SecondRoughnessWeight
@@ -3212,36 +3235,10 @@ void FSubstrateBSDF::SubstrateSanitizeBSDF()
 // Convert to simple volume to get some transmittance
 void FSubstrateBSDF::PostProcessSlabBeforeLighting(bool bIsSubstrateOpaqueMaterial)
 {
-	// Two sided lighting can only be evaluated on the bottom layer, together with option SSS. Otherwise, the simple volume will be used.
-	// At the moment, a material can't be thin & simple volume at the same time. This should be fixed by unifying SSSDiffusion|SSSWrap|SimpleVolume into a single evaluation function
-	const bool bIsThin = this.bIsBottom && BSDF_GETISTHIN(this) && bIsSubstrateOpaqueMaterial;
-
-	// We want simple volumetric only if the MFP input is pluged in, otherwise we avoid dealing with simple volumetric.
-	// We also do not want to use the simple volumetric path if we are dealing with an opaque material and this slab is not at the bottom.
-	// And when opaque, the bottom layer must be opaque and thus will ahve a solid opaque diffuse color.
-	const bool bIsSimpleVolume = !bIsThin && BSDF_GETHASMFP(this) && (!bIsSubstrateOpaqueMaterial || (bIsSubstrateOpaqueMaterial && !this.bIsBottom));
-
-	if (bIsSimpleVolume)
+	if (!this.bIsBottom && BSDF_GETHASMFP(this))
 	{
-		// * Rescale MFP to computation space
-		// * Set Invalid SSS type
-		// * Set IsSimpleVolume
 		this.EnableSlabBSDFSimpleVolumetric();
 	}
-	else if (bIsThin)
-	{
-		// Only rescale MFP to computation space for SSS wrap. 
-		// SSS Diffuse types use their 'original' MFP for diffusion
-		if (BSDF_GETSSSTYPE(this) == SSS_TYPE_WRAP)
-		{
-			this.EnableSlabMFPComputationSpace();
-		}
-	}
-	else
-	{
-		// If not selected or not possible to have two sided lighting enabled (not bottom layer or SSSProfile is enabled), let's make sure we disable the feature.
-		BSDF_SETISTHIN(this, 0);
-	}
 }
 
 void FSubstrateBSDF::PostProcessBSDFBeforeLighting(bool bIsSubstrateOpaqueMaterial)
@@ -3613,16 +3610,10 @@ float SubstrateGetLegacyShadingModels(in FSubstrateBSDF BSDF)
 	{
 	case SUBSTRATE_BSDF_TYPE_SLAB:
 	{
-		// Base on features, we can only partially recovers the original shading models. Shading models like foliage, eye, clear coat are impossible/harder to infer.
+		// Based on features, we can only partially recovers the original shading models. Shading models like foliage, eye, clear coat are impossible/harder to infer.
 		if (BSDF_GETSSSTYPE(BSDF) == SSS_TYPE_DIFFUSION_PROFILE)	{ return SHADINGMODELID_SUBSURFACE_PROFILE; }
-		else if (BSDF_GETSSSTYPE(BSDF) == SSS_TYPE_WRAP)		
-		{ 
-			if (BSDF_GETISTHIN(BSDF))
-			{
-				return SHADINGMODELID_TWOSIDED_FOLIAGE;
-			}
-			return SHADINGMODELID_SUBSURFACE; 
-		}
+		else if (BSDF_GETSSSTYPE(BSDF) == SSS_TYPE_WRAP)			{ return SHADINGMODELID_SUBSURFACE; }
+		else if (BSDF_GETSSSTYPE(BSDF) == SSS_TYPE_TWO_SIDED_WRAP)	{ return SHADINGMODELID_TWOSIDED_FOLIAGE; }
 		else if (BSDF_GETHASFUZZ(BSDF))								{ return SHADINGMODELID_CLOTH; }
 		else
 		{ 
@@ -4280,7 +4271,7 @@ FSubstrateBSDF UnpackSubstrateBSDFIn(FSubstrateMaterialContainer SubstrateBuffer
 				BSDF_SETHASTRANSABOVE(OutBSDF, 0);
 				BSDF_SETHASGLINT(OutBSDF, 0);
 			}
-			else if (OptimisedLegacyMode == SINGLE_OPTLEGACYMODE_SSSWRAP || OptimisedLegacyMode == SINGLE_OPTLEGACYMODE_SSSWRAP_THIN)
+			else if (OptimisedLegacyMode == SINGLE_OPTLEGACYMODE_SSSWRAP || OptimisedLegacyMode == SINGLE_OPTLEGACYMODE_TWO_SIDED_SSSWRAP)
 			{
 			#if SUBSTRATE_DEFERRED_SHADING
 				uint Data0 = SubstrateHeader.PackedHeader;
@@ -4297,7 +4288,7 @@ FSubstrateBSDF UnpackSubstrateBSDFIn(FSubstrateMaterialContainer SubstrateBuffer
 
 				const uint PackedSSSWOpacity7bits = (Data0 >> (HEADER_SINGLEENCODING_BIT_COUNT + HEADER_SINGLE_OPTLEGACYMODE_BIT_COUNT)) & 0x7F;
 				const float Opacity = UnpackR7(PackedSSSWOpacity7bits);
-
+				
 				SLAB_DIFFUSEALBEDO(OutBSDF)	= UnpackR7G7B6Gamma2(PackedDiffuse20Bits);
 				SLAB_F0(OutBSDF)			= UnpackR7G7B6Gamma2(PackedF020Bits);
 				SLAB_ROUGHNESS(OutBSDF)		= SubstrateUnpackTopLayerData(SubstrateHeader.PackedTopLayerData).Roughness;
@@ -4305,16 +4296,18 @@ FSubstrateBSDF UnpackSubstrateBSDFIn(FSubstrateMaterialContainer SubstrateBuffer
 				SLAB_SSSPHASEANISOTROPY(OutBSDF)= saturate(1.f - Opacity);
 			#endif
 
+				const bool bTwoSided = OptimisedLegacyMode == SINGLE_OPTLEGACYMODE_TWO_SIDED_SSSWRAP;
+
 				// Force a BSDF state to be a single layer on top.
 				BSDF_SETTYPE(OutBSDF, SUBSTRATE_BSDF_TYPE_SLAB);
 				BSDF_SETHASGREYWEIGHT_V(OutBSDF, 1);
 				BSDF_SETHASF90(OutBSDF, 0);
 				BSDF_SETHASANISOTROPY(OutBSDF, 0);
 				BSDF_SETISTOPLAYER(OutBSDF, 1);
-				BSDF_SETSSSTYPE(OutBSDF, SSS_TYPE_WRAP);
+				BSDF_SETSSSTYPE(OutBSDF, bTwoSided ? SSS_TYPE_TWO_SIDED_WRAP : SSS_TYPE_WRAP);
 				BSDF_SETHASMFP(OutBSDF, 1);
 				BSDF_SETHASHAZINESS(OutBSDF, 0);
-				BSDF_SETISTHIN(OutBSDF, (OptimisedLegacyMode == SINGLE_OPTLEGACYMODE_SSSWRAP_THIN ? 1 : 0));
+				BSDF_SETISTHIN(OutBSDF, (bTwoSided ? 1 : 0));
 				BSDF_SETHASFUZZ(OutBSDF, 0);
 				BSDF_SETHASTRANSABOVE(OutBSDF, 0);
 				BSDF_SETHASGLINT(OutBSDF, 0);
diff --git a/Engine/Shaders/Private/Substrate/SubstrateDBuffer.usf b/Engine/Shaders/Private/Substrate/SubstrateDBuffer.usf
--- a/Engine/Shaders/Private/Substrate/SubstrateDBuffer.usf
+++ b/Engine/Shaders/Private/Substrate/SubstrateDBuffer.usf
@@ -299,7 +299,7 @@ void MainCS(
 				MaterialTextureArrayUAV[uint3(PixelCoord, 1)] = (PackedDiffuse20Bits | ((PackedFuzzColor20bits & 0xFFC00) << 10));
 				MaterialTextureArrayUAV[uint3(PixelCoord, 2)] = (PackedF020Bits | ((PackedFuzzColor20bits & 0x3FF) << 20));
 			}
-			else if (OptimisedLegacyMode == SINGLE_OPTLEGACYMODE_SSSWRAP || OptimisedLegacyMode == SINGLE_OPTLEGACYMODE_SSSWRAP_THIN)
+			else if (OptimisedLegacyMode == SINGLE_OPTLEGACYMODE_SSSWRAP || OptimisedLegacyMode == SINGLE_OPTLEGACYMODE_TWO_SIDED_SSSWRAP)
 			{
 				// Roughness is only stored in the TopLayer texture.
 				const float3 MFP = SLAB_SSSMFP(BSDF) * DBufferData.OneMinusCoverage_BaseColor;
diff --git a/Engine/Shaders/Private/Substrate/SubstrateEvaluation.ush b/Engine/Shaders/Private/Substrate/SubstrateEvaluation.ush
--- a/Engine/Shaders/Private/Substrate/SubstrateEvaluation.ush
+++ b/Engine/Shaders/Private/Substrate/SubstrateEvaluation.ush
@@ -821,19 +821,19 @@ FSubstrateEvaluateResult SubstrateEvaluateBSDFCommon(FSubstrateBSDFContext BSDFC
 			////
 
 			#if SUBSTRATE_FASTPATH==0
-			if (BSDF_GETSSSTYPE(BSDFContext.BSDF) == SSS_TYPE_WRAP)
+			if (BSDF_GETSSSTYPE(BSDFContext.BSDF) == SSS_TYPE_WRAP || BSDF_GETSSSTYPE(BSDFContext.BSDF) == SSS_TYPE_TWO_SIDED_WRAP)
 			{
 				// The wrap lighting is a non-physically based shading which intends to match legacy behavior.
 				// For doing so the look MFP & phase function anisotropy are respectively reinterpreted as 'SubSurfaceColor' and 'Opacity', per legacy shading model
-				const bool bIsThin = BSDF_GETISTHIN(BSDFContext.BSDF);
+				const bool bTwoSided = BSDF_GETSSSTYPE(BSDFContext.BSDF) == SSS_TYPE_TWO_SIDED_WRAP;
 				const float TransmittanceNoL = 1.0f;
-				const float3 SlabDiffuseColor = bIsThin ? DiffuseColor : float3(1, 1, 1);
+				const float3 SlabDiffuseColor = bTwoSided ? DiffuseColor : float3(1, 1, 1);
 				const FParticipatingMedia PM = SubstrateSlabCreateParticipatingMedia(SlabDiffuseColor, SLAB_SSSMFP(BSDFContext.BSDF));
 				const float3 SubSurfaceColor = IsotropicMediumSlabTransmittance(PM, SUBSTRATE_SIMPLEVOLUME_THICKNESS_M, TransmittanceNoL);
 				const float Opacity = 1.f - abs(SLAB_SSSPHASEANISOTROPY(BSDFContext.BSDF));
 
 				float3 TransmissionThroughput = 0;
-				if (bIsThin)
+				if (bTwoSided)
 				{
 					// Legacy Foliage
 					
@@ -1574,7 +1574,6 @@ FSubstrateEnvLightResult SubstrateEvaluateForEnvLight(FSubstrateBSDFContext BSDF
 			#if SUBSTRATE_FASTPATH==0
 			const bool bIsSubsurfaceDiffusionProfile = BSDF_GETSSSTYPE(BSDFContext.BSDF) == SSS_TYPE_DIFFUSION_PROFILE;
 			const bool bIsSubsurfaceDiffusionMFP = BSDF_GETSSSTYPE(BSDFContext.BSDF) == SSS_TYPE_DIFFUSION;
-			const bool bIsThin = BSDF_GETISTHIN(BSDFContext.BSDF);
 			SubstrateEnvLightResult.bPostProcessSubsurface = (bIsSubsurfaceDiffusionProfile || bIsSubsurfaceDiffusionMFP) && SUBSTRATE_SSS_MATERIAL_OVERRIDE;
 
 			BRANCH
@@ -1725,10 +1724,10 @@ FSubstrateEnvLightResult SubstrateEvaluateForEnvLight(FSubstrateBSDFContext BSDF
 
 			#if SUBSTRATE_FASTPATH==0
 			BRANCH
-			if (BSDF_GETSSSTYPE(BSDFContext.BSDF) == SSS_TYPE_WRAP)
+				if (BSDF_GETSSSTYPE(BSDFContext.BSDF) == SSS_TYPE_WRAP || BSDF_GETSSSTYPE(BSDFContext.BSDF) == SSS_TYPE_TWO_SIDED_WRAP)
 			{
-				const bool bIsThin = BSDF_GETISTHIN(BSDFContext.BSDF);
-				const float3 SlabDiffuseColor = bIsThin ? DiffuseColor : float3(1, 1, 1);
+				const bool bTwoSided = BSDF_GETSSSTYPE(BSDFContext.BSDF) == SSS_TYPE_TWO_SIDED_WRAP;
+				const float3 SlabDiffuseColor = bTwoSided ? DiffuseColor : float3(1, 1, 1);
 				FParticipatingMedia PM = SubstrateSlabCreateParticipatingMedia(SlabDiffuseColor, SLAB_SSSMFP(BSDFContext.BSDF));
 			
 				// We consider single scattering and perpendicular to the surface light transmittance.
@@ -1736,7 +1735,7 @@ FSubstrateEnvLightResult SubstrateEvaluateForEnvLight(FSubstrateBSDFContext BSDF
 				const float TransmittanceNoL = 1.0f;
 				const float3 SlabTransmittance = IsotropicMediumSlabTransmittance(PM, SUBSTRATE_SIMPLEVOLUME_THICKNESS_M, TransmittanceNoL);
 			
-				if (bIsThin)
+				if (bTwoSided)
 				{
 					// We recover the transmitted luminance from as the back face Lambert affected by the transmittance.
 					// The MFP is specified on the slab node and then transmittance recoved from it.
diff --git a/Engine/Shaders/Private/Substrate/SubstrateExport.ush b/Engine/Shaders/Private/Substrate/SubstrateExport.ush
--- a/Engine/Shaders/Private/Substrate/SubstrateExport.ush
+++ b/Engine/Shaders/Private/Substrate/SubstrateExport.ush
@@ -328,11 +328,10 @@ void PackSubstrateOut(
 				const uint SingleHasHaziness		= 0x01;
 				const uint SingleHasClearCoat		= 0x02;
 				const uint SingleHasCloth			= 0x04;
-				const uint SingleIsThin				= 0x08;
 				const uint SingleIsSimpleVolume		= 0x10;
 				const uint SingleIsWrap				= 0x20;
-				const uint SingleIsSSSProfile		= 0x40;
-				const uint SingleIsWrapThin			= SingleIsThin | SingleIsWrap;
+				const uint SingleIsTwoSidedWrap		= 0x40;
+				const uint SingleIsSSSProfile		= 0x80;
 				// No need to add anisotropy since it is not supported by the simple mode.
 
 				uint OptimisedModeMask = 0;
@@ -350,10 +349,6 @@ void PackSubstrateOut(
 				{
 					OptimisedModeMask |= SingleHasCloth;
 				}
-				if (BSDF_GETISTHIN(BSDF))
-				{
-					OptimisedModeMask |= SingleIsThin;
-				}
 				if (BSDF_GETSSSTYPE(BSDF) == SSS_TYPE_SIMPLEVOLUME)
 				{
 					OptimisedModeMask |= SingleIsSimpleVolume;
@@ -362,6 +357,10 @@ void PackSubstrateOut(
 				{
 					OptimisedModeMask |= SingleIsWrap;
 				}
+				if (BSDF_GETSSSTYPE(BSDF) == SSS_TYPE_TWO_SIDED_WRAP)
+				{
+					OptimisedModeMask |= SingleIsTwoSidedWrap;
+				}
 				if (BSDF_GETSSSTYPE(BSDF) == SSS_TYPE_DIFFUSION_PROFILE)
 				{
 					OptimisedModeMask |= SingleIsSSSProfile;
@@ -375,9 +374,9 @@ void PackSubstrateOut(
 				{
 					OptimisedLegacyMode = SINGLE_OPTLEGACYMODE_CLOTH;
 				}				
-				else if ((OptimisedModeMask & SingleIsWrapThin) == SingleIsWrapThin && !BSDF_GETHASFUZZ(BSDF))
+				else if ((OptimisedModeMask & SingleIsTwoSidedWrap) == SingleIsTwoSidedWrap && !BSDF_GETHASFUZZ(BSDF))
 				{
-					OptimisedLegacyMode = SINGLE_OPTLEGACYMODE_SSSWRAP_THIN;
+					OptimisedLegacyMode = SINGLE_OPTLEGACYMODE_TWO_SIDED_SSSWRAP;
 					SSSData.Header.Bytes = 0; // Force to skip SSSData writing
 				}
 				else if ((OptimisedModeMask & SingleIsWrap) == SingleIsWrap && !BSDF_GETHASFUZZ(BSDF))
@@ -385,9 +384,8 @@ void PackSubstrateOut(
 					OptimisedLegacyMode = SINGLE_OPTLEGACYMODE_SSSWRAP;
 					SSSData.Header.Bytes = 0; // Force to skip SSSData writing
 				}
-				else if ((OptimisedModeMask & (SingleIsSSSProfile|SingleIsThin)) == SingleIsSSSProfile && !BSDF_GETHASFUZZ(BSDF)) // Force to skip SINGLE_OPTLEGACYMODE_SSSPROFILE when set to ThinSurface
+				else if ((OptimisedModeMask & SingleIsSSSProfile) == SingleIsSSSProfile && !BSDF_GETHASFUZZ(BSDF))
 				{
-					// SingleIsSSSProfile appears critically after SingleIsWrapThin because thin will disable subsurface profile.
 					OptimisedLegacyMode = SINGLE_OPTLEGACYMODE_SSSPROFILE;
 					SSSData.Header.Bytes = 0; // Force to skip SSSData writing
 				}
@@ -611,7 +609,7 @@ void PackSubstrateOut(
 							// Ensure the rest of the BSDF is not stored with the regular path
 							BSDF.State = 0;
 						}
-						else if (OptimisedLegacyMode == SINGLE_OPTLEGACYMODE_SSSWRAP || OptimisedLegacyMode == SINGLE_OPTLEGACYMODE_SSSWRAP_THIN)
+						else if (OptimisedLegacyMode == SINGLE_OPTLEGACYMODE_SSSWRAP || OptimisedLegacyMode == SINGLE_OPTLEGACYMODE_TWO_SIDED_SSSWRAP)
 						{
 							const float SSSWOpacity = SubstrateSubSurfaceGetWrapOpacityFromAnisotropy(SLAB_SSSPHASEANISOTROPY(BSDF));
 
@@ -782,7 +780,7 @@ void PackSubstrateOut(
 									// Legacy SSSProfile ID is still set by the last encountered Slab.
 									// Attributes are blended using opacity only, no need to use the throughput to the view (SSS post process wants raw source data like BaseColor and the BSDF will otherwise run the correct lighting math).
 
-									if (BSDF_GETSSSTYPE(BSDF) == SSS_TYPE_WRAP)
+									if (BSDF_GETSSSTYPE(BSDF) == SSS_TYPE_WRAP || BSDF_GETSSSTYPE(BSDF) == SSS_TYPE_TWO_SIDED_WRAP)
 									{
 										SSSDataAniso += SLAB_SSSPHASEANISOTROPY(BSDF) * BSDF.Coverage;
 									}
@@ -908,7 +906,7 @@ void PackSubstrateOut(
 				// Always override the type, we can only blend similar types of subsrface
 				SubstrateSubSurfaceHeaderSetSSSType(SSSData.Header, SSSDataType);
 
-				if (SSSDataType == SSS_TYPE_WRAP)
+				if (SSSDataType == SSS_TYPE_WRAP || SSSDataType == SSS_TYPE_TWO_SIDED_WRAP)
 				{
 					SubstrateSubSurfaceHeaderSetWrap(SSSData.Header, SSSDataAniso);
 				}
@@ -1015,7 +1013,7 @@ FExportResult SubstrateMaterialExportOut(
 					Out.Metallic	= HAIR_SCATTER(CurrentBSDF);
 					Out.CustomData.z= HAIR_BACKLIT(CurrentBSDF);
 				}
-				else if (SSSType == SSS_TYPE_WRAP) // Thin and not thin, Subsurface and TwoSided Foliage
+				else if (SSSType == SSS_TYPE_WRAP || SSSType == SSS_TYPE_TWO_SIDED_WRAP)
 				{
 					Out.CustomData.rgb = EncodeSubsurfaceColor(Out.SubsurfaceColor);
 					Out.CustomData.a   = SubstrateSubSurfaceGetWrapOpacityFromAnisotropy(SLAB_SSSPHASEANISOTROPY(CurrentBSDF));
diff --git a/Engine/Shaders/Private/Substrate/SubstrateLegacyConversion.ush b/Engine/Shaders/Private/Substrate/SubstrateLegacyConversion.ush
--- a/Engine/Shaders/Private/Substrate/SubstrateLegacyConversion.ush
+++ b/Engine/Shaders/Private/Substrate/SubstrateLegacyConversion.ush
@@ -41,7 +41,7 @@ struct FSubstrateLegacyParameters
 	float3 SSSMFP;
 	float  SSSMFPScale;
 	float  SSSPhaseAniso;
-	bool   UseSSSDiffusion;
+	uint   SSSType;
 	float3 Emissive;
 	float  SecondRoughness;
 	float  SecondRoughnessWeight;
@@ -75,7 +75,7 @@ FSubstrateLegacyParameters InitSubstrateLegacyParameters(FSubstratePixelFootprin
 	Out.SSSMFP = Zeros;
 	Out.SSSMFPScale = 1.f;
 	Out.SSSPhaseAniso = 0.f;
-	Out.UseSSSDiffusion = false;
+	Out.SSSType = SSS_TYPE_INVALID;
 	Out.Emissive = Zeros;
 	Out.SecondRoughness = 0.0f;
 	Out.SecondRoughnessWeight = 0.0f;
@@ -118,7 +118,7 @@ FSubstrateData CreateLegacySlab(
 		In.SSSMFP,											// SSSMFP
 		In.SSSMFPScale,										// SSSMFPScale
 		In.SSSPhaseAniso,									// SSSPhaseAniso
-		In.UseSSSDiffusion,									// UseSSSDiffusion
+		In.SSSType,											// SSSType
 		In.Emissive,										// Emissive
 		In.SecondRoughness,									// SecondRoughness
 		In.SecondRoughnessWeight,							// SecondRoughnessWeight
@@ -271,7 +271,7 @@ FSubstrateData SubstrateConvertLegacyMaterialDynamic(
 		LegacySlab.SSSMFPScale = 1.0f;
 		LegacySlab.SSSPhaseAniso = OneMinusOpacity; // Opaque-Thick: isotropic phase function, Thin: forward phase scattering function
 		LegacySlab.Emissive = Emissive;
-		LegacySlab.UseSSSDiffusion = false;
+		LegacySlab.SSSType = SSS_TYPE_WRAP;
 		LegacySlab.Thickness = ThicknessInCm;
 
 		FinalWeight = 1.0f;
@@ -305,7 +305,7 @@ FSubstrateData SubstrateConvertLegacyMaterialDynamic(
 		LegacySlab.SSSProfileID = 0.f;
 		LegacySlab.SSSMFPScale = OneMinusOpacity;
 		LegacySlab.SSSPhaseAniso = 0.93f;
-		LegacySlab.UseSSSDiffusion = true;
+		LegacySlab.SSSType = SSS_TYPE_DIFFUSION;
 		LegacySlab.Emissive = Emissive;
 
 		FinalWeight = 1.0f;
@@ -328,7 +328,7 @@ FSubstrateData SubstrateConvertLegacyMaterialDynamic(
 			LegacySlab.SSSMFPScale = Opacity;
 		}
 		LegacySlab.SSSPhaseAniso = 0.93f;
-		LegacySlab.UseSSSDiffusion = true;
+		LegacySlab.SSSType = SSS_TYPE_DIFFUSION_PROFILE;
 		LegacySlab.Emissive = Emissive;
 
 		FinalWeight = 1.0f;
@@ -373,7 +373,7 @@ FSubstrateData SubstrateConvertLegacyMaterialDynamic(
 		LegacySlab.SSSMFPScale = 1.0f;
 		LegacySlab.SSSPhaseAniso = OneMinusOpacity; // Opaque-Thick: isotropic phase function, Thin: forward phase scattering function
 		LegacySlab.Emissive = Emissive;
-		LegacySlab.UseSSSDiffusion = false;
+		LegacySlab.SSSType = SSS_TYPE_TWO_SIDED_WRAP;
 		LegacySlab.Thickness = Thickness;
 		LegacySlab.bIsThinSurface = true;
 
diff --git a/Engine/Shaders/Private/Substrate/SubstrateMaterialClassification.usf b/Engine/Shaders/Private/Substrate/SubstrateMaterialClassification.usf
--- a/Engine/Shaders/Private/Substrate/SubstrateMaterialClassification.usf
+++ b/Engine/Shaders/Private/Substrate/SubstrateMaterialClassification.usf
@@ -152,7 +152,7 @@ void TileMainCS(uint2 DispatchThreadId : SV_DispatchThreadID, uint LinearIndex :
 
 		// Output/Patch SSS data for legacy encoding (this allows to save ALU & bandwidth during the base pass0
 		uint OptimisedLegacyMode = ((PackedHeader >> (HEADER_SINGLEENCODING_BIT_COUNT)) & HEADER_SINGLE_OPTLEGACYMODE_BIT_MASK);
-		const bool bIsLegacyWrapOrWrapThin = OptimisedLegacyMode == SINGLE_OPTLEGACYMODE_SSSWRAP || OptimisedLegacyMode == SINGLE_OPTLEGACYMODE_SSSWRAP_THIN; // Wrap and Wrap thin have same packing
+		const bool bIsLegacyWrapOrWrapThin = OptimisedLegacyMode == SINGLE_OPTLEGACYMODE_SSSWRAP || OptimisedLegacyMode == SINGLE_OPTLEGACYMODE_TWO_SIDED_SSSWRAP; // Wrap and Wrap thin have same packing
 		const bool bIsLegacySSSProfile = OptimisedLegacyMode == SINGLE_OPTLEGACYMODE_SSSPROFILE;
 		if (bIsSingle && (bIsLegacyWrapOrWrapThin || bIsLegacySSSProfile))
 		{
@@ -161,7 +161,14 @@ void TileMainCS(uint2 DispatchThreadId : SV_DispatchThreadID, uint LinearIndex :
 			{
 				const uint PackedSSSWOpacity7bits = (PackedHeader >> (HEADER_SINGLEENCODING_BIT_COUNT + HEADER_SINGLE_OPTLEGACYMODE_BIT_COUNT)) & 0x7F;
 				FSubstrateSubsurfaceHeader SSSHeader = (FSubstrateSubsurfaceHeader)0;
-				SubstrateSubSurfaceHeaderSetSSSType(SSSHeader, SSS_TYPE_WRAP);
+				if (OptimisedLegacyMode == SINGLE_OPTLEGACYMODE_SSSWRAP)
+				{
+					SubstrateSubSurfaceHeaderSetSSSType(SSSHeader, SSS_TYPE_WRAP);
+				}
+				else
+				{
+					SubstrateSubSurfaceHeaderSetSSSType(SSSHeader, SSS_TYPE_TWO_SIDED_WRAP);
+				}
 				SubstrateSubSurfaceHeaderSetWrapOpacity(SSSHeader, UnpackR7(PackedSSSWOpacity7bits));
 				SubstrateStoreSubsurfaceHeader(MaterialTextureArrayUAV, FirstSliceStoringSubstrateSSSData, PixelCoord, SSSHeader.Bytes);
 			}
diff --git a/Engine/Shaders/Private/Substrate/SubstrateSubsurface.ush b/Engine/Shaders/Private/Substrate/SubstrateSubsurface.ush
--- a/Engine/Shaders/Private/Substrate/SubstrateSubsurface.ush
+++ b/Engine/Shaders/Private/Substrate/SubstrateSubsurface.ush
@@ -17,26 +17,26 @@ float4 UnpackColorGamma2ToLinearAlphaLinear(uint In);
 // This is hot data that is often accessed by SSS related passes.
 struct FSubstrateSubsurfaceHeader
 {
-	// 2 bits
+	// 3 bits
 	// SSS type:
 	// 0: Invalid
 	// 1: Wrap
-	// 2: Diffusion
-	// 3: Diffusion with profile
+	// 2: TwoSided
+	// 3: Diffusion
+	// 4: Diffusion with profile
+	// 5: SimpleVolume
 
-	// 30 bits of data.
+	// 29 bits of data.
 	// If  bIsProfile: 8bit profile radius, 8bits ProfileId, 16bits unused
-	// If !bIsProfile: 101010 mean free path
+	// If !bIsProfile: 10109 mean free path
 	uint Bytes;
 };
 
-// Note: SSSHeader only stores INVALID, WRAP, DIFFUSION, DIFFUSION_PROFILE type. This is because the type is represented onto 2bits and 
-//       other SSS types don't have/need side payload. This means SIMPLEVOLUME is trunked.
-#if SSS_TYPE_COUNT != 5
-#error Update this code to ensure all SSS types requiering side payload can be represented in the 2bits allocated for types.
+#if SSS_TYPE_COUNT != 6
+#error Update this code to ensure all SSS types requiering side payload can be represented.
 #endif
-#define SSSHEADER_TYPE(Header)	(Header.Bytes & 0x3)
-#define SSSHEADER_TYPE_MASK		(0x00000003)
+#define SSSHEADER_TYPE(Header)	(Header.Bytes & 0x7)
+#define SSSHEADER_TYPE_MASK		(0x00000007)
 
 void SubstrateSubSurfaceHeaderSetSSSType(inout FSubstrateSubsurfaceHeader SSSHeader, uint SSSType)
 {
@@ -54,7 +54,8 @@ void SubstrateSubSurfaceHeaderSetProfile(inout FSubstrateSubsurfaceHeader SSSHea
 void SubstrateSubSurfaceHeaderSetNonProfile(inout FSubstrateSubsurfaceHeader SSSHeader, float3 MeanFreePath)
 {
 	SSSHeader.Bytes &= SSSHEADER_TYPE_MASK;
-	SSSHeader.Bytes |= (Pack10F(MeanFreePath.x) << 22) | (Pack10F(MeanFreePath.y) << 12) | (Pack10F(MeanFreePath.z) << 2);
+	// Encode MFP onto 10 10 9 bits. For the blue channel the last bit of the mantissa is dropped in order to make room for SSSType encoding
+	SSSHeader.Bytes |= (Pack10F(MeanFreePath.x) << 22) | (Pack10F(MeanFreePath.y) << 12) | ((Pack10F(MeanFreePath.z)&0x3FE) << 2);
 }
 
 float SubstrateSubSurfaceGetWrapOpacityFromAnisotropy(float PhaseAnisotropy)
@@ -69,7 +70,7 @@ float SubstrateSubSurfaceGetWrapOpacityFromAnisotropy(float PhaseAnisotropy)
 void SubstrateSubSurfaceHeaderSetWrapOpacity(inout FSubstrateSubsurfaceHeader SSSHeader, float Opacity)
 {
 	SSSHeader.Bytes &= SSSHEADER_TYPE_MASK;
-	SSSHeader.Bytes |= PackR8(Opacity) << 2;
+	SSSHeader.Bytes |= PackR8(Opacity) << 8;
 }
 
 void SubstrateSubSurfaceHeaderSetWrap(inout FSubstrateSubsurfaceHeader SSSHeader, float PhaseAnisotropy)
@@ -102,7 +103,7 @@ bool SubstrateSubSurfaceHeaderGetIsProfile(in FSubstrateSubsurfaceHeader SSSHead
 
 bool SubstrateSubSurfaceHeaderGetIsWrap(in FSubstrateSubsurfaceHeader SSSHeader)
 {
-	return SSSHEADER_TYPE(SSSHeader) == SSS_TYPE_WRAP;
+	return SSSHEADER_TYPE(SSSHeader) == SSS_TYPE_WRAP || SSSHEADER_TYPE(SSSHeader) == SSS_TYPE_TWO_SIDED_WRAP;
 }
 
 uint SubstrateSubSurfaceHeaderGetSSSType(in FSubstrateSubsurfaceHeader SSSHeader)
@@ -125,15 +126,22 @@ float SubstrateSubSurfaceHeaderGetProfileRadiusScale(in FSubstrateSubsurfaceHead
 
 float3 SubstrateSubSurfaceHeaderGetMFP(in FSubstrateSubsurfaceHeader SSSHeader)				// Only when !IsProfile
 {
+	// xxxxxxxxxx
+	//           xxxxxxxxxx
+	//                     xxxxxxxxx
+	//                              xxx
+	// 10bits    10bits    9bits    3bits
+	// MFP.R     MFP.G     MFP.B    SSSType
 	return float3(
-		Unpack10F(SSSHeader.Bytes >> 22),
-		Unpack10F(SSSHeader.Bytes >> 12),
-		Unpack10F(SSSHeader.Bytes >> 2));
+		Unpack10F(SSSHeader.Bytes >> 22),			// 10bits
+		Unpack10F(SSSHeader.Bytes >> 12),			// 10bits
+		Unpack10F((SSSHeader.Bytes >> 2) & 0x3FE)); //  9bits
 }
 
 float SubstrateSubSurfaceHeaderGetWrapOpacity(in FSubstrateSubsurfaceHeader SSSHeader)
 {
-	return SSSHEADER_TYPE(SSSHeader) == SSS_TYPE_WRAP ? UnpackR8(SSSHeader.Bytes >> 2) : 1.0f;
+	const uint SSSType = SSSHEADER_TYPE(SSSHeader);
+	return (SSSType == SSS_TYPE_WRAP || SSSType == SSS_TYPE_TWO_SIDED_WRAP) ? UnpackR8(SSSHeader.Bytes >> 8) : 1.0f;
 }
 
 float SubstrateSubSurfaceHeaderGetOpacity(in FSubstrateSubsurfaceHeader SSSHeader)
@@ -146,7 +154,7 @@ float SubstrateSubSurfaceHeaderGetOpacity(in FSubstrateSubsurfaceHeader SSSHeade
 	// * SSS Diffusion Profile -> SubstrateSubSurfaceHeaderGetProfileRadiusScale()
 	float Opacity = 1.0f;
 	Opacity = SSSType == SSS_TYPE_DIFFUSION_PROFILE ? SubstrateSubSurfaceHeaderGetProfileRadiusScale(SSSHeader) : Opacity;
-	Opacity = SSSType == SSS_TYPE_WRAP ? SubstrateSubSurfaceHeaderGetWrapOpacity(SSSHeader) : Opacity;
+	Opacity = (SSSType == SSS_TYPE_WRAP || SSSType == SSS_TYPE_TWO_SIDED_WRAP) ? SubstrateSubSurfaceHeaderGetWrapOpacity(SSSHeader) : Opacity;
 	return Opacity;
 }
 
@@ -269,7 +277,7 @@ FSubsurfaceOpacityMFP SubstrateGetSubsurfaceOpacityMFP(FSubstrateSubsurfaceHeade
 		FTransmissionProfileParams TransmissionParams = GetTransmissionProfileParams(SubstrateSubSurfaceHeaderGetProfileId(SSSHeader));
 		SubsurfaceOpacityMFP.Density = SubsurfaceDensityFromOpacity(SubsurfaceOpacityMFP.Data) * TransmissionParams.ExtinctionScale * 3.1f;	// This matches the weird computations done in CalcTransmissionThickness
 	}
-	else if (SSSType == SSS_TYPE_WRAP)
+	else if (SSSType == SSS_TYPE_WRAP || SSSType == SSS_TYPE_TWO_SIDED_WRAP)
 	{
 		SubsurfaceOpacityMFP.bDataIsOpacity = true;
 
diff --git a/Engine/Shaders/Private/Substrate/SubstrateVisualizeCommon.ush b/Engine/Shaders/Private/Substrate/SubstrateVisualizeCommon.ush
--- a/Engine/Shaders/Private/Substrate/SubstrateVisualizeCommon.ush
+++ b/Engine/Shaders/Private/Substrate/SubstrateVisualizeCommon.ush
@@ -90,7 +90,7 @@ void PrintPixelType(inout FShaderPrintContext Ctx, in FSubstratePixelHeader Head
 
 void PrintSSSType(inout FShaderPrintContext Ctx, uint SSSType, bool bThin, FFontColor InColor)
 {
-		 if (SSSType == SSS_TYPE_WRAP && bThin)		Print(Ctx, TEXT("Wrap & Thin"), InColor);
+		 if (SSSType == SSS_TYPE_TWO_SIDED_WRAP)	Print(Ctx, TEXT("Two-Sided Wrap"), InColor);
 	else if (SSSType == SSS_TYPE_WRAP)				Print(Ctx, TEXT("Wrap"), InColor);
 	else if (SSSType == SSS_TYPE_DIFFUSION)			Print(Ctx, TEXT("Diffusion"), InColor);
 	else if (SSSType == SSS_TYPE_DIFFUSION_PROFILE)	Print(Ctx, TEXT("Diffusion Profile"), InColor);
diff --git a/Engine/Shaders/Shared/SubstrateDefinitions.h b/Engine/Shaders/Shared/SubstrateDefinitions.h
--- a/Engine/Shaders/Shared/SubstrateDefinitions.h
+++ b/Engine/Shaders/Shared/SubstrateDefinitions.h
@@ -122,13 +122,14 @@
 #define SINGLE_OPTLEGACYMODE_CLOTH				2
 #define SINGLE_OPTLEGACYMODE_SSSWRAP			3
 #define SINGLE_OPTLEGACYMODE_SSSPROFILE			4
-#define SINGLE_OPTLEGACYMODE_SSSWRAP_THIN		5
+#define SINGLE_OPTLEGACYMODE_TWO_SIDED_SSSWRAP	5
 // 5 more remains: Also how to deal with Aniso?
 
 // Sub-surface scattering types
 #define SSS_TYPE_INVALID						(0u)
 #define SSS_TYPE_WRAP							(1u)
-#define SSS_TYPE_DIFFUSION						(2u)
-#define SSS_TYPE_DIFFUSION_PROFILE				(3u)
-#define SSS_TYPE_SIMPLEVOLUME					(4u)
-#define SSS_TYPE_COUNT							(5u)
+#define SSS_TYPE_TWO_SIDED_WRAP					(2u)
+#define SSS_TYPE_DIFFUSION						(3u)
+#define SSS_TYPE_DIFFUSION_PROFILE				(4u)
+#define SSS_TYPE_SIMPLEVOLUME					(5u)
+#define SSS_TYPE_COUNT							(6u)
diff --git a/Engine/Shaders/Shared/SubstrateVisualizeDefinitions.h b/Engine/Shaders/Shared/SubstrateVisualizeDefinitions.h
--- a/Engine/Shaders/Shared/SubstrateVisualizeDefinitions.h
+++ b/Engine/Shaders/Shared/SubstrateVisualizeDefinitions.h
@@ -378,7 +378,7 @@ FSubstratePixelDebugData ConvertToSerializableSubstratePixelDebugData(uint2 InCo
 
 			if (Data.Closures[i].SSSType != SSS_TYPE_INVALID || Data.Closures[i].bIsThin)
 			{
-				if (Data.Closures[i].SSSType == SSS_TYPE_WRAP)
+				if (Data.Closures[i].SSSType == SSS_TYPE_WRAP || Data.Closures[i].SSSType == SSS_TYPE_TWO_SIDED_WRAP)
 				{
 				#if SUBSTRATE_MATERIALCONTAINER_IS_VIEWRESOURCE
 					Data.Closures[i].SSSOpacity = SubstrateSubSurfaceHeaderGetWrapOpacity(SSSHeader);
diff --git a/Engine/Source/Runtime/Core/Private/UObject/FortniteMainBranchObjectVersion.cpp b/Engine/Source/Runtime/Core/Private/UObject/FortniteMainBranchObjectVersion.cpp
--- a/Engine/Source/Runtime/Core/Private/UObject/FortniteMainBranchObjectVersion.cpp
+++ b/Engine/Source/Runtime/Core/Private/UObject/FortniteMainBranchObjectVersion.cpp
@@ -14,6 +14,6 @@ TMap<FGuid, FGuid> FFortniteMainBranchObjectVersion::GetSystemGuids()
 	SystemGuids.Add(DevGuids.Niagara_LatestScriptCompileVersion, FGuid("670A2EBDEE1A4794BB77BBFF331EA7EF"));
 	SystemGuids.Add(DevGuids.SkeletalMeshDerivedDataVersion, FGuid("2F90180769024483BB7D90AACEB2798D"));
 	SystemGuids.Add(DevGuids.STATICMESH_DERIVEDDATA_VER, FGuid("656E338C78C34641A5D3BFAE303997AF"));
-	SystemGuids.Add(DevGuids.MaterialTranslationDDCVersion, FGuid("1F1845B6EF8F49ED880DE4EA60BF2EA4"));
+	SystemGuids.Add(DevGuids.MaterialTranslationDDCVersion, FGuid("C36543C38CBF45A1ABFB5775B8B93C13"));
 	return SystemGuids;
 }
diff --git a/Engine/Source/Runtime/Core/Public/UObject/FortniteMainBranchObjectVersions.inl b/Engine/Source/Runtime/Core/Public/UObject/FortniteMainBranchObjectVersions.inl
--- a/Engine/Source/Runtime/Core/Public/UObject/FortniteMainBranchObjectVersions.inl
+++ b/Engine/Source/Runtime/Core/Public/UObject/FortniteMainBranchObjectVersions.inl
@@ -574,5 +574,8 @@ FORTNITE_MAIN_VERSION(AddClothAssetBase, 186)
 // Add inline constant default values to the PCG graph nodes.
 FORTNITE_MAIN_VERSION(PCGInlineConstantDefaultValues, 187)
 
+// Add MaterialSubstrateSubsurfaceType type to UMaterialExpressionSubstrateSlabBSDF for replacing bUseSSSDifffusion
+FORTNITE_MAIN_VERSION(AddMaterialSubstrateSubsurfaceType, 188)
+
 // -----<new versions can be added above this line>-------------------------------------------------
 #undef FORTNITE_MAIN_VERSION
\ No newline at end of file
diff --git a/Engine/Source/Runtime/Engine/Private/Materials/HLSLMaterialTranslator.cpp b/Engine/Source/Runtime/Engine/Private/Materials/HLSLMaterialTranslator.cpp
--- a/Engine/Source/Runtime/Engine/Private/Materials/HLSLMaterialTranslator.cpp
+++ b/Engine/Source/Runtime/Engine/Private/Materials/HLSLMaterialTranslator.cpp
@@ -13872,7 +13872,7 @@ FString FHLSLMaterialTranslator::SubstrateGetCastParameterCodeWithDeriv(int32 In
 int32 FHLSLMaterialTranslator::SubstrateSlabBSDF(
 	int32 DiffuseAlbedo, int32 F0, int32 F90,
 	int32 Roughness, int32 Anisotropy,
-	int32 SSSProfileId, int32 SSSMFP, int32 SSSMFPScale, int32 SSSPhaseAniso, int32 UseSSSDiffusion,
+	int32 SSSProfileId, int32 SSSMFP, int32 SSSMFPScale, int32 SSSPhaseAniso, int32 SSSType,
 	int32 EmissiveColor,
 	int32 SecondRoughness, int32 SecondRoughnessWeight, int32 SecondRoughnessAsSimpleClearCoat,
 	int32 FuzzAmount, int32 FuzzColor, int32 FuzzRoughness,
@@ -13907,7 +13907,7 @@ int32 FHLSLMaterialTranslator::SubstrateSlabBSDF(
 				/*SSSMFP*/                           %s,\n\
 				/*SSSMFPScale*/                      %s,\n\
 				/*SSSPhaseAniso*/                    %s,\n\
-				/*UseSSSDiffusion*/                  %s,\n\
+				/*SSSType*/                          %s,\n\
 				/*EmissiveColor*/                    %s,\n\
 				/*SecondRoughness*/                  %s,\n\
 				/*SecondRoughnessWeight*/            %s,\n\
@@ -13927,16 +13927,16 @@ int32 FHLSLMaterialTranslator::SubstrateSlabBSDF(
 			*GetParametersSubstrateTreeName(CurrentSubstrateCompilationContext),
 			*SubstrateGetCastParameterCode(DiffuseAlbedo,			MCT_Float3),
 			*SubstrateGetCastParameterCode(F0,						MCT_Float3),
-			*SubstrateGetCastParameterCode(F90,					MCT_Float3),
+			*SubstrateGetCastParameterCode(F90,						MCT_Float3),
 			*SubstrateGetCastParameterCode(Roughness,				MCT_Float),
 			*SubstrateGetCastParameterCode(Anisotropy,				MCT_Float),
 			*SubstrateGetCastParameterCode(SSSProfileId,			MCT_Float),
 			*SubstrateGetCastParameterCode(SSSMFP,					MCT_Float3),
-			*SubstrateGetCastParameterCode(SSSMFPScale,			MCT_Float),
+			*SubstrateGetCastParameterCode(SSSMFPScale,				MCT_Float),
 			*SubstrateGetCastParameterCode(SSSPhaseAniso,			MCT_Float),
-			*SubstrateGetCastParameterCode(UseSSSDiffusion,		MCT_Float),
+			*SubstrateGetCastParameterCode(SSSType,					MCT_Float),
 			*SubstrateGetCastParameterCode(EmissiveColor,			MCT_Float3),
-			*SubstrateGetCastParameterCode(SecondRoughness,		MCT_Float),
+			*SubstrateGetCastParameterCode(SecondRoughness,			MCT_Float),
 			*SubstrateGetCastParameterCode(SecondRoughnessWeight,	MCT_Float),
 			*SubstrateGetCastParameterCode(SecondRoughnessAsSimpleClearCoat, MCT_Float),
 			*SubstrateGetCastParameterCode(FuzzAmount,				MCT_Float),
@@ -13971,7 +13971,7 @@ int32 FHLSLMaterialTranslator::SubstrateSlabBSDF(
 			/*SSSMFP*/                           %s,\n\
 			/*SSSMFPScale*/                      %s,\n\
 			/*SSSPhaseAniso*/                    %s,\n\
-			/*UseSSSDiffusion*/                  %s,\n\
+			/*SSSType*/                          %s,\n\
 			/*EmissiveColor*/                    %s,\n\
 			/*SecondRoughness*/                  %s,\n\
 			/*SecondRoughnessWeight*/            %s,\n\
@@ -13989,16 +13989,16 @@ int32 FHLSLMaterialTranslator::SubstrateSlabBSDF(
 			Parameters.%s.Types) /* Normal = %s ; Tangent = %s ; Thickness = %s */"),
 		*SubstrateGetCastParameterCode(DiffuseAlbedo,			MCT_Float3),
 		*SubstrateGetCastParameterCode(F0,						MCT_Float3),
-		*SubstrateGetCastParameterCode(F90,					MCT_Float3),
+		*SubstrateGetCastParameterCode(F90,						MCT_Float3),
 		*SubstrateGetCastParameterCode(Roughness,				MCT_Float),
 		*SubstrateGetCastParameterCode(Anisotropy,				MCT_Float),
 		*SubstrateGetCastParameterCode(SSSProfileId,			MCT_Float),
 		*SubstrateGetCastParameterCode(SSSMFP,					MCT_Float3),
-		*SubstrateGetCastParameterCode(SSSMFPScale,			MCT_Float),
+		*SubstrateGetCastParameterCode(SSSMFPScale,				MCT_Float),
 		*SubstrateGetCastParameterCode(SSSPhaseAniso,			MCT_Float),
-		*SubstrateGetCastParameterCode(UseSSSDiffusion,		MCT_Float),
+		*SubstrateGetCastParameterCode(SSSType,					MCT_Float),
 		*SubstrateGetCastParameterCode(EmissiveColor,			MCT_Float3),
-		*SubstrateGetCastParameterCode(SecondRoughness,		MCT_Float),
+		*SubstrateGetCastParameterCode(SecondRoughness,			MCT_Float),
 		*SubstrateGetCastParameterCode(SecondRoughnessWeight,	MCT_Float),
 		*SubstrateGetCastParameterCode(SecondRoughnessAsSimpleClearCoat, MCT_Float),
 		*SubstrateGetCastParameterCode(FuzzAmount,				MCT_Float),
@@ -14007,7 +14007,7 @@ int32 FHLSLMaterialTranslator::SubstrateSlabBSDF(
 		*SubstrateGetCastParameterCode(GlintValue,				MCT_Float),
 		*SubstrateGetCastParameterCodeWithDeriv(GlintUV, 		MCT_Float2),
 		*SubstrateGetCastParameterCode(SpecularProfileId,		MCT_Float),
-		*SubstrateGetCastParameterCode(Thickness, MCT_Float),
+		*SubstrateGetCastParameterCode(Thickness, 				MCT_Float),
 		bIsThinSurface ? TEXT("true") : TEXT("false"),
 		bIsAtTheBottomOfTopology ? TEXT("true") : TEXT("false"),
 		*SharedLocalBasisIndexMacro,
diff --git a/Engine/Source/Runtime/Engine/Private/Materials/HLSLMaterialTranslator.h b/Engine/Source/Runtime/Engine/Private/Materials/HLSLMaterialTranslator.h
--- a/Engine/Source/Runtime/Engine/Private/Materials/HLSLMaterialTranslator.h
+++ b/Engine/Source/Runtime/Engine/Private/Materials/HLSLMaterialTranslator.h
@@ -1320,7 +1320,7 @@ protected:
 	virtual int32 SubstrateSlabBSDF(
 		int32 DiffuseAlbedo, int32 F0, int32 F90,
 		int32 Roughness, int32 Anisotropy,
-		int32 SSSProfileId, int32 SSSMFP, int32 SSSMFPScale, int32 SSSPhaseAniso, int32 bUseSSSDiffusion,
+		int32 SSSProfileId, int32 SSSMFP, int32 SSSMFPScale, int32 SSSPhaseAniso, int32 SSSType,
 		int32 EmissiveColor, 
 		int32 SecondRoughness, int32 SecondRoughnessWeight, int32 SecondRoughnessAsSimpleClearCoat,
 		int32 FuzzAmount, int32 FuzzColor, int32 FuzzRoughness,
diff --git a/Engine/Source/Runtime/Engine/Private/Materials/Material.cpp b/Engine/Source/Runtime/Engine/Private/Materials/Material.cpp
--- a/Engine/Source/Runtime/Engine/Private/Materials/Material.cpp
+++ b/Engine/Source/Runtime/Engine/Private/Materials/Material.cpp
@@ -4249,6 +4249,7 @@ void UMaterial::PostLoad()
 	const FPackageFileVersion UEVer = GetLinkerUEVersion();
 	const int32 RenderObjVer = GetLinkerCustomVersion(FRenderingObjectVersion::GUID);
 	const int32 UE5MainVer = GetLinkerCustomVersion(FUE5MainStreamObjectVersion::GUID);
+	const int32 FortMainVer = GetLinkerCustomVersion(FFortniteMainBranchObjectVersion::GUID);
 
 	UMaterialEditorOnlyData* EditorOnly = GetEditorOnlyData();
 	check(EditorOnly != nullptr);
@@ -4502,8 +4503,33 @@ void UMaterial::PostLoad()
 #endif
 
 #if WITH_EDITOR
+	// Convert UMaterialExpressionSubstrateSlabBSDF::bUseSSSDiffution into UMaterialExpressionSubstrateSlabBSDF::SubSurfaceType
 	if (Substrate::IsSubstrateEnabled())
 	{
+		if (FortMainVer < FFortniteMainBranchObjectVersion::AddMaterialSubstrateSubsurfaceType)
+		{
+			for (UMaterialExpression* Expression : GetExpressions())
+			{
+				if (UMaterialExpressionSubstrateSlabBSDF* SlabExpression = Cast<UMaterialExpressionSubstrateSlabBSDF>(Expression))
+				{
+					PRAGMA_DISABLE_DEPRECATION_WARNINGS
+					if (SlabExpression->bUseSSSDiffusion)
+					{
+						SlabExpression->SubSurfaceType = EMaterialSubSurfaceType::MSS_Diffusion;
+					}
+					else if (bIsThinSurface)
+					{
+						SlabExpression->SubSurfaceType = EMaterialSubSurfaceType::MSS_TwoSidedWrap;
+					}
+					else
+					{
+						SlabExpression->SubSurfaceType = EMaterialSubSurfaceType::MSS_Wrap;
+					}
+					PRAGMA_ENABLE_DEPRECATION_WARNINGS
+				}
+			}
+		}
+
 		// Substrate materials conversion needs to be done after expressions are cached, otherwise material function won't have 
 		// valid inputs in certain cases
 		if (ConvertMaterialToSubstrateMaterial(true /*bAllowEmptyMaterialUpdate*/))
diff --git a/Engine/Source/Runtime/Engine/Private/Materials/MaterialExpressionSubstrate.cpp b/Engine/Source/Runtime/Engine/Private/Materials/MaterialExpressionSubstrate.cpp
--- a/Engine/Source/Runtime/Engine/Private/Materials/MaterialExpressionSubstrate.cpp
+++ b/Engine/Source/Runtime/Engine/Private/Materials/MaterialExpressionSubstrate.cpp
@@ -748,6 +748,7 @@ int32 UMaterialExpressionSubstrateBSDF::CompilePreview(class FMaterialCompiler*
 UMaterialExpressionSubstrateSlabBSDF::UMaterialExpressionSubstrateSlabBSDF(const FObjectInitializer& ObjectInitializer)
 	: Super(ObjectInitializer)
 	, bUseSSSDiffusion(true)
+	, SubSurfaceType(EMaterialSubSurfaceType::MSS_Diffusion)
 {
 	struct FConstructorStatics
 	{
@@ -758,6 +759,9 @@ UMaterialExpressionSubstrateSlabBSDF::UMaterialExpressionSubstrateSlabBSDF(const
 #if WITH_EDITORONLY_DATA
 	MenuCategories.Add(ConstructorStatics.NAME_Substrate);
 #endif
+
+	// Sanity check
+	static_assert(MSS_MAX == SSS_TYPE_COUNT);
 }
 
 #if WITH_EDITOR
@@ -864,7 +868,7 @@ int32 UMaterialExpressionSubstrateSlabBSDF::Compile(class FMaterialCompiler* Com
 		SSSMFPCodeChunk,
 		SSSMFPScaleCodeChunk,
 		SSSPhaseAnisotropyCodeChunk,
-		bUseSSSDiffusion ? Compiler->Constant(1.0f) : Compiler->Constant(0.0f),
+		Compiler->Constant(SubSurfaceType),
 		CompileWithDefaultFloat3(Compiler, EmissiveColor, 0.0f, 0.0f, 0.0f),
 		SecondRoughnessCodeChunk,
 		SecondRoughnessWeightCodeChunk,
@@ -934,7 +938,7 @@ int32 UMaterialExpressionSubstrateSlabBSDF::CompileDefaultSlab(FMaterialCompiler
 		SSSMFPCodeChunk,
 		SSSMFPScaleCodeChunk,
 		SSSPhaseAnisotropyCodeChunk,
-		Compiler->Constant(0.0f),
+		Compiler->Constant(EMaterialSubSurfaceType::MSS_None),
 		Compiler->Constant3(EmissiveOverride.X, EmissiveOverride.Y, EmissiveOverride.Z),
 		SecondRoughnessCodeChunk,
 		SecondRoughnessWeightCodeChunk,
@@ -1341,7 +1345,7 @@ int32 UMaterialExpressionSubstrateSimpleClearCoatBSDF::Compile(class FMaterialCo
 		Compiler->Constant3(0.0f, 0.0f, 0.0f),					// SSSMFP
 		Compiler->Constant(0.0f),								// SSSMFPScale
 		Compiler->Constant(0.0f),								// SSSPhaseAnisotropy
-		Compiler->Constant(0.0f),								// bUseSSSDiffusion
+		Compiler->Constant(EMaterialSubSurfaceType::MSS_None),	// SSSType
 		CompileWithDefaultFloat3(Compiler, EmissiveColor, 0.0f, 0.0f, 0.0f),
 		CompileWithDefaultFloat1(Compiler, ClearCoatRoughness, 0.1f),
 		CompileWithDefaultFloat1(Compiler, ClearCoatCoverage, 1.0f),
diff --git a/Engine/Source/Runtime/Engine/Public/MaterialCompiler.h b/Engine/Source/Runtime/Engine/Public/MaterialCompiler.h
--- a/Engine/Source/Runtime/Engine/Public/MaterialCompiler.h
+++ b/Engine/Source/Runtime/Engine/Public/MaterialCompiler.h
@@ -722,7 +722,7 @@ public:
 	virtual int32 SubstrateSlabBSDF(
 		int32 DiffuseAlbedo, int32 F0, int32 F90,
 		int32 Roughness, int32 Anisotropy,
-		int32 SSSProfileId, int32 SSSMFP, int32 SSSMFPScale, int32 SSSPhaseAniso, int32 UseSSSDiffusion,
+		int32 SSSProfileId, int32 SSSMFP, int32 SSSMFPScale, int32 SSSPhaseAniso, int32 SSSType,
 		int32 EmissiveColor, 
 		int32 SecondRoughness, int32 SecondRoughnessWeight, int32 SecondRoughnessAsSimpleClearCoat,
 		int32 FuzzAmount, int32 FuzzColor, int32 FuzzRoughness,
@@ -1411,7 +1411,7 @@ public:
 	virtual int32 SubstrateSlabBSDF(
 		int32 DiffuseAlbedo, int32 F0, int32 F90,
 		int32 Roughness, int32 Anisotropy,
-		int32 SSSProfileId, int32 SSSMFP, int32 SSSMFPScale, int32 SSSPhaseAniso, int32 UseSSSDiffusion,
+		int32 SSSProfileId, int32 SSSMFP, int32 SSSMFPScale, int32 SSSPhaseAniso, int32 SSSType,
 		int32 EmissiveColor, 
 		int32 SecondRoughness, int32 SecondRoughnessWeight, int32 SecondRoughnessAsSimpleClearCoat,
 		int32 FuzzAmount, int32 FuzzColor, int32 FuzzRoughness,
@@ -1425,7 +1425,7 @@ public:
 		return Compiler->SubstrateSlabBSDF(
 			DiffuseAlbedo, F0, F90,
 			Roughness, Anisotropy,
-			SSSProfileId, SSSMFP, SSSMFPScale, SSSPhaseAniso, UseSSSDiffusion,
+			SSSProfileId, SSSMFP, SSSMFPScale, SSSPhaseAniso, SSSType,
 			EmissiveColor, 
 			SecondRoughness, SecondRoughnessWeight, SecondRoughnessAsSimpleClearCoat,
 			FuzzAmount, FuzzColor, FuzzRoughness,
diff --git a/Engine/Source/Runtime/Engine/Public/Materials/MaterialExpressionSubstrate.h b/Engine/Source/Runtime/Engine/Public/Materials/MaterialExpressionSubstrate.h
--- a/Engine/Source/Runtime/Engine/Public/Materials/MaterialExpressionSubstrate.h
+++ b/Engine/Source/Runtime/Engine/Public/Materials/MaterialExpressionSubstrate.h
@@ -5,13 +5,24 @@
 #include "CoreMinimal.h"
 #include "Engine/EngineTypes.h"
 #include "Materials/MaterialExpression.h"
+#include "SubstrateDefinitions.h"
 #include "MaterialExpressionSubstrate.generated.h"
 
 class FMaterialCompiler;
 struct FSubstrateMaterialComplexity;
 struct FSubstrateOperator;
 
-
+UENUM()
+enum EMaterialSubSurfaceType : int
+{
+	MSS_None				= SSS_TYPE_INVALID 				UMETA(DisplayName="None"),
+	MSS_Wrap				= SSS_TYPE_WRAP 				UMETA(DisplayName="Wrap", ToolTip="Approximation using wrap-lighting. "),
+	MSS_TwoSidedWrap		= SSS_TYPE_TWO_SIDED_WRAP		UMETA(DisplayName="Two-Sided Wrap", ToolTip="Approximation using wrap-lighting and handling thin surface (e.g.: foliage)"),
+	MSS_Diffusion			= SSS_TYPE_DIFFUSION 			UMETA(DisplayName="Diffusion", ToolTip="Diffusion based sub-surface scattering"),
+	MSS_DiffusionProfile	= SSS_TYPE_DIFFUSION_PROFILE 	UMETA(Hidden),
+	MSS_SimpleVolume		= SSS_TYPE_SIMPLEVOLUME 		UMETA(DisplayName="Simple Volume", ToolTip="Approximation of optically thin slab (i.e., light visible through the material)"),
+	MSS_MAX
+};
 
 ///////////////////////////////////////////////////////////////////////////////
 // Functions
@@ -340,10 +351,17 @@ class UMaterialExpressionSubstrateSlabBSDF : public UMaterialExpressionSubstrate
 	UPROPERTY(EditAnywhere, BlueprintReadOnly, Category = Material, meta = (DisplayName = "Specular Profile"))
 	TObjectPtr<class USpecularProfile> SpecularProfile;
 
-	/** Whether to use light diffusion (i.e., SSS diffusion) or wrap-approximation for material with scattering behavior. This option trades quality over performance and will result into visual differences. */
-	UPROPERTY(EditAnywhere, Category = Mode, meta = (DisplayName = "Use Subsurface Diffusion"))
+	UE_DEPRECATED(5.6, "Please, use SubSurfaceType to enable/disable SSS-diffusion")
+	UPROPERTY()
 	uint32 bUseSSSDiffusion : 1;
 
+	/** Define sub-subsurface used behavior of the slab. This option trades quality over performance and will result into visual differences. 
+	 *   * For slab not sitting at the bottom of the topology (e.g. slabs stacked with a vertial operator), only the SimpleVolume SSS type is available.
+	 *   * In non-opaque blend modes, SSS Diffusion and SSS Diffusion Profile are not available and will fallback onto Wrap mode.
+	 */
+	UPROPERTY(EditAnywhere, Category = ShadingModel, meta = (DisplayName = "Sub-Surface Type"))
+	TEnumAsByte<enum EMaterialSubSurfaceType> SubSurfaceType = MSS_None;
+
 	//~ Begin UMaterialExpression Interface
 #if WITH_EDITOR
 	virtual int32 Compile(class FMaterialCompiler* Compiler, int32 OutputIndex) override;
diff --git a/Engine/Source/Runtime/Renderer/Private/Substrate/SubstrateVisualize.cpp b/Engine/Source/Runtime/Renderer/Private/Substrate/SubstrateVisualize.cpp
--- a/Engine/Source/Runtime/Renderer/Private/Substrate/SubstrateVisualize.cpp
+++ b/Engine/Source/Runtime/Renderer/Private/Substrate/SubstrateVisualize.cpp
@@ -350,8 +350,8 @@ void AddProcessAndPrintSubstrateMaterialPropertiesPasses(FRDGBuilder& GraphBuild
 				case SINGLE_OPTLEGACYMODE_SSSPROFILE:
 					MaterialMode += TEXT(" - Legacy SSS-Profile");
 					break;
-				case SINGLE_OPTLEGACYMODE_SSSWRAP_THIN:
-					MaterialMode += TEXT(" - Legacy Foliage");
+				case SINGLE_OPTLEGACYMODE_TWO_SIDED_SSSWRAP:
+					MaterialMode += TEXT(" - Legacy Two-Sided SSS-Wrap (Foliage)");
 					break;
 				default:
 					MaterialMode += TEXT(" - Unkown optimised legacy material mode, please update visualization code.");
@@ -607,9 +607,16 @@ void AddProcessAndPrintSubstrateMaterialPropertiesPasses(FRDGBuilder& GraphBuild
 
 						if (Closure.SSSType != SSS_TYPE_INVALID || Closure.bIsThin)
 						{
-							if (Closure.SSSType == SSS_TYPE_WRAP)
+							if (Closure.SSSType == SSS_TYPE_WRAP || Closure.SSSType == SSS_TYPE_TWO_SIDED_WRAP)
 							{
-								PrintS(DrawPosX + 0.0f, DrawPosY, TEXT("SSS Wrap (Legacy Subsurface)"), Orange);
+								if (Closure.SSSType == SSS_TYPE_WRAP)
+								{
+									PrintS(DrawPosX + 0.0f, DrawPosY, TEXT("SSS Wrap (Legacy Subsurface)"), Orange);
+								}
+								else
+								{
+									PrintS(DrawPosX + 0.0f, DrawPosY, TEXT("SSS Two-Sided Wrap (Legacy Foliage)"), Orange);
+								}
 								NewLine();
 
 								PrintS(DrawPosX + 0.0f, DrawPosY, TEXT(" - MFP"), Orange);
