Fix CIS shader compilation.

diff --git a/Engine/Shaders/Private/RayTracing/RayTracingOcclusionRGS.usf b/Engine/Shaders/Private/RayTracing/RayTracingOcclusionRGS.usf
--- a/Engine/Shaders/Private/RayTracing/RayTracingOcclusionRGS.usf
+++ b/Engine/Shaders/Private/RayTracing/RayTracingOcclusionRGS.usf
@@ -690,7 +690,7 @@ RAY_TRACING_ENTRY_RAYGEN_OR_INLINE(Occlusion)
 			ShadingParameters.bIsHair = BSDFType == SUBSTRATE_BSDF_TYPE_HAIR;
 			ShadingParameters.bIsSubsurface = bHasSSS ? SSSType == SSS_TYPE_WRAP || SSSType == SSS_TYPE_TWO_SIDED_WRAP : false;
 			ShadingParameters.bHasSubsurfaceProfile = bHasSSS ? (SSSType == SSS_TYPE_DIFFUSION_PROFILE || SSSType == SSS_TYPE_DIFFUSION) : false;
-			ShadingParameters.bBackfaceCulling = !(BSDFType == SUBSTRATE_BSDF_TYPE_EYE || (bHasSSS && SSSType == SSS_TYPE_WRAP || SSSType == SSS_TYPE_TWO_SIDED_WRAP));
+			ShadingParameters.bBackfaceCulling = !(BSDFType == SUBSTRATE_BSDF_TYPE_EYE || (bHasSSS && (SSSType == SSS_TYPE_WRAP || SSSType == SSS_TYPE_TWO_SIDED_WRAP)));
 
 			// The SSDenoiser reads the top layer normal in order to know how to process the data.
 			// RTShadow traces using the top layer normal when NoL>0 only. So if normals for slabs below do not match, light leaking will occur.
