Fix shader compilation error in SM5.

diff --git a/Engine/Shaders/Private/HZB.ush b/Engine/Shaders/Private/HZB.ush
--- a/Engine/Shaders/Private/HZB.ush
+++ b/Engine/Shaders/Private/HZB.ush
@@ -2,6 +2,10 @@
 
 #pragma once
 
+#ifndef USE_HZB_SHARED_SAMPLERS
+#define USE_HZB_SHARED_SAMPLERS 1
+#endif
+
 #ifndef USE_HZB_HALF_TYPE
 #define USE_HZB_HALF_TYPE 0
 #endif
@@ -31,7 +35,7 @@ uint bIsHZBValid;
 uint bIsFurthestHZBValid;
 uint bIsClosestHZBValid;
 
-#if SUPPORTS_INDEPENDENT_SAMPLERS
+#if SUPPORTS_INDEPENDENT_SAMPLERS && USE_HZB_SHARED_SAMPLERS
 #define HZBSampler GlobalPointClampedSampler
 #define HZBTextureSampler GlobalPointClampedSampler
 #define FurthestHZBTextureSampler GlobalPointClampedSampler
diff --git a/Engine/Shaders/Private/InstanceCulling/BuildInstanceDrawCommands.usf b/Engine/Shaders/Private/InstanceCulling/BuildInstanceDrawCommands.usf
--- a/Engine/Shaders/Private/InstanceCulling/BuildInstanceDrawCommands.usf
+++ b/Engine/Shaders/Private/InstanceCulling/BuildInstanceDrawCommands.usf
@@ -1,7 +1,7 @@
 // Copyright Epic Games, Inc. All Rights Reserved.
 
 // Just used to trigger shader recompiles, guid should be regenerated on merge conflicts
-#pragma message("UESHADERMETADATA_VERSION DC8F4920-57F0-41EF-ACC0-B6643D1D80C7")
+#pragma message("UESHADERMETADATA_VERSION 5F4355AE-A326-40DB-99A4-0666E9076245")
 
 #ifndef ENABLE_PER_INSTANCE_CUSTOM_DATA
 	#define ENABLE_PER_INSTANCE_CUSTOM_DATA 1
@@ -25,6 +25,9 @@
 // cross-compilation creates bugs with ENABLE_INSTANCE_COMPACTION.
 #define ALLOW_TEMPLATES 0
 
+// Do not use shared samplers as it requires the View uniform buffer, which is not bound for this shader.
+#define USE_HZB_SHARED_SAMPLERS 0
+
 #include "../Common.ush"
 #include "../SceneData.ush"
 #include "../LightmapData.ush"
