Rename INVALID into NONE

diff --git a/Engine/Shaders/Private/DBufferDecalShared.ush b/Engine/Shaders/Private/DBufferDecalShared.ush
--- a/Engine/Shaders/Private/DBufferDecalShared.ush
+++ b/Engine/Shaders/Private/DBufferDecalShared.ush
@@ -233,7 +233,7 @@ void ApplyDBufferData(in FSubstrateDBuffer SubstrateDBufferData, inout FSubstrat
 						BSDF.TmpMFP = BSDF.TmpMFP * SubstrateDBufferData.OneMinusCoverage_BaseColor;
 						if (SubstrateDBufferData.OneMinusCoverage_BaseColor == 0.0f)
 						{
-							BSDF_SETSSSTYPE(BSDF, SSS_TYPE_INVALID);
+							BSDF_SETSSSTYPE(BSDF, SSS_TYPE_NONE);
 						}
 					}
 				#endif
diff --git a/Engine/Shaders/Private/DebugProbes.usf b/Engine/Shaders/Private/DebugProbes.usf
--- a/Engine/Shaders/Private/DebugProbes.usf
+++ b/Engine/Shaders/Private/DebugProbes.usf
@@ -129,7 +129,7 @@ void WriteOutMaterial(
 		/*float3 SSSMFP							*/0.0f,
 		/*float SSSMFPScale						*/0.0f,
 		/*float SSSPhaseAniso					*/0.0f,
-		/*float SSSType							*/SSS_TYPE_INVALID,
+		/*float SSSType							*/SSS_TYPE_NONE,
 		/*float3 Emissive						*/0.0f,
 		/*float SecondRoughness					*/0.0f,
 		/*float SecondRoughnessWeight			*/0.0f,
diff --git a/Engine/Shaders/Private/ShadowProjectionPixelShader.usf b/Engine/Shaders/Private/ShadowProjectionPixelShader.usf
--- a/Engine/Shaders/Private/ShadowProjectionPixelShader.usf
+++ b/Engine/Shaders/Private/ShadowProjectionPixelShader.usf
@@ -305,7 +305,7 @@ void Main(
 			if (!bIsEyeOrHair)
 			{
 			#if SUBSTRATE_ENABLED
-				if (SSSType != SSS_TYPE_INVALID)
+				if (SSSType != SSS_TYPE_NONE)
 				{
 					// Disble accurate SSS transmission if SSS type is no using Diffusion or Diffusion Profile
 					bHasSubsurfaceTransmission = bHasSubsurfaceTransmission && (SSSType >= SSS_TYPE_DIFFUSION);
diff --git a/Engine/Shaders/Private/Substrate/Substrate.ush b/Engine/Shaders/Private/Substrate/Substrate.ush
--- a/Engine/Shaders/Private/Substrate/Substrate.ush
+++ b/Engine/Shaders/Private/Substrate/Substrate.ush
@@ -1490,7 +1490,7 @@ FSubstrateData InternalGetSubstrateSlabBSDF(
 	}
 	else if (!bMFPPlugged)
 	{
-		SSSType = SSS_TYPE_INVALID;
+		SSSType = SSS_TYPE_NONE;
 	}
 
 	FHaziness Haziness = InitialiseHaziness();
@@ -1786,7 +1786,7 @@ FSubstrateData GetSubstrateEyeBSDF(float3 DiffuseAlbedo, float Roughness, float
 	EYE_F90							(SubstrateData.InlinedBSDF) = 1.0f;
 
 	const bool bHasSSS = SSSProfileId != 0;
-	BSDF_SETSSSTYPE(SubstrateData.InlinedBSDF, bHasSSS ? SSS_TYPE_DIFFUSION_PROFILE : SSS_TYPE_INVALID);
+	BSDF_SETSSSTYPE(SubstrateData.InlinedBSDF, bHasSSS ? SSS_TYPE_DIFFUSION_PROFILE : SSS_TYPE_NONE);
 #if SUBSTRATE_INLINE_SHADING
 	SubstrateData.InlinedBSDF.Coverage = 1.0f;
 #endif
@@ -2406,7 +2406,7 @@ void FSubstratePixelHeader::SubstrateConvertFromDBuffer(
 		0.f,				// SSSMFP
 		0.f,				// SSSMFPScale
 		0.f,				// SSSPhaseAniso
-		SSS_TYPE_INVALID,	// UseSSSDiffusion
+		SSS_TYPE_NONE,	// UseSSSDiffusion
 		0.f,				// Emissive
 		0.f,				// SecondRoughness
 		0.f,				// SecondRoughnessWeight
@@ -2704,7 +2704,7 @@ FSubstrateData SubstrateHorizontalMixingParameterBlending(FSubstrateData Backgro
 		if (SLAB_SSSPROFILERADIUSSCALE(ResultBSDF) == 0.0f)
 		{
 			SLAB_SSSPROFILEID(ResultBSDF) = SSS_PROFILE_ID_INVALID;
-			BSDF_SETSSSTYPE(ResultBSDF, SSS_TYPE_INVALID);
+			BSDF_SETSSSTYPE(ResultBSDF, SSS_TYPE_NONE);
 		}
 	}
 	else
@@ -2715,7 +2715,7 @@ FSubstrateData SubstrateHorizontalMixingParameterBlending(FSubstrateData Backgro
 		// We still try to disable SSS if possible: for performance and also to preserve the surface color (especially specular reflection color).
 		if (all(SLAB_SSSMFP(ResultBSDF) == 0.0f))
 		{
-			BSDF_SETSSSTYPE(ResultBSDF, SSS_TYPE_INVALID);
+			BSDF_SETSSSTYPE(ResultBSDF, SSS_TYPE_NONE);
 		}
 	}
 
@@ -3256,13 +3256,13 @@ void FSubstrateBSDF::PostProcessBSDFBeforeLighting(bool bIsSubstrateOpaqueMateri
 bool FSubstrateBSDF::HasScatteringData()
 {
 	const uint SSSType = BSDF_GETSSSTYPE(this);
-	return SSSType != SSS_TYPE_INVALID && SSSType != SSS_TYPE_SIMPLEVOLUME;
+	return SSSType != SSS_TYPE_NONE && SSSType != SSS_TYPE_SIMPLEVOLUME;
 }
 
 bool FSubstrateBSDF::HasBackScattering()
 {
 	const uint SSSType = BSDF_GETSSSTYPE(this);
-	return SSSType != SSS_TYPE_INVALID && SSSType != SSS_TYPE_SIMPLEVOLUME;
+	return SSSType != SSS_TYPE_NONE && SSSType != SSS_TYPE_SIMPLEVOLUME;
 }
 
 void FSubstrateBSDF::SubstrateBakeSpecularProfileToReflectivity(in float NoV)
@@ -3481,7 +3481,7 @@ float3 SubstrateGetBSDFSubSurfaceColor(in FSubstrateBSDF BSDF)
 	{
 	case SUBSTRATE_BSDF_TYPE_SLAB:
 	{
-		if (BSDF_GETSSSTYPE(BSDF) != SSS_TYPE_INVALID)
+		if (BSDF_GETSSSTYPE(BSDF) != SSS_TYPE_NONE)
 		{
 			const float3 Extinction = 1.0f / max(SUBSTRATE_EPSILON, SLAB_SSSMFP(BSDF) * CENTIMETER_TO_METER);
 			return ExtinctionToTransmittance(Extinction, SUBSTRATE_SIMPLEVOLUME_THICKNESS_CM * CENTIMETER_TO_METER);
@@ -4099,7 +4099,7 @@ FSubstrateBSDF UnpackFastPathSubstrateBSDFIn(FSubstrateMaterialContainer Substra
 	BSDF_SETHASF90(OutBSDF, 0);
 	BSDF_SETHASANISOTROPY(OutBSDF, 0);
 	BSDF_SETISTOPLAYER(OutBSDF, 1);
-	BSDF_SETSSSTYPE(OutBSDF, SSS_TYPE_INVALID);
+	BSDF_SETSSSTYPE(OutBSDF, SSS_TYPE_NONE);
 	BSDF_SETHASHAZINESS(OutBSDF, 0);
 	BSDF_SETISTHIN(OutBSDF, 0);
 	BSDF_SETHASFUZZ(OutBSDF, 0);
@@ -4145,7 +4145,7 @@ FSubstrateBSDF UnpackFastWaterPathSubstrateBSDFIn(FSubstrateMaterialContainer Su
 	BSDF_SETHASF90(OutBSDF, 0);
 	BSDF_SETHASANISOTROPY(OutBSDF, 0);
 	BSDF_SETISTOPLAYER(OutBSDF, 1);
-	BSDF_SETSSSTYPE(OutBSDF, SSS_TYPE_INVALID);
+	BSDF_SETSSSTYPE(OutBSDF, SSS_TYPE_NONE);
 	BSDF_SETHASHAZINESS(OutBSDF, 0);
 	BSDF_SETISTHIN(OutBSDF, 0);
 	BSDF_SETHASFUZZ(OutBSDF, 0);
@@ -4231,7 +4231,7 @@ FSubstrateBSDF UnpackSubstrateBSDFIn(FSubstrateMaterialContainer SubstrateBuffer
 				BSDF_SETHASF90(OutBSDF, 0);
 				BSDF_SETHASANISOTROPY(OutBSDF, 0);
 				BSDF_SETISTOPLAYER(OutBSDF, 1);
-				BSDF_SETSSSTYPE(OutBSDF, SSS_TYPE_INVALID);
+				BSDF_SETSSSTYPE(OutBSDF, SSS_TYPE_NONE);
 				BSDF_SETHASHAZINESS(OutBSDF, 1);
 				BSDF_SETISTHIN(OutBSDF, 0);
 				BSDF_SETHASFUZZ(OutBSDF, 0);
@@ -4264,7 +4264,7 @@ FSubstrateBSDF UnpackSubstrateBSDFIn(FSubstrateMaterialContainer SubstrateBuffer
 				BSDF_SETHASF90(OutBSDF, 0);
 				BSDF_SETHASANISOTROPY(OutBSDF, 0);
 				BSDF_SETISTOPLAYER(OutBSDF, 1);
-				BSDF_SETSSSTYPE(OutBSDF, SSS_TYPE_INVALID);
+				BSDF_SETSSSTYPE(OutBSDF, SSS_TYPE_NONE);
 				BSDF_SETHASHAZINESS(OutBSDF, 0);
 				BSDF_SETISTHIN(OutBSDF, 0);
 				BSDF_SETHASFUZZ(OutBSDF, 1);
@@ -4385,7 +4385,7 @@ FSubstrateBSDF UnpackSubstrateBSDFIn(FSubstrateMaterialContainer SubstrateBuffer
 		BSDF_SETSHAREDLOCALBASISID(OutBSDF, 0);
 		BSDF_SETHASANISOTROPY(OutBSDF, 0);
 		BSDF_SETISTOPLAYER(OutBSDF, 1);
-		BSDF_SETSSSTYPE(OutBSDF, SubstrateHeader.HasSubsurface() ? SSS_TYPE_DIFFUSION_PROFILE : SSS_TYPE_INVALID);
+		BSDF_SETSSSTYPE(OutBSDF, SubstrateHeader.HasSubsurface() ? SSS_TYPE_DIFFUSION_PROFILE : SSS_TYPE_NONE);
 	}
 	else if (SubstrateHeader.IsHair())
 	{
@@ -4445,7 +4445,7 @@ FSubstrateBSDF UnpackSubstrateBSDFIn(FSubstrateMaterialContainer SubstrateBuffer
 			SLAB_HAZINESS(OutBSDF) = RawData >> 16; // Keeping the top 16 bits
 		}
 
-		const bool bHasSSS = BSDF_GETSSSTYPE(OutBSDF) != SSS_TYPE_INVALID;
+		const bool bHasSSS = BSDF_GETSSSTYPE(OutBSDF) != SSS_TYPE_NONE;
 		if (bHasSSS)
 		{
 			if (BSDF_GETSSSTYPE(OutBSDF) == SSS_TYPE_DIFFUSION_PROFILE)
diff --git a/Engine/Shaders/Private/Substrate/SubstrateDBuffer.usf b/Engine/Shaders/Private/Substrate/SubstrateDBuffer.usf
--- a/Engine/Shaders/Private/Substrate/SubstrateDBuffer.usf
+++ b/Engine/Shaders/Private/Substrate/SubstrateDBuffer.usf
@@ -219,14 +219,14 @@ void MainCS(
 					++RWOffset;
 				}
 
-				const bool bHasSSS = BSDF_GETSSSTYPE(BSDF) != SSS_TYPE_INVALID;
+				const bool bHasSSS = BSDF_GETSSSTYPE(BSDF) != SSS_TYPE_NONE;
 				if (bHasSSS)
 				{
 					// Disable SSS when the coverage reach one to preserve the decal color
 					SLAB_SSSMFP(BSDF) *= DBufferData.OneMinusCoverage_BaseColor;
 					if (DBufferData.OneMinusCoverage_BaseColor == 0.0f)
 					{
-						BSDF_SETSSSTYPE(BSDF, SSS_TYPE_INVALID);
+						BSDF_SETSSSTYPE(BSDF, SSS_TYPE_NONE);
 					}
 
 					if (BSDF_GETSSSTYPE(BSDF) == SSS_TYPE_DIFFUSION_PROFILE)
@@ -330,7 +330,7 @@ void MainCS(
 				float ProfileScl= SubstrateSubSurfaceHeaderGetProfileRadiusScale(SSSDataHeader) * DBufferData.OneMinusCoverage_BaseColor;
 				uint  ProfileId = SubstrateSubSurfaceHeaderGetProfileId(SSSDataHeader);
 				SubstrateSubSurfaceHeaderSetProfile(SSSDataHeader, ProfileScl, DBufferData.OneMinusCoverage_BaseColor == 0.0f ? 0 : ProfileId);
-				// Note We cannot change the type using SubstrateSubSurfaceHeaderSetSSSType(SSSDataHeader, SSS_TYPE_INVALID);
+				// Note We cannot change the type using SubstrateSubSurfaceHeaderSetSSSType(SSSDataHeader, SSS_TYPE_NONE);
 				// Because we are stuck using the SINGLE_OPTLEGACYMODE_SSSPROFILE mode.
 				// So some SSSProfile material with decal will look difference when using deferred DBuffer instead of in shader application.
 
diff --git a/Engine/Shaders/Private/Substrate/SubstrateExport.ush b/Engine/Shaders/Private/Substrate/SubstrateExport.ush
--- a/Engine/Shaders/Private/Substrate/SubstrateExport.ush
+++ b/Engine/Shaders/Private/Substrate/SubstrateExport.ush
@@ -140,7 +140,7 @@ void PackSubstrateOut(
 	OpaqueRoughRefractionData = (FSubstrateOpaqueRoughRefractionData)0;
 
 	// The SSSData we accumulate to allow blending of Wrap and SubstrateDiffusion types.
-	uint SSSDataType = SSS_TYPE_INVALID;
+	uint SSSDataType = SSS_TYPE_NONE;
 	float SSSDataAniso = 0.0f;
 	float SSSDataProfileID = 0.0f;
 	float SSSDataRadiusScale = 0.0f;
@@ -231,7 +231,7 @@ void PackSubstrateOut(
 
 						TopLayerData.Material = TOP_LAYER_MATERIAL_VALID;
 
-						const bool bHasSSS = BSDF_GETSSSTYPE(BSDF) != SSS_TYPE_INVALID;
+						const bool bHasSSS = BSDF_GETSSSTYPE(BSDF) != SSS_TYPE_NONE;
 						const bool bIsSimpleVolume = BSDF_GETSSSTYPE(BSDF) == SSS_TYPE_SIMPLEVOLUME;
 						bSubstrateSubsurfaceEnable = bSubstrateSubsurfaceEnable || (bHasSSS && BSDF.Coverage > 0.0f && !bIsSimpleVolume); // Should it be coverage>0 or any(LuminanceWeight>0) ?
 
@@ -370,7 +370,7 @@ void PackSubstrateOut(
 				{
 					OptimisedLegacyMode = SINGLE_OPTLEGACYMODE_CLEARCOAT;
 				}
-				else if ((OptimisedModeMask & SingleHasCloth) == SingleHasCloth && BSDF_GETSSSTYPE(BSDF) == SSS_TYPE_INVALID)
+				else if ((OptimisedModeMask & SingleHasCloth) == SingleHasCloth && BSDF_GETSSSTYPE(BSDF) == SSS_TYPE_NONE)
 				{
 					OptimisedLegacyMode = SINGLE_OPTLEGACYMODE_CLOTH;
 				}				
@@ -768,7 +768,7 @@ void PackSubstrateOut(
 								SUBSTRATE_STORE_UINT1(HazinessData << 16 | F90Data);
 							}
 
-							if (BSDF_GETSSSTYPE(BSDF) != SSS_TYPE_INVALID)
+							if (BSDF_GETSSSTYPE(BSDF) != SSS_TYPE_NONE)
 							{
 								// All SSS types store side data (SSSHeader/SSSData), apart from SIMPLEVOLUME
 								if (BSDF_GETSSSTYPE(BSDF) != SSS_TYPE_SIMPLEVOLUME)
@@ -898,7 +898,7 @@ void PackSubstrateOut(
 		FinalizeWrites(SubstrateBuffer, SubstrateAddressing);
 
 		// Now generate the final SSSData structure from the parameter blended attributes
-		if (SSSDataType != SSS_TYPE_INVALID)
+		if (SSSDataType != SSS_TYPE_NONE)
 		{
 			// All SSS types store side data (SSSHeader/SSSData), apart from SIMPLEVOLUME
 			if (SSSDataType != SSS_TYPE_SIMPLEVOLUME)
diff --git a/Engine/Shaders/Private/Substrate/SubstrateLegacyConversion.ush b/Engine/Shaders/Private/Substrate/SubstrateLegacyConversion.ush
--- a/Engine/Shaders/Private/Substrate/SubstrateLegacyConversion.ush
+++ b/Engine/Shaders/Private/Substrate/SubstrateLegacyConversion.ush
@@ -75,7 +75,7 @@ FSubstrateLegacyParameters InitSubstrateLegacyParameters(FSubstratePixelFootprin
 	Out.SSSMFP = Zeros;
 	Out.SSSMFPScale = 1.f;
 	Out.SSSPhaseAniso = 0.f;
-	Out.SSSType = SSS_TYPE_INVALID;
+	Out.SSSType = SSS_TYPE_NONE;
 	Out.Emissive = Zeros;
 	Out.SecondRoughness = 0.0f;
 	Out.SecondRoughnessWeight = 0.0f;
diff --git a/Engine/Shaders/Private/Substrate/SubstrateSubsurface.ush b/Engine/Shaders/Private/Substrate/SubstrateSubsurface.ush
--- a/Engine/Shaders/Private/Substrate/SubstrateSubsurface.ush
+++ b/Engine/Shaders/Private/Substrate/SubstrateSubsurface.ush
@@ -81,7 +81,7 @@ void SubstrateSubSurfaceHeaderSetWrap(inout FSubstrateSubsurfaceHeader SSSHeader
 
 bool SubstrateSubSurfaceHeaderGetIsValid(in FSubstrateSubsurfaceHeader SSSHeader)
 {
-	return SSSHEADER_TYPE(SSSHeader) != SSS_TYPE_INVALID;
+	return SSSHEADER_TYPE(SSSHeader) != SSS_TYPE_NONE;
 }
 
 bool SubstrateSubSurfaceHeaderHasExtras(in FSubstrateSubsurfaceHeader SSSHeader)
diff --git a/Engine/Shaders/Private/Substrate/SubstrateVisualizeCommon.ush b/Engine/Shaders/Private/Substrate/SubstrateVisualizeCommon.ush
--- a/Engine/Shaders/Private/Substrate/SubstrateVisualizeCommon.ush
+++ b/Engine/Shaders/Private/Substrate/SubstrateVisualizeCommon.ush
@@ -96,7 +96,7 @@ void PrintSSSType(inout FShaderPrintContext Ctx, uint SSSType, bool bThin, FFont
 	else if (SSSType == SSS_TYPE_DIFFUSION_PROFILE)	Print(Ctx, TEXT("Diffusion Profile"), InColor);
 	else if (SSSType == SSS_TYPE_SIMPLEVOLUME)		Print(Ctx, TEXT("Simple Volume"), InColor);
 	else if (bThin)									Print(Ctx, TEXT("Thin"), InColor);
-	else /*	(SSSType == SSS_TYPE_INVALID) */		Print(Ctx, TEXT("Invalid"), InColor);
+	else /*	(SSSType == SSS_TYPE_NONE) */			Print(Ctx, TEXT("None"), InColor);
 }
 
 void PrintBackgroundRect(inout FShaderPrintContext Ctx, float2 StartRect, float2 EndRect)
diff --git a/Engine/Shaders/Private/VirtualShadowMaps/VirtualShadowMapProjection.usf b/Engine/Shaders/Private/VirtualShadowMaps/VirtualShadowMapProjection.usf
--- a/Engine/Shaders/Private/VirtualShadowMaps/VirtualShadowMapProjection.usf
+++ b/Engine/Shaders/Private/VirtualShadowMaps/VirtualShadowMapProjection.usf
@@ -74,7 +74,7 @@ FProjectionShadingInfo GetProjectionShadingInfo(uint2 PixelPos)
 	Out.bIsValid = SubstratePixelHeader.IsSubstrateMaterial();
 	Out.bIsHair = BSDFType == SUBSTRATE_BSDF_TYPE_HAIR;
 	Out.WorldNormal = TopLayerData.WorldNormal;
-	Out.bIsSubsurface = SSSType != SSS_TYPE_INVALID; // legacy IsSubsurfaceModel is true for subsurface profile, eyes, etc.
+	Out.bIsSubsurface = SSSType != SSS_TYPE_NONE; // legacy IsSubsurfaceModel is true for subsurface profile, eyes, etc.
 
 	// VSM traces using the top layer normal when NoL>0 only. So if normals for slabs below do not match, light leaking will occur.
 	// For instance: a flat translucent top layer above a normal mapped bottom layer.
diff --git a/Engine/Shaders/Public/SubstratePublic.ush b/Engine/Shaders/Public/SubstratePublic.ush
--- a/Engine/Shaders/Public/SubstratePublic.ush
+++ b/Engine/Shaders/Public/SubstratePublic.ush
@@ -78,7 +78,7 @@ FSubstrateGBuffer SubstratePublic_GetSubstrateGBufferFromFirstBSDF(uint2 PixelCo
 		SGBuffer.bIsFirstPerson		= SubstratePixelHeader.IsFirstPerson();
 
 		SGBuffer.Opacity			= 0.0f;
-		if (SubstrateSubSurfaceHeaderGetSSSType(SSSHeader) != SSS_TYPE_INVALID)
+		if (SubstrateSubSurfaceHeaderGetSSSType(SSSHeader) != SSS_TYPE_NONE)
 		{
 			SGBuffer.Opacity = SubstrateSubSurfaceHeaderGetOpacity(SSSHeader);
 		}
diff --git a/Engine/Shaders/Shared/SubstrateDefinitions.h b/Engine/Shaders/Shared/SubstrateDefinitions.h
--- a/Engine/Shaders/Shared/SubstrateDefinitions.h
+++ b/Engine/Shaders/Shared/SubstrateDefinitions.h
@@ -126,7 +126,7 @@
 // 5 more remains: Also how to deal with Aniso?
 
 // Sub-surface scattering types
-#define SSS_TYPE_INVALID						(0u)
+#define SSS_TYPE_NONE							(0u)
 #define SSS_TYPE_WRAP							(1u)
 #define SSS_TYPE_TWO_SIDED_WRAP					(2u)
 #define SSS_TYPE_DIFFUSION						(3u)
diff --git a/Engine/Shaders/Shared/SubstrateVisualizeDefinitions.h b/Engine/Shaders/Shared/SubstrateVisualizeDefinitions.h
--- a/Engine/Shaders/Shared/SubstrateVisualizeDefinitions.h
+++ b/Engine/Shaders/Shared/SubstrateVisualizeDefinitions.h
@@ -31,7 +31,7 @@ struct FSubsterateDebugClosure
 	//////////
 	// Slabs
 	int bIsTopLayer; // BSDF_GETISTOPLAYER(BSDF)
-	int SSSType; // BSDF_GETSSSTYPE(BSDF) != SSS_TYPE_INVALID
+	int SSSType; // BSDF_GETSSSTYPE(BSDF) != SSS_TYPE_NONE
 	int bIsThin; // BSDF_GETISTHIN(BSDF)
 	float LuminanceWeightR, LuminanceWeightG, LuminanceWeightB; // bHasGreyWeightV BSDF.LuminanceWeightV.x or rgb
 	float TransmittanceAboveAlongNR, TransmittanceAboveAlongNG, TransmittanceAboveAlongNB; // bHasGreyTopTrans BSDF.TransmittanceAboveAlongN.x or rgb
@@ -376,7 +376,7 @@ FSubstratePixelDebugData ConvertToSerializableSubstratePixelDebugData(uint2 InCo
 				Data.Closures[i].HazeSimpleClearCoatMode = Haziness.bSimpleClearCoat ? 1u : 0u;
 			}
 
-			if (Data.Closures[i].SSSType != SSS_TYPE_INVALID || Data.Closures[i].bIsThin)
+			if (Data.Closures[i].SSSType != SSS_TYPE_NONE || Data.Closures[i].bIsThin)
 			{
 				if (Data.Closures[i].SSSType == SSS_TYPE_WRAP || Data.Closures[i].SSSType == SSS_TYPE_TWO_SIDED_WRAP)
 				{
diff --git a/Engine/Source/Runtime/Engine/Public/Materials/MaterialExpressionSubstrate.h b/Engine/Source/Runtime/Engine/Public/Materials/MaterialExpressionSubstrate.h
--- a/Engine/Source/Runtime/Engine/Public/Materials/MaterialExpressionSubstrate.h
+++ b/Engine/Source/Runtime/Engine/Public/Materials/MaterialExpressionSubstrate.h
@@ -15,7 +15,7 @@ struct FSubstrateOperator;
 UENUM()
 enum EMaterialSubSurfaceType : int
 {
-	MSS_None				= SSS_TYPE_INVALID 				UMETA(DisplayName="None"),
+	MSS_None				= SSS_TYPE_NONE 				UMETA(DisplayName="None"),
 	MSS_Wrap				= SSS_TYPE_WRAP 				UMETA(DisplayName="Wrap", ToolTip="Approximation using wrap-lighting. "),
 	MSS_TwoSidedWrap		= SSS_TYPE_TWO_SIDED_WRAP		UMETA(DisplayName="Two-Sided Wrap", ToolTip="Approximation using wrap-lighting and handling thin surface (e.g.: foliage)"),
 	MSS_Diffusion			= SSS_TYPE_DIFFUSION 			UMETA(DisplayName="Diffusion", ToolTip="Diffusion based sub-surface scattering"),
diff --git a/Engine/Source/Runtime/Renderer/Private/Substrate/SubstrateVisualize.cpp b/Engine/Source/Runtime/Renderer/Private/Substrate/SubstrateVisualize.cpp
--- a/Engine/Source/Runtime/Renderer/Private/Substrate/SubstrateVisualize.cpp
+++ b/Engine/Source/Runtime/Renderer/Private/Substrate/SubstrateVisualize.cpp
@@ -605,7 +605,7 @@ void AddProcessAndPrintSubstrateMaterialPropertiesPasses(FRDGBuilder& GraphBuild
 							NewLine();
 						}
 
-						if (Closure.SSSType != SSS_TYPE_INVALID || Closure.bIsThin)
+						if (Closure.SSSType != SSS_TYPE_NONE || Closure.bIsThin)
 						{
 							if (Closure.SSSType == SSS_TYPE_WRAP || Closure.SSSType == SSS_TYPE_TWO_SIDED_WRAP)
 							{
