VelocityPass shader optimization. Decreased used VGPRs

diff --git a/Engine/Shaders/Private/VelocityShader.usf b/Engine/Shaders/Private/VelocityShader.usf
--- a/Engine/Shaders/Private/VelocityShader.usf
+++ b/Engine/Shaders/Private/VelocityShader.usf
@@ -19,8 +19,6 @@ static const float GDepthBias = 0.001f;
 struct FVelocityInterpsVSToPS
 {
 #if !TRANSLUCENCY_VELOCITY_FROM_DEPTH
-	// float4(ScreenPos.xyzw)
-	float4 PackedVelocityA : TEXCOORD6;
 	// float4(PrevScreenPos.xyzw)
 	float4 PackedVelocityC : TEXCOORD7;
 #endif
@@ -76,27 +74,22 @@ void MainVertexShader(
 	PrevTranslatedWorldPosition.xyz += GetMaterialPreviousWorldPositionOffset(VertexParameters);
 	
 	{
-		float4 ScreenPos;
 		// compute the old screen pos with the old world position and the old camera matrix
 		float4 PrevScreenPosObj = mul(float4(PrevTranslatedWorldPosition.xyz, 1), ResolvedView.PrevTranslatedWorldToClip);
-
-		{
-			ScreenPos = mul(float4(TranslatedWorldPosition.xyz, 1), ResolvedView.TranslatedWorldToClip);
-			Output.Position = ScreenPos;
-		}
+		Output.Position = INVARIANT(mul(float4(TranslatedWorldPosition.xyz, 1), ResolvedView.TranslatedWorldToClip));
 
 #if USE_GLOBAL_CLIP_PLANE
 		OutGlobalClipPlaneDistance = dot(ResolvedView.GlobalClippingPlane, float4(TranslatedWorldPosition.xyz, 1));
 #endif
 
 #if !TRANSLUCENCY_VELOCITY_FROM_DEPTH
-		Output.Interps.PackedVelocityA = INVARIANT(ScreenPos);
 		Output.Interps.PackedVelocityC = INVARIANT(PrevScreenPosObj);
 #endif
 	}
 
 // Move all geometry a little bit towards the camera (to not get z fighting with existing zbuffer of the same meshes rendered with slightly different float computations).
     Output.Position.z += View.NumSceneColorMSAASamples > 1 ? GDepthBias * Output.Position.w : 0.0;
+    Output.Interps.PackedVelocityC.z += View.NumSceneColorMSAASamples > 1 ? GDepthBias * Output.Position.w : 0.0;
 }
 #endif // VERTEXSHADER
 
@@ -127,8 +120,8 @@ void MainPixelShader(
 	FMaterialPixelParameters MaterialParameters = GetMaterialPixelParameters(Input.FactoryInterpolants, SvPosition);
 	FPixelMaterialInputs PixelMaterialInputs;
 
-#if USE_WORLD_POSITION_EXCLUDING_SHADER_OFFSETS && !IS_NANITE_PASS
 	float4 ScreenPosition = SvPositionToResolvedScreenPosition(SvPosition);
+#if USE_WORLD_POSITION_EXCLUDING_SHADER_OFFSETS && !IS_NANITE_PASS
 	float3 TranslatedWorldPosition = SvPositionToResolvedTranslatedWorld(SvPosition);
 	CalcMaterialParametersEx(MaterialParameters, PixelMaterialInputs, SvPosition, ScreenPosition, bIsFrontFace, TranslatedWorldPosition, Input.PixelPositionExcludingWPO);
 #else
@@ -178,7 +171,7 @@ void MainPixelShader(
 
 #if !TRANSLUCENCY_VELOCITY_FROM_DEPTH
 	// 3d velocity, includes camera an object motion
-	float3 Velocity = Calculate3DVelocity(Input.PackedVelocityA, Input.PackedVelocityC);
+	float3 Velocity = Calculate3DVelocity(ScreenPosition, Input.PackedVelocityC);
 
 	OutColor = EncodeVelocityToTexture(Velocity, (GetPrimitiveData(MaterialParameters).Flags & PRIMITIVE_SCENE_DATA_FLAG_HAS_PIXEL_ANIMATION) != 0);
 #else
