Fix compilation issue in VelocityShader due to missing preprocessor condition

diff --git a/Engine/Shaders/Private/VelocityShader.usf b/Engine/Shaders/Private/VelocityShader.usf
--- a/Engine/Shaders/Private/VelocityShader.usf
+++ b/Engine/Shaders/Private/VelocityShader.usf
@@ -87,9 +87,11 @@ void MainVertexShader(
 #endif
 	}
 
-// Move all geometry a little bit towards the camera (to not get z fighting with existing zbuffer of the same meshes rendered with slightly different float computations).
+	// Move all geometry a little bit towards the camera (to not get z fighting with existing zbuffer of the same meshes rendered with slightly different float computations).
     Output.Position.z += View.NumSceneColorMSAASamples > 1 ? GDepthBias * Output.Position.w : 0.0;
-    Output.Interps.PackedVelocityC.z += View.NumSceneColorMSAASamples > 1 ? GDepthBias * Output.Position.w : 0.0;
+#if !TRANSLUCENCY_VELOCITY_FROM_DEPTH
+	Output.Interps.PackedVelocityC.z += View.NumSceneColorMSAASamples > 1 ? GDepthBias * Output.Position.w : 0.0;
+#endif
 }
 #endif // VERTEXSHADER
 
