Fix shader binding issue.

diff --git a/Engine/Shaders/Private/VirtualShadowMaps/VirtualShadowMapBuildPerPageDrawCommands.usf b/Engine/Shaders/Private/VirtualShadowMaps/VirtualShadowMapBuildPerPageDrawCommands.usf
--- a/Engine/Shaders/Private/VirtualShadowMaps/VirtualShadowMapBuildPerPageDrawCommands.usf
+++ b/Engine/Shaders/Private/VirtualShadowMaps/VirtualShadowMapBuildPerPageDrawCommands.usf
@@ -4,6 +4,9 @@
 	VirtualShadowMapBuildPerPageDrawCommands.usf: 
 =============================================================================*/
 
+// Do not use shared samplers as it requires the View uniform buffer, which is not bound for this shader.
+#define USE_HZB_SHARED_SAMPLERS 0
+
 #include "../Common.ush"
 #include "../ViewData.ush"
 #include "VirtualShadowMapPageOverlap.ush"
