Path Tracer: Small improvement to payload setup for cloud material callable shaders

Reduce the active payload size to 4 uints for both input and output.

#rnx

diff --git a/Engine/Shaders/Private/PathTracing/PathTracingCommon.ush b/Engine/Shaders/Private/PathTracing/PathTracingCommon.ush
--- a/Engine/Shaders/Private/PathTracing/PathTracingCommon.ush
+++ b/Engine/Shaders/Private/PathTracing/PathTracingCommon.ush
@@ -399,34 +399,28 @@ struct FPackedPathTracingPayload : FMinimalPayload
 	}
 
 	// Volumetric callable shaders
-	void SetVolumetricCallableShaderInput(FDFVector3 AbsolutePosition, float CloudHeightKm)
+	void SetVolumetricCallableShaderInput(float3 TranslatedWorldPos, float CloudHeightKm)
 	{
-		PackedData[0] = asuint(AbsolutePosition.High.x);
-		PackedData[1] = asuint(AbsolutePosition.High.y);
-		PackedData[2] = asuint(AbsolutePosition.High.z);
-		PackedData[3] = asuint(AbsolutePosition.Low.x);
-		PackedData[4] = asuint(AbsolutePosition.Low.y);
-		PackedData[5] = asuint(AbsolutePosition.Low.z);
-		PackedData[6] = asuint(CloudHeightKm);
+		PackedData[0] = asuint(TranslatedWorldPos.x);
+		PackedData[1] = asuint(TranslatedWorldPos.y);
+		PackedData[2] = asuint(TranslatedWorldPos.z);
+		PackedData[3] = asuint(CloudHeightKm);
 	}
-	FDFVector3 GetVolumetricCallableShaderInputAbsolutePosition()
+	float3 GetVolumetricCallableShaderInputTranslatedWorldPos()
 	{
-		return MakeDFVector3(
-			float3(asfloat(PackedData[0]), asfloat(PackedData[1]), asfloat(PackedData[2])),
-			float3(asfloat(PackedData[3]), asfloat(PackedData[4]), asfloat(PackedData[5]))
-		);
+		return float3(asfloat(PackedData[0]), asfloat(PackedData[1]), asfloat(PackedData[2]));
 	}
-	float GetVolumetricCallableShaderInputCloudHeightKm() { return asfloat(PackedData[6]); }
+	float GetVolumetricCallableShaderInputCloudHeightKm() { return asfloat(PackedData[3]); }
 
 	void SetVolumetricCallableShaderOutput(uint4 Data)
 	{
-		PackedData[ 7] = Data.x;
-		PackedData[ 8] = Data.y;
-		PackedData[ 9] = Data.z;
-		PackedData[10] = Data.w;
+		PackedData[0] = Data.x;
+		PackedData[1] = Data.y;
+		PackedData[2] = Data.z;
+		PackedData[3] = Data.w;
 	}
 
-	uint4 GetVolumetricCallableShaderOutput() { return uint4(PackedData[7], PackedData[8], PackedData[9], PackedData[10]); }
+	uint4 GetVolumetricCallableShaderOutput() { return uint4(PackedData[0], PackedData[1], PackedData[2], PackedData[3]); }
 
 #endif // IS_PAYLOAD_ENABLED(RT_PAYLOAD_TYPE_GPULIGHTMASS)
 };
diff --git a/Engine/Shaders/Private/PathTracing/PathTracingVolumetricCloudMaterialShader.usf b/Engine/Shaders/Private/PathTracing/PathTracingVolumetricCloudMaterialShader.usf
--- a/Engine/Shaders/Private/PathTracing/PathTracingVolumetricCloudMaterialShader.usf
+++ b/Engine/Shaders/Private/PathTracing/PathTracingVolumetricCloudMaterialShader.usf
@@ -18,7 +18,8 @@ RAY_TRACING_ENTRY_CALLABLE(PathTracingVolumetricCloudMaterialShader, FPackedPath
 #if CLOUD_LAYER_PIXEL_SHADER
 	ResolvedView = ResolveView();
 
-	FDFVector3 AbsoluteWorldPosition = PackedPayload.GetVolumetricCallableShaderInputAbsolutePosition();
+	float3 TranslatedWorldPos = PackedPayload.GetVolumetricCallableShaderInputTranslatedWorldPos();
+	FDFVector3 AbsoluteWorldPosition = DFFastSubtract(TranslatedWorldPos, ResolvedView.PreViewTranslation);
 	float H = PackedPayload.GetVolumetricCallableShaderInputCloudHeightKm();
 	float B = PathTracingCloudParameters.CloudLayerBotKm;
 	float T = PathTracingCloudParameters.CloudLayerTopKm;
diff --git a/Engine/Shaders/Private/PathTracing/Volume/PathTracingClouds.ush b/Engine/Shaders/Private/PathTracing/Volume/PathTracingClouds.ush
--- a/Engine/Shaders/Private/PathTracing/Volume/PathTracingClouds.ush
+++ b/Engine/Shaders/Private/PathTracing/Volume/PathTracingClouds.ush
@@ -208,7 +208,7 @@ FVolumeShadedResult CloudsGetDensity(float3 TranslatedWorldPos)
 		// Reference path -- invoke the actual callable shader
 		FPackedPathTracingPayload Payload = (FPackedPathTracingPayload)0;
 		Payload.SetVolumetricCallableShaderInput(
-			DFFastSubtract(TranslatedWorldPos, PrimaryView.PreViewTranslation),
+			TranslatedWorldPos,
 			GetSphericalPlanetParameters(TranslatedWorldPos, float3(0, 0, 1)).HeightKm
 		);
 		CallShader(CloudCallableShaderId, Payload);
diff --git a/Engine/Shaders/Private/PathTracing/Volume/PathTracingCloudsCommon.ush b/Engine/Shaders/Private/PathTracing/Volume/PathTracingCloudsCommon.ush
--- a/Engine/Shaders/Private/PathTracing/Volume/PathTracingCloudsCommon.ush
+++ b/Engine/Shaders/Private/PathTracing/Volume/PathTracingCloudsCommon.ush
@@ -145,9 +145,9 @@ FSampleCloudMaterialResult DecodeCloudResult(uint4 Data)
 	Result.Extinction    = PayloadDecodeHDRColor(Data.x) * (1.0 / 64.0);
 	Result.Albedo        = PayloadDecodeLDRColor(Data.y);
 	Result.EmissiveColor = PayloadDecodeHDRColor(Data.z);
-	float3 PhaseData     = PayloadDecodeLDRColor(Data.w) * float3(2.0, 2.0, 1.0) - float3(1.0, 1.0, 0.0);
-	Result.PhaseG1       = PhaseData.x;
-	Result.PhaseG2       = PhaseData.y;
+	float3 PhaseData     = PayloadDecodeLDRColor(Data.w);
+	Result.PhaseG1       = PhaseData.x * 2.0 - 1.0;
+	Result.PhaseG2       = PhaseData.y * 2.0 - 1.0;
 	Result.PhaseBlend    = PhaseData.z;
 	return Result;
 }
@@ -158,6 +158,6 @@ uint4 EncodeCloudResult(FSampleCloudMaterialResult Result)
 		PayloadEncodeHDRColor(Result.Extinction * 64.0), // scale up so we don't loose any density when clipping bits and make better use of the higher range of the encoding
 		PayloadEncodeLDRColor(Result.Albedo),
 		PayloadEncodeHDRColor(Result.EmissiveColor),
-		PayloadEncodeLDRColor(float3(Result.PhaseG1, Result.PhaseG2, Result.PhaseBlend) * float3(0.5, 0.5, 1.0) + float3(0.5, 0.5, 0.0))
+		PayloadEncodeLDRColor(float3(Result.PhaseG1 * 0.5 + 0.5, Result.PhaseG2 * 0.5 + 0.5, Result.PhaseBlend))
 	);
 }
