Use simpler hash key calculation for feedback with a single element stride.

diff --git a/Engine/Shaders/Private/GPUFeedbackCompaction.usf b/Engine/Shaders/Private/GPUFeedbackCompaction.usf
--- a/Engine/Shaders/Private/GPUFeedbackCompaction.usf
+++ b/Engine/Shaders/Private/GPUFeedbackCompaction.usf
@@ -47,6 +47,19 @@ void BuildFeedbackHashTableIndirectArgsCS(
 	}
 }
 
+/** 
+ * Hash key generation functions for each FeedbackType
+ */
+uint GetHashKey(uint PackedFeedback)
+{
+	return MurmurMix(PackedFeedback);
+}
+
+uint GetHashKey(uint2 PackedFeedback)
+{
+	return MurmurAdd(MurmurMix(PackedFeedback.x), PackedFeedback.y);
+}
+
 /** 
  * Takes input Key, which is a hash and tries to add it to the linear probing hash table
  * Returns true if element was added for the first time
@@ -103,9 +116,9 @@ void BuildFeedbackHashTableCS(
 
 	if (ElementIndex < FeedbackBufferAllocator[0])
 	{
-		uint2 PackedFeedback = FeedbackBuffer[ElementIndex];
+		FeedbackType PackedFeedback = FeedbackBuffer[ElementIndex];
 
-		uint Key = MurmurAdd(MurmurMix(PackedFeedback.x), PackedFeedback.y);
+		uint Key = GetHashKey(PackedFeedback);
 
 		uint HashTableIndex = 0;
 		const bool bAdded = HashTableAdd(Key, HashTableIndex);
