Fix Substrate opaque rough refraction.

The thickness value was rescaled too soon, causing incorrect variance computation.

diff --git a/Engine/Shaders/Private/Substrate/Substrate.ush b/Engine/Shaders/Private/Substrate/Substrate.ush
--- a/Engine/Shaders/Private/Substrate/Substrate.ush
+++ b/Engine/Shaders/Private/Substrate/Substrate.ush
@@ -1555,7 +1555,8 @@ FSubstrateData InternalGetSubstrateSlabBSDF(
 	}
 
 	// Convert MFP be expressed relative to simple volume thickness
-	if (!SUBSTRATE_FASTPATH && (SSSType == SSS_TYPE_WRAP || SSSType == SSS_TYPE_TWO_SIDED_WRAP || SSSType == SSS_TYPE_SIMPLEVOLUME))
+	// Note: MFP rescaling for SSS_TYPE_SIMPLEVOLUME is done later, as the thickness needs to be preserved (i.e. not normalized) for rough refraction computation
+	if (!SUBSTRATE_FASTPATH && (SSSType == SSS_TYPE_WRAP || SSSType == SSS_TYPE_TWO_SIDED_WRAP))
 	{
 		SSSMFP = RescaleMFPToComputationSpace(SSSMFP, ThicknessCm, SUBSTRATE_SIMPLEVOLUME_THICKNESS_CM);
 		ThicknessCm = SUBSTRATE_SIMPLEVOLUME_THICKNESS_CM;
@@ -3288,7 +3289,14 @@ void FSubstrateBSDF::SubstrateSanitizeBSDF()
 // Convert to simple volume to get some transmittance
 void FSubstrateBSDF::PostProcessSlabBeforeLighting(bool bIsSubstrateOpaqueMaterial)
 {
-	if (!this.bIsBottom && BSDF_GETHASMFP(this))
+	// Force simple volume and MFP rescaling if the layer is not at the bottom or if sss type is already SimpleVolume. 
+	// In this case, it only rescales the MFP. The MFP rescaling can't be done before, as the thickness needs to be 
+	// preserved (i.e. not normalized) for rough refraction computation
+	// Side-effect:
+	// * Rescale MFP to computation space
+	// * Set thickness to normalized simple volume thickness
+	// * Set SSS type to SimpleVolume
+	if (BSDF_GETHASMFP(this) && (!this.bIsBottom || BSDF_GETSSSTYPE(this) == SSS_TYPE_SIMPLEVOLUME))
 	{
 		this.EnableSlabBSDFSimpleVolumetric();
 	}
