Clean up.

diff --git a/Engine/Shaders/Private/BRDF.ush b/Engine/Shaders/Private/BRDF.ush
--- a/Engine/Shaders/Private/BRDF.ush
+++ b/Engine/Shaders/Private/BRDF.ush
@@ -19,7 +19,6 @@ struct BxDFContext
 	half YoV;
 	half YoL;
 	half YoH;
-#if 1	// Shit code until we get member functions
 };
 
 void Init( inout BxDFContext Context, half3 N, half3 V, half3 L )
@@ -149,113 +148,6 @@ void SphereMaxNoH( inout BxDFContext Context, float SinAlpha, bool bNewtonIterat
 		}
 	}
 }
-#else
-	void Init( half3 N, half3 V, half3 L )
-	{
-		NoL = dot(N, L);
-		NoV = dot(N, V);
-		VoL = dot(V, L);
-		float InvLenH = rsqrt( 2 + 2 * VoL );
-		NoH = saturate( ( NoL + NoV ) * InvLenH );
-		VoH = saturate( InvLenH + InvLenH * VoL );
-		//NoL = saturate( NoL );
-		//NoV = saturate( abs( NoV ) + 1e-5 );
-
-		XoV = 0.0f;
-		XoL = 0.0f;
-		XoH = 0.0f;
-		YoV = 0.0f;
-		YoL = 0.0f;
-		YoH = 0.0f;
-	}
-
-	void Init( half3 N, half3 X, half3 Y, half3 V, half3 L )
-	{
-		NoL = dot(N, L);
-		NoV = dot(N, V);
-		VoL = dot(V, L);
-		float InvLenH = rsqrt( 2 + 2 * VoL );
-		NoH = saturate( ( NoL + NoV ) * InvLenH );
-		VoH = saturate( InvLenH + InvLenH * VoL );
-		//NoL = saturate( NoL );
-		//NoV = saturate( abs( NoV ) + 1e-5 );
-
-		XoV = dot(X, V);
-		XoL = dot(X, L);
-		XoH = (XoL + XoV) * InvLenH;
-		YoV = dot(Y, V);
-		YoL = dot(Y, L);
-		YoH = (YoL + YoV) * InvLenH;
-	}
-
-	// [ de Carpentier 2017, "Decima Engine: Advances in Lighting and AA" ]
-	void SphereMaxNoH( float SinAlpha, bool bNewtonIteration )
-	{
-		float CosAlpha = sqrt( 1 - Pow2( SinAlpha ) );
-	
-		float RoL = 2 * NoL * NoV - VoL;
-		if( RoL >= CosAlpha )
-		{
-			NoH = 1;
-			XoH = 0;
-			YoH = 0;
-			VoH = abs( NoV );
-		}
-		else
-		{
-			float rInvLengthT = SinAlpha * rsqrt( 1 - RoL*RoL );
-			float NoTr = rInvLengthT * ( NoV - RoL * NoL );
-// Enable once anisotropic materials support area lights
-#if 0
-			float XoTr = rInvLengthT * ( XoV - RoL * XoL );
-			float YoTr = rInvLengthT * ( YoV - RoL * YoL );
-#endif
-			float VoTr = rInvLengthT * ( 2 * NoV*NoV - 1 - RoL * VoL );
-
-			if (bNewtonIteration && SinAlpha != 0)
-			{
-				// dot( cross(N,L), V )
-				float NxLoV = sqrt( saturate( 1 - Pow2(NoL) - Pow2(NoV) - Pow2(VoL) + 2 * NoL * NoV * VoL ) );
-
-				float NoBr = rInvLengthT * NxLoV;
-				float VoBr = rInvLengthT * NxLoV * 2 * NoV;
-
-				float NoLVTr = NoL * CosAlpha + NoV + NoTr;
-				float VoLVTr = VoL * CosAlpha + 1   + VoTr;
-
-				float p = NoBr   * VoLVTr;
-				float q = NoLVTr * VoLVTr;
-				float s = VoBr   * NoLVTr;
-
-				float xNum = q * ( -0.5 * p + 0.25 * VoBr * NoLVTr );
-				float xDenom = p*p + s * (s - 2*p) + NoLVTr * ( (NoL * CosAlpha + NoV) * Pow2(VoLVTr) + q * (-0.5 * (VoLVTr + VoL * CosAlpha) - 0.5) );
-				float TwoX1 = 2 * xNum / ( Pow2(xDenom) + Pow2(xNum) );
-				float SinTheta = TwoX1 * xDenom;
-				float CosTheta = 1.0 - TwoX1 * xNum;
-				NoTr = CosTheta * NoTr + SinTheta * NoBr;
-				VoTr = CosTheta * VoTr + SinTheta * VoBr;
-			}
-
-			NoL = NoL * CosAlpha + NoTr;
-// Enable once anisotropic materials support area lights
-#if 0
-			XoL = XoL * CosAlpha + XoTr;
-			YoL = YoL * CosAlpha + YoTr;
-#endif
-			VoL = VoL * CosAlpha + VoTr;
-
-			float InvLenH = rsqrt( 2 + 2 * VoL );
-			NoH = saturate( ( NoL + NoV ) * InvLenH );
-// Enable once anisotropic materials support area lights
-#if 0
-			XoH = ((XoL + XoV) * InvLenH);	// dot(X, (L+V)/|L+V|)
-			YoH = ((YoL + YoV) * InvLenH);
-#endif
-			VoH = saturate( InvLenH + InvLenH * VoL );
-		}
-	}
-};
-#endif
 
 // Physically based shading model
 // parameterized with the below options
diff --git a/Engine/Source/Runtime/Renderer/Private/Shadows/ShadowScene.h b/Engine/Source/Runtime/Renderer/Private/Shadows/ShadowScene.h
--- a/Engine/Source/Runtime/Renderer/Private/Shadows/ShadowScene.h
+++ b/Engine/Source/Runtime/Renderer/Private/Shadows/ShadowScene.h
@@ -58,7 +58,6 @@ private:
 	public:
 		DECLARE_SCENE_EXTENSION_UPDATER(FUpdater, FShadowScene);
 
-		// I really don't know why I can't use an initializer for this shit
 		FUpdater(FShadowScene& InShadowScene) : ShadowScene(InShadowScene) {}
 
 		// Handle all scene changes wrt lights
