Fix IsThin support for SubSurface Profile.

* When a material use IsThin, we no longer fall into the optimized legacy storage, as we don't have enough bit into the tightly packed header.
* Fix legacy conversion to no longer set material to is thin, as this is now an orthogonal concept to SSS type.
* Fix incorrect material size computation for Subsurface profile when not using legacy compact format. This cause the SSSData to override the SSS data attached to the slab data (ProfileID, Scale, thickness). This was due to SSS profile having haziness feature (due to dual-lobe specular) and which was not tracked correctly/

diff --git a/Engine/Shaders/Private/Substrate/Substrate.ush b/Engine/Shaders/Private/Substrate/Substrate.ush
--- a/Engine/Shaders/Private/Substrate/Substrate.ush
+++ b/Engine/Shaders/Private/Substrate/Substrate.ush
@@ -963,7 +963,7 @@ uint PackSSSProfile(float InProfileId, float InRadius, half InThickness)
 	return Out;
 }
 
-void UnpackSSSProfile(uint In, inout float OutProfileId, inout float OutRadius, inout float OutThickness)
+void UnpackSSSProfile(uint In, inout float OutProfileId, inout float OutRadius, inout half OutThickness)
 {
 	OutProfileId = UnpackR8(In);
 	OutRadius = UnpackR8(In >> 8);
@@ -4386,7 +4386,7 @@ FSubstrateBSDF UnpackSubstrateBSDFIn(FSubstrateMaterialContainer SubstrateBuffer
 				SLAB_SSSPHASEANISOTROPY(OutBSDF)= saturate(1.f - Opacity);
 			#endif
 
-				const bool bTwoSided = OptimisedLegacyMode == SINGLE_OPTLEGACYMODE_TWO_SIDED_SSSWRAP;
+				const bool bTwoSidedWrap = OptimisedLegacyMode == SINGLE_OPTLEGACYMODE_TWO_SIDED_SSSWRAP;
 
 				// Force a BSDF state to be a single layer on top.
 				BSDF_SETTYPE(OutBSDF, SUBSTRATE_BSDF_TYPE_SLAB);
@@ -4394,10 +4394,10 @@ FSubstrateBSDF UnpackSubstrateBSDFIn(FSubstrateMaterialContainer SubstrateBuffer
 				BSDF_SETHASF90(OutBSDF, 0);
 				BSDF_SETHASANISOTROPY(OutBSDF, 0);
 				BSDF_SETISTOPLAYER(OutBSDF, 1);
-				BSDF_SETSSSTYPE(OutBSDF, bTwoSided ? SSS_TYPE_TWO_SIDED_WRAP : SSS_TYPE_WRAP);
+				BSDF_SETSSSTYPE(OutBSDF, bTwoSidedWrap ? SSS_TYPE_TWO_SIDED_WRAP : SSS_TYPE_WRAP);
 				BSDF_SETHASMFP(OutBSDF, 1);
 				BSDF_SETHASHAZINESS(OutBSDF, 0);
-				BSDF_SETISTHIN(OutBSDF, (bTwoSided ? 1 : 0));
+				BSDF_SETISTHIN(OutBSDF, 0);
 				BSDF_SETHASFUZZ(OutBSDF, 0);
 				BSDF_SETHASTRANSABOVE(OutBSDF, 0);
 				BSDF_SETHASGLINT(OutBSDF, 0);
diff --git a/Engine/Shaders/Private/Substrate/SubstrateExport.ush b/Engine/Shaders/Private/Substrate/SubstrateExport.ush
--- a/Engine/Shaders/Private/Substrate/SubstrateExport.ush
+++ b/Engine/Shaders/Private/Substrate/SubstrateExport.ush
@@ -335,6 +335,10 @@ void PackSubstrateOut(
 				const uint SingleIsSSSProfile		= 0x80;
 				// No need to add anisotropy since it is not supported by the simple mode.
 
+				// If a surface has IsThin bit, opt-out legacy format.
+				// Legacy mode doesn't need support for IsThin, and their storage format does not have space for storing IsThin bit. 
+				const bool bIsThinSurface = BSDF_GETISTHIN(BSDF);
+
 				uint OptimisedModeMask = 0;
 				if (BSDF_GETHASHAZINESS(BSDF))
 				{
@@ -375,17 +379,17 @@ void PackSubstrateOut(
 				{
 					OptimisedLegacyMode = SINGLE_OPTLEGACYMODE_CLOTH;
 				}				
-				else if ((OptimisedModeMask & SingleIsTwoSidedWrap) == SingleIsTwoSidedWrap && !BSDF_GETHASFUZZ(BSDF))
+				else if ((OptimisedModeMask & SingleIsTwoSidedWrap) == SingleIsTwoSidedWrap && !BSDF_GETHASFUZZ(BSDF) && !bIsThinSurface)
 				{
 					OptimisedLegacyMode = SINGLE_OPTLEGACYMODE_TWO_SIDED_SSSWRAP;
 					SSSData.Header.Bytes = 0; // Force to skip SSSData writing
 				}
-				else if ((OptimisedModeMask & SingleIsWrap) == SingleIsWrap && !BSDF_GETHASFUZZ(BSDF))
+				else if ((OptimisedModeMask & SingleIsWrap) == SingleIsWrap && !BSDF_GETHASFUZZ(BSDF) && !bIsThinSurface)
 				{
 					OptimisedLegacyMode = SINGLE_OPTLEGACYMODE_SSSWRAP;
 					SSSData.Header.Bytes = 0; // Force to skip SSSData writing
 				}
-				else if ((OptimisedModeMask & SingleIsSSSProfile) == SingleIsSSSProfile && !BSDF_GETHASFUZZ(BSDF))
+				else if ((OptimisedModeMask & SingleIsSSSProfile) == SingleIsSSSProfile && !bIsThinSurface)
 				{
 					OptimisedLegacyMode = SINGLE_OPTLEGACYMODE_SSSPROFILE;
 					SSSData.Header.Bytes = 0; // Force to skip SSSData writing
diff --git a/Engine/Shaders/Private/Substrate/SubstrateLegacyConversion.ush b/Engine/Shaders/Private/Substrate/SubstrateLegacyConversion.ush
--- a/Engine/Shaders/Private/Substrate/SubstrateLegacyConversion.ush
+++ b/Engine/Shaders/Private/Substrate/SubstrateLegacyConversion.ush
@@ -387,7 +387,6 @@ FSubstrateData SubstrateConvertLegacyMaterialDynamic(
 		LegacySlab.Emissive = Emissive;
 		LegacySlab.SSSType = SSS_TYPE_TWO_SIDED_WRAP;
 		LegacySlab.Thickness = Thickness;
-		LegacySlab.bIsThinSurface = true;
 
 		FinalWeight = 1.0f;
 #if SUBSTRATE_LEGACYCONVERT_OUTPUT_OPACITY
diff --git a/Engine/Source/Runtime/Engine/Private/Materials/MaterialExpressionSubstrate.cpp b/Engine/Source/Runtime/Engine/Private/Materials/MaterialExpressionSubstrate.cpp
--- a/Engine/Source/Runtime/Engine/Private/Materials/MaterialExpressionSubstrate.cpp
+++ b/Engine/Source/Runtime/Engine/Private/Materials/MaterialExpressionSubstrate.cpp
@@ -1253,6 +1253,7 @@ FSubstrateOperator* UMaterialExpressionSubstrateSlabBSDF::SubstrateGenerateMater
 	if (HasAnisotropy())		{ SubstrateOperator.BSDFFeatures |= ESubstrateBsdfFeature_Anisotropy; }
 	if (HasGlint())				{ SubstrateOperator.BSDFFeatures |= ESubstrateBsdfFeature_Glint; }
 	if (HasSpecularProfile())	{ SubstrateOperator.BSDFFeatures |= ESubstrateBsdfFeature_SpecularProfile; }
+	if (HasSSSProfile())		{ SubstrateOperator.BSDFFeatures |= ESubstrateBsdfFeature_SecondRoughnessOrSimpleClearCoat; } // If a Slab has a subsurface profile, it will have haziness in order to to support the dual-specular lobe from the profile.
 
 	SubstrateOperator.ThicknessIndex = Compiler->SubstrateThicknessStackGetThicknessIndex();
 	SubstrateOperator.bBSDFWritesEmissive = EmissiveColor.IsConnected();
