Substrate - added support for clear coat bottom normal to the path tracer.

#tests PC editor + PathTracer
#rnx

diff --git a/Engine/Shaders/Private/PathTracing/Material/PathTracingSubstrate.ush b/Engine/Shaders/Private/PathTracing/Material/PathTracingSubstrate.ush
--- a/Engine/Shaders/Private/PathTracing/Material/PathTracingSubstrate.ush
+++ b/Engine/Shaders/Private/PathTracing/Material/PathTracingSubstrate.ush
@@ -242,8 +242,11 @@ FMaterialEval Substrate_EvalMaterial(
 
 struct FSubstrateClearCoatData
 {
-	float3x3 Basis;
-	float3 V;
+	float3x3 BaseBasis;
+	float3x3 CoatBasis;
+
+	float3 BaseV;
+	float3 CoatV;
 
 	float  DiffuseWeight;
 	float3 F0;
@@ -261,20 +264,34 @@ struct FSubstrateClearCoatData
 
 	void PrepareBSDF(float3 V_World, FPathTracingPayload Payload, bool bUseSSSColor = false)
 	{
-		Basis = GetGGXBasis(Payload.RoughnessData.x, Payload.Anisotropy, Payload.WorldNormal, Payload.WorldTangent, BaseAlpha);
+		float3 BaseWorldNormal = Payload.WorldNormal;
+#if CLEAR_COAT_BOTTOM_NORMAL
+		if ((Payload.BottomNormalOct16bits & 0xFFFF) != 0)
+		{
+			const float2 BottomNormalOct = UnpackBottomNormalOct(Payload.BottomNormalOct16bits);
+			BaseWorldNormal = SubstrateGetSimpleCoatBottomNormal(BottomNormalOct, Payload.WorldNormal);
+		}
+#endif
+
+		BaseBasis = GetGGXBasis(Payload.RoughnessData.x, Payload.Anisotropy, BaseWorldNormal, Payload.WorldTangent, BaseAlpha);
+		BaseAlpha = GetGGXAlpha(Payload.RoughnessData.x, Payload.Anisotropy);
+
+		CoatBasis = GetGGXBasis(Payload.RoughnessData.y, Payload.Anisotropy, Payload.WorldNormal, Payload.WorldTangent, CoatAlpha);
 		CoatAlpha = GetGGXAlpha(Payload.RoughnessData.y, Payload.Anisotropy);
 
-		V = mul(Basis, V_World);
+		BaseV = mul(BaseBasis, V_World);
+		CoatV = mul(CoatBasis, V_World);
 
 		// make sure F0=0.0 fades off the specular lobe completely
 		F0 = Payload.SpecularColor;
 		F90 = Payload.SpecularEdgeColor * F0RGBToMicroOcclusion(F0);
 
-		const float NoV = saturate(V.z);
+		const float BaseNoV = saturate(BaseV.z);
+		const float CoatNoV = saturate(CoatV.z);
 		// SUBSTRATE_TODO: Add support for F82 model?
-		BaseSpec = ComputeGGXSpecEnergyTermsRGB(Payload.RoughnessData.x, NoV, F0, F90);
-		CoatSpec = ComputeGGXSpecEnergyTermsA(Payload.RoughnessData.y, NoV, CLEAR_COAT_F0);
-		Fuzz.Prepare(V, Payload.FuzzRoughness, Payload.FuzzAmount);
+		BaseSpec = ComputeGGXSpecEnergyTermsRGB(Payload.RoughnessData.x, BaseNoV, F0, F90);
+		CoatSpec = ComputeGGXSpecEnergyTermsA(Payload.RoughnessData.y, CoatNoV, CLEAR_COAT_F0);
+		Fuzz.Prepare(CoatNoV, Payload.FuzzRoughness, Payload.FuzzAmount);
 
 		CoatAttenuation = 1.0 - Payload.RoughnessData.z * CoatSpec.E;
 		DiffuseWeight = Fuzz.Attenuation * CoatAttenuation * (1.0 - Luminance(BaseSpec.E));
@@ -305,8 +322,12 @@ FMaterialSample SubstrateClearCoat_SampleMaterial(
 	FSubstrateClearCoatData Data = (FSubstrateClearCoatData)0;
 	Data.PrepareBSDF(V_World, Payload);
 
-	float3 L = 0, H = 0, V = Data.V;
 	float OutRoughness = 1;
+	float3 OutNormal = Payload.WorldNormal;
+	float3 OutL = 0;
+	float3 OutV = 0;
+	float3 OutH = 0;
+	float3x3 OutBasis;
 
 	const bool bSampledDiffuse = RandSample.x < Data.LobeCdf.x;
 	const bool bSampledSpecular = RandSample.x < Data.LobeCdf.z;
@@ -314,8 +335,10 @@ FMaterialSample SubstrateClearCoat_SampleMaterial(
 	{
 		RandSample.x = RescaleRandomNumber(RandSample.x, 0.0, Data.LobeCdf.x);
 		// diffuse lobe
-		L = CosineSampleHemisphere(RandSample.xy).xyz;
-		H = normalize(L + V);
+		OutL = CosineSampleHemisphere(RandSample.xy).xyz;
+		OutV = Data.BaseV;
+		OutH = normalize(OutL + OutV);
+		OutBasis = Data.BaseBasis;
 	}
 	else if (bSampledSpecular)
 	{
@@ -325,16 +348,26 @@ FMaterialSample SubstrateClearCoat_SampleMaterial(
 		{
 			RandSample.x = RescaleRandomNumber(RandSample.x, Data.LobeCdf.x, Data.LobeCdf.y);
 			OutRoughness = Payload.RoughnessData.x;
+
+			const float2 BottomNormalOct = UnpackBottomNormalOct(Payload.BottomNormalOct16bits);
+			OutNormal = SubstrateGetSimpleCoatBottomNormal(BottomNormalOct, Payload.WorldNormal);
+
+			OutBasis = Data.BaseBasis;
+			OutV = Data.BaseV;
 		}
 		else
 		{
 			RandSample.x = RescaleRandomNumber(RandSample.x, Data.LobeCdf.y, Data.LobeCdf.z);
 			OutRoughness = Payload.RoughnessData.y;
+
+			OutBasis = Data.CoatBasis;
+			OutV = Data.CoatV;
 		}
 
-		H = ImportanceSampleVisibleGGX(RandSample.xy, bUseBaseSpec ? Data.BaseAlpha : Data.CoatAlpha, V).xyz;
-		L = reflect(-V, H);
-		if (L.z <= 0)
+		OutV = bUseBaseSpec ? Data.BaseV : Data.CoatV;
+		OutH = ImportanceSampleVisibleGGX(RandSample.xy, bUseBaseSpec ? Data.BaseAlpha : Data.CoatAlpha, OutV).xyz;
+		OutL = reflect(-OutV, OutH);
+		if (OutL.z <= 0)
 		{
 			// invalid output direction, exit early
 			return NullMaterialSample();
@@ -344,13 +377,19 @@ FMaterialSample SubstrateClearCoat_SampleMaterial(
 	{
 		// cloth lobe
 		RandSample.x = RescaleRandomNumber(RandSample.x, Data.LobeCdf.z, 1.0);
-		L = Data.Fuzz.Sample(RandSample.xy);
-		H = normalize(L + V);
+		OutV = Data.CoatV;
+		OutL = Data.Fuzz.Sample(RandSample.xy);
+		OutH = normalize(OutL + OutV);
+
+		OutBasis = Data.CoatBasis;
 	}
 
 	// transform to world space
-	const float3 L_World = normalize(mul(L, Data.Basis));
-	const float3 N_World = Payload.WorldNormal;
+	float3 L = OutL;
+	float3 V = OutV;
+	float3 H = OutH;
+	const float3 L_World = normalize(mul(L, OutBasis)); 
+	const float3 N_World = OutNormal;
 	const float NoV = saturate(V.z);
 	const float NoL = saturate(L.z);
 	const float VoH = saturate(dot(V, H));
@@ -426,54 +465,61 @@ FMaterialEval SubstrateClearCoat_EvalMaterial(
 	FSubstrateClearCoatData Data = (FSubstrateClearCoatData)0;
 	Data.PrepareBSDF(V_World, Payload);
 
-	const float3 N_World = Payload.WorldNormal;
-
 	// move vectors into right shading frame
-	const float3 V = Data.V;
-	const float3 L = mul(Data.Basis, L_World);
-	const float3 H = normalize(V + L);
-	const float NoV = saturate(V.z);
-	const float NoL = saturate(L.z);
-	const float VoH = saturate(dot(V, H));
-	const float NoH = saturate(H.z);
+	const float3 BaseV   = Data.BaseV;
+	const float3 BaseL   = mul(Data.BaseBasis, L_World);
+	const float3 BaseH   = normalize(BaseV + BaseL);
+	const float  BaseNoV = saturate(BaseV.z);
+	const float  BaseNoL = saturate(BaseL.z);
+	const float  BaseVoH = saturate(dot(BaseV, BaseH));
+	const float  BaseNoH = saturate(BaseH.z);
+
+	const float3 CoatV   = Data.CoatV;
+	const float3 CoatL   = mul(Data.CoatBasis, L_World);
+	const float3 CoatH   = normalize(CoatV + CoatL);
+	const float  CoatNoV = saturate(CoatV.z);
+	const float  CoatNoL = saturate(CoatL.z);
+	const float  CoatVoH = saturate(dot(CoatV, CoatH));
+	const float  CoatNoH = saturate(CoatH.z);
 
 	FMaterialEval Result = NullMaterialEval();
 
 	const float Metallic = F0RGBToMetallic(Payload.SpecularColor);
 	const float3 BaseColor = lerp(Payload.DiffuseColor, Payload.SpecularColor, Metallic);
-	const float3 Transmission = lerp(1.0, SimpleClearCoatTransmittance(NoL, NoV, Metallic, BaseColor), Payload.RoughnessData.z);
+	const float3 Transmission = lerp(1.0, SimpleClearCoatTransmittance(CoatNoL, CoatNoV, Metallic, BaseColor), Payload.RoughnessData.z);
 
-	const float ShadowTerminator = ShadowTerminatorTerm(L_World, N_World, Payload.WorldSmoothNormal);
+	const float BasehadowTerminator = ShadowTerminatorTerm(L_World,  Data.BaseBasis[2], Payload.WorldSmoothNormal);
+	const float CoatShadowTerminator = ShadowTerminatorTerm(L_World, Data.CoatBasis[2], Payload.WorldSmoothNormal);
 	{
-		const float3 Diffuse = Transmission * Data.DiffuseWeight * GetPathTracingDiffuseModel(Payload.DiffuseColor, Payload.RoughnessData.x, NoV, NoL, VoH, NoH);
-		Result.AddLobeWithMIS(DiffuseSpecularScale.x * Diffuse * ShadowTerminator, NoL / PI, Data.LobePdf.x);
+		const float3 Diffuse = Transmission * Data.DiffuseWeight * GetPathTracingDiffuseModel(Payload.DiffuseColor, Payload.RoughnessData.x, BaseNoV, BaseNoL, BaseVoH, BaseNoH);
+		Result.AddLobeWithMIS(DiffuseSpecularScale.x * Diffuse * BasehadowTerminator, BaseNoL / PI, Data.LobePdf.x);
 	}
 	{
-		const float2 BaseGGXResult = GGXEvalReflection(L, V, H, Data.BaseAlpha);
-		const float2 CoatGGXResult = GGXEvalReflection(L, V, H, Data.CoatAlpha);
+		const float2 BaseGGXResult = GGXEvalReflection(BaseL, BaseV, BaseH, Data.BaseAlpha);
+		const float2 CoatGGXResult = GGXEvalReflection(CoatL, CoatV, CoatH, Data.CoatAlpha);
 		const float BaseSpecPdf = BaseGGXResult.y;
 		const float CoatSpecPdf = CoatGGXResult.y;
 
-		const float3 FTint = SubstrateSpecularTint(Payload, NoV, NoL, VoH, NoH);
+		const float3 FTint = SubstrateSpecularTint(Payload, CoatNoV, CoatNoL, CoatVoH, CoatNoH);
 #if SUBSTRATE_FRESNEL_F82
 		// SUBSTRATE_TODO: Need to tune the logic for specular micro-occlusion ...
-		const float3 FBase = Data.Fuzz.Attenuation * F_AdobeF82(Data.F0, Data.F90, VoH) * FTint;
+		const float3 FBase = Data.Fuzz.Attenuation * F_AdobeF82(Data.F0, Data.F90, CoatVoH) * FTint;
 #else
-		const float3 FBase = Data.Fuzz.Attenuation * F_Schlick(Data.F0, Data.F90, VoH) * FTint;
+		const float3 FBase = Data.Fuzz.Attenuation * F_Schlick(Data.F0, Data.F90, CoatVoH) * FTint;
 #endif
-		const float3 FCoat = Payload.RoughnessData.z * ClearCoatFresnel(VoH) * FTint;
+		const float3 FCoat = Payload.RoughnessData.z * ClearCoatFresnel(CoatVoH) * FTint;
 
 		Result.AddLobeWithMIS(Data.CoatAttenuation * FBase * BaseGGXResult.x * Data.BaseSpec.W * Transmission, BaseSpecPdf, Data.LobePdf.y);
 		Result.AddLobeWithMIS(Data.Fuzz.Attenuation * FCoat * CoatGGXResult.x * Data.CoatSpec.W               , CoatSpecPdf, Data.LobePdf.z);
 	}
 	{
 		// Cloth Lobe
-		const float4 ClothResult = Data.Fuzz.Eval(L, V, H, Payload.FuzzColor);
+		const float4 ClothResult = Data.Fuzz.Eval(CoatL, CoatV, CoatH, Payload.FuzzColor);
 		const float ClothPdf = ClothResult.w;
-		Result.AddLobeWithMIS(DiffuseSpecularScale.y * ClothResult.xyz * ShadowTerminator, ClothPdf, Data.LobePdf.w);
+		Result.AddLobeWithMIS(DiffuseSpecularScale.y * ClothResult.xyz * CoatShadowTerminator, ClothPdf, Data.LobePdf.w);
 	}
 
-	Result.Weight *= SubstrateLobeWeight(Payload, NoL);
+	Result.Weight *= SubstrateLobeWeight(Payload, CoatNoL);
 
 	return Result;
 }
diff --git a/Engine/Shaders/Private/PathTracing/PathTracingCommon.ush b/Engine/Shaders/Private/PathTracing/PathTracingCommon.ush
--- a/Engine/Shaders/Private/PathTracing/PathTracingCommon.ush
+++ b/Engine/Shaders/Private/PathTracing/PathTracingCommon.ush
@@ -149,6 +149,20 @@ float4 PayloadDecodeRoughnessAniso(uint rgba)
 	return float4(r, g, b, a);
 }
 
+uint PayloadEncodeFuzzAndBottomNormal(float FuzzAmount, float FuzzRoughness, uint BottomNormalOct16bits)
+{
+	uint r = uint(floor(saturate(FuzzAmount) * 255.0 + 0.5));
+	uint g = uint(floor(saturate(FuzzRoughness) * 255.0 + 0.5));
+	return (r << 24) | (g << 16) | (BottomNormalOct16bits & 0xFFFF);
+}
+
+void PayloadDecodeFuzzAndBottomNormal(uint rgba, inout float FuzzAmount, inout float FuzzRoughness, inout uint BottomNormalOct16bits)
+{
+	FuzzAmount = float((rgba >> 24)) * (1.0 / 255.0);
+	FuzzRoughness = float((rgba >> 16) & 255) * (1.0 / 255.0);
+	BottomNormalOct16bits = rgba & 0xFFFF;
+}
+
 float3 AdjustShadingNormal(float3 ShadingNormal, float3 GeoNormal, float3 RayDirection)
 {
 	// Clip shading normal in a view dependent way such that the reflection stays above the geometric normal
@@ -289,7 +303,6 @@ struct FPackedPathTracingPayload : FMinimalPayload
 	                                                  f16tof32(PackedData[12])); }
 	float UnpackPhaseG()              { return f16tof32(PackedData[12] >> 16); }
 	float3 UnpackFuzzColor()          { return PayloadDecodeLDRColor(PackedData[13]); }
-	float2 UnpackFuzzData()           { return UnpackFloat2FromUInt(PackedData[14]); }
 	float3 UnpackWeightV()            { return PayloadDecodeLDRColor(PackedData[15]); }
 	float3 UnpackTransmittanceN()     { return PayloadDecodeLDRColor(PackedData[16]); }
 
@@ -455,6 +468,7 @@ struct FPathTracingPayload : FMinimalPayload
 	float3 SpecularColor;
 	float3 SpecularEdgeColor;
 	float3 RoughnessData; // r0,r1,rblend
+	uint BottomNormalOct16bits;
 	float  Anisotropy;
 	float3 WorldTangent;
 	float3 MeanFreePath;
@@ -774,7 +788,7 @@ FPackedPathTracingPayload PackPathTracingPayload(FPathTracingPayload Input)
 	Output.PackedData[11] = PackFloat2ToUInt(float2(Input.MeanFreePath.xy));
 	Output.PackedData[12] = PackFloat2ToUInt(float2(Input.MeanFreePath.z, Input.PhaseG));
 	Output.PackedData[13] = PayloadEncodeLDRColor(Input.FuzzColor);
-    Output.PackedData[14] = PackFloat2ToUInt(float2(Input.FuzzAmount, Input.FuzzRoughness));
+	Output.PackedData[14] = PayloadEncodeFuzzAndBottomNormal(Input.FuzzAmount, Input.FuzzRoughness, Input.BottomNormalOct16bits);
 	Output.PackedData[15] = PayloadEncodeLDRColor(Input.WeightV);
 	Output.PackedData[16] = PayloadEncodeLDRColor(Input.TransmittanceN);
 	Output.PackedData[17] = PackFloat2ToUInt(float2(Input.SpecularProfileId, abs(Input.CoverageAboveAlongN)));
@@ -845,9 +859,7 @@ FPathTracingPayload UnpackPathTracingPayload(FPackedPathTracingPayload Input, FR
 	Output.MeanFreePath        = Input.UnpackMeanFreePath();
 	Output.PhaseG              = Input.UnpackPhaseG();
 	Output.FuzzColor           = Input.UnpackFuzzColor();
-	const float2 FuzzData      = Input.UnpackFuzzData();
-	Output.FuzzAmount          = FuzzData.x;
-	Output.FuzzRoughness       = FuzzData.y;
+	PayloadDecodeFuzzAndBottomNormal(Input.PackedData[14], Output.FuzzAmount, Output.FuzzRoughness, Output.BottomNormalOct16bits);
 	Output.WeightV             = Input.UnpackWeightV();
 	Output.TransmittanceN      = Input.UnpackTransmittanceN();
 	Output.CoverageAboveAlongN = Input.UnpackCoverageAboveAlongN();
diff --git a/Engine/Shaders/Private/PathTracing/PathTracingMaterialHitShader.usf b/Engine/Shaders/Private/PathTracing/PathTracingMaterialHitShader.usf
--- a/Engine/Shaders/Private/PathTracing/PathTracingMaterialHitShader.usf
+++ b/Engine/Shaders/Private/PathTracing/PathTracingMaterialHitShader.usf
@@ -576,6 +576,11 @@ RAY_TRACING_ENTRY_CLOSEST_HIT(PathTracingMaterialCHS,
 						}
 						Payload.RoughnessData.y = MakeRoughnessSafe(Haziness.Roughness, PackedPayload.GetPathRoughness());
 						Payload.RoughnessData.z = Haziness.Weight;
+					#if CLEAR_COAT_BOTTOM_NORMAL
+						Payload.BottomNormalOct16bits = SLAB_HAZINESS(BSDF) >> 16;
+					#else
+						Payload.BottomNormalOct16bits = 0;
+					#endif
 					}
 					Payload.Anisotropy = BSDF_GETHASANISOTROPY(BSDF) ? SLAB_ANISOTROPY(BSDF) : 0.0;
 					if (BSDF_GETHASFUZZ(BSDF))
diff --git a/Engine/Shaders/Private/Substrate/Substrate.ush b/Engine/Shaders/Private/Substrate/Substrate.ush
--- a/Engine/Shaders/Private/Substrate/Substrate.ush
+++ b/Engine/Shaders/Private/Substrate/Substrate.ush
@@ -740,6 +740,11 @@ float3 SubstrateGetSimpleCoatBottomNormal(float2 BottomNormalOct, float3 TopWorl
 	return OctahedronToUnitVector(OctNormal);
 }
 
+float2 UnpackBottomNormalOct(uint BottomNormalOct16bits)
+{
+	return float2(UnpackR8((BottomNormalOct16bits >> 0) & 0xFF), UnpackR8((BottomNormalOct16bits >> 8) & 0xFF));
+}
+
 FHaziness InitialiseHaziness()
 {
 	FHaziness Haziness = (FHaziness)0;
@@ -764,8 +769,7 @@ FHaziness UnpackHaziness(uint Data)
 	Haziness.Weight				= float(uint((Data>>8) & 0x7F)) * (1.0f / 127.0f);
 	Haziness.Roughness			= UnpackR8(Data);
 #if CLEAR_COAT_BOTTOM_NORMAL
-	Haziness.BottomNormalOct.x = UnpackR8((Data >> 16) & 0xFF);
-	Haziness.BottomNormalOct.y = UnpackR8((Data >> 24) & 0xFF);
+	Haziness.BottomNormalOct = UnpackBottomNormalOct(Data >> 16);
 #endif
 	return Haziness;
 }
diff --git a/Engine/Source/Runtime/Core/Private/UObject/FortniteMainBranchObjectVersion.cpp b/Engine/Source/Runtime/Core/Private/UObject/FortniteMainBranchObjectVersion.cpp
--- a/Engine/Source/Runtime/Core/Private/UObject/FortniteMainBranchObjectVersion.cpp
+++ b/Engine/Source/Runtime/Core/Private/UObject/FortniteMainBranchObjectVersion.cpp
@@ -14,6 +14,6 @@ TMap<FGuid, FGuid> FFortniteMainBranchObjectVersion::GetSystemGuids()
 	SystemGuids.Add(DevGuids.Niagara_LatestScriptCompileVersion, FGuid("670A2EBDEE1A4794BB77BBFF331EA7EF"));
 	SystemGuids.Add(DevGuids.SkeletalMeshDerivedDataVersion, FGuid("5E93192FCF3647EEB7F9D4FFF1750050"));
 	SystemGuids.Add(DevGuids.STATICMESH_DERIVEDDATA_VER, FGuid("656E338C78C34641A5D3BFAE303997AF"));
-	SystemGuids.Add(DevGuids.MaterialTranslationDDCVersion, FGuid("D4A49317317F483C86516F96115A1CEC"));
+	SystemGuids.Add(DevGuids.MaterialTranslationDDCVersion, FGuid("3DC4760CBA5A40E58DF2A2365A8C4F28"));
 	return SystemGuids;
 }
