CombinedLUT: Enforce floating-point texture format when ScRGB is active - currently the HDR backbuffer color space in editor on Windows.

* Simpler logic which will facilitate ACES 2.0 integration.
* Less work on the tonemap shader side.
* Note: The output transform scales normalized colors by Y_MAX (with STRETCH_BLACK defaulting to true) to output values in nits for PQ. As such, we can simply divide by 80 to obtain ScRGB units.

diff --git a/Engine/Shaders/Private/PostProcessCombineLUTs.usf b/Engine/Shaders/Private/PostProcessCombineLUTs.usf
--- a/Engine/Shaders/Private/PostProcessCombineLUTs.usf
+++ b/Engine/Shaders/Private/PostProcessCombineLUTs.usf
@@ -348,7 +348,7 @@ float4 CombineLUTsCommon(float2 InUV, uint InLayerIndex)
 
 	// ACES transform with PQ/2084 encoding, user specified gamut 
 	else if( GetOutputDevice() == TONEMAPPER_OUTPUT_ACES1000nitST2084 || GetOutputDevice() == TONEMAPPER_OUTPUT_ACES2000nitST2084)
-	{		
+	{
 		FACESTonemapParams AcesParams = ComputeACESTonemapParams(ACESMinMaxData, ACESMidData, ACESCoefsLow_0, ACESCoefsHigh_0, ACESCoefsLow_4, ACESCoefsHigh_4, ACESSceneColorMultiplier, ACESGamutCompression);
 		float3 ODTColor = ACESOutputTransform( GradedColor, (float3x3)WorkingColorSpace.ToAP0, AcesParams);
 
@@ -360,15 +360,11 @@ float4 CombineLUTsCommon(float2 InUV, uint InLayerIndex)
 
 	// ACES transform with PQ/2084 encoding, AP1 gamut
 	else if( GetOutputDevice() == TONEMAPPER_OUTPUT_ACES1000nitScRGB || GetOutputDevice() == TONEMAPPER_OUTPUT_ACES2000nitScRGB)
-	{		
+	{
 		FACESTonemapParams AcesParams = ComputeACESTonemapParams(ACESMinMaxData, ACESMidData, ACESCoefsLow_0, ACESCoefsHigh_0, ACESCoefsLow_4, ACESCoefsHigh_4, ACESSceneColorMultiplier, ACESGamutCompression);
 		float3 ODTColor = ACESOutputTransform( GradedColor, (float3x3)WorkingColorSpace.ToAP0, AcesParams);
-
-		// No conversion from AP1 to specified output gamut for ScRGB as it might contain negative values. 
-		// In this case, the AP1->output conversion happens through ST2084ToScRGB in the tonemap function
-
-		// Apply conversion to ST-2084 (Dolby PQ)
-		OutDeviceColor = LinearToST2084( ODTColor );
+		
+		OutDeviceColor = mul(AP1_2_sRGB, ODTColor / UE_SCRGB_WHITE_NITS);
 	}	
 	
 	else if( GetOutputDevice() == TONEMAPPER_OUTPUT_LinearEXR)
diff --git a/Engine/Shaders/Private/PostProcessDeviceEncodingOnly.usf b/Engine/Shaders/Private/PostProcessDeviceEncodingOnly.usf
--- a/Engine/Shaders/Private/PostProcessDeviceEncodingOnly.usf
+++ b/Engine/Shaders/Private/PostProcessDeviceEncodingOnly.usf
@@ -61,6 +61,7 @@ float4 DeviceEncodingOnlyCommonPS(
 
 	half4 SceneColor = SampleSceneColor(SceneUV);
 
+	const float3x3 AP1_2_sRGB = mul( XYZ_2_sRGB_MAT, mul( D60_2_D65_CAT, AP1_2_XYZ_MAT ) );
 	const float3x3 AP1_2_Output  = OuputGamutMappingMatrix( OutputGamut );
 
 	// Apply "gamma" curve adjustment.
@@ -118,7 +119,7 @@ float4 DeviceEncodingOnlyCommonPS(
 		float3 ODTColor = ACESOutputTransform(FilmColor, (float3x3) WorkingColorSpace.ToAP0, AcesParams);
 
 		// Apply conversion to ScRGB
-		OutDeviceColor = ST2084ToScRGB( LinearToST2084( ODTColor ), GetOutputDevice(), OutputMaxLuminance );
+		OutDeviceColor = mul( AP1_2_sRGB, ODTColor / UE_SCRGB_WHITE_NITS );
 	}
 	
 	else if( GetOutputDevice() == TONEMAPPER_OUTPUT_LinearEXR)
diff --git a/Engine/Shaders/Private/PostProcessTonemap.usf b/Engine/Shaders/Private/PostProcessTonemap.usf
--- a/Engine/Shaders/Private/PostProcessTonemap.usf
+++ b/Engine/Shaders/Private/PostProcessTonemap.usf
@@ -546,11 +546,7 @@ float4 TonemapCommonPS(
 		#endif
 
 		BRANCH
-		if(GetOutputDevice() == TONEMAPPER_OUTPUT_ACES1000nitScRGB || GetOutputDevice() == TONEMAPPER_OUTPUT_ACES2000nitScRGB)
-		{
-			OutColor.xyz = ST2084ToScRGB(OutColor.xyz, GetOutputDevice(), OutputMaxLuminance);
-		}
-		else if(GetOutputDevice() == TONEMAPPER_OUTPUT_LinearEXR)
+		if(GetOutputDevice() == TONEMAPPER_OUTPUT_LinearEXR)
 		{
 			OutColor.xyz = ST2084ToLinear(OutColor.xyz);
 		}
diff --git a/Engine/Shaders/Private/TonemapCommon.ush b/Engine/Shaders/Private/TonemapCommon.ush
--- a/Engine/Shaders/Private/TonemapCommon.ush
+++ b/Engine/Shaders/Private/TonemapCommon.ush
@@ -307,55 +307,7 @@ static const bool UE_ACES_STRETCH_BLACK = true; // stretch black luminance to a
 static const bool UE_ACES_D60_SIM = false;
 static const bool UE_ACES_LEGAL_RANGE = false;
 
-static const float UE_ACES_HDR_Y_MIN = 0.0001;  // black luminance (cd/m^2)
-static const float UE_ACES_HDR_Y_MID = 15.0;    // mid-point luminance (cd/m^2)
-
-static const float UE_ACES_SDR_Y_MIN = 0.0001;  // black luminance (cd/m^2)
-static const float UE_ACES_SDR_Y_MID = 7.2;     // mid-point luminance (cd/m^2)
-static const float UE_ACES_SDR_Y_MAX = 108.0;   // peak white luminance (cd/m^2)
-
-//
-// ACES sRGB D65 Output Transform - Forward and Inverse
-//  Input is scene-referred linear values in the working space gamut
-//  Output is output-referred linear values in the sRGB gamut
-//
-float3 ACESOutputTransformsRGBD65( float3 SceneReferredLinearColor, const float3x3 WCS_2_AP0 )
-{
-	static const float Gain = 2.22; // chosen such that outputTransform(0.18 * Gain) / UE_ACES_SDR_Y_MAX = 0.18
-	float3 aces = mul( WCS_2_AP0, SceneReferredLinearColor * Gain );
-
-	const FColorSpace UE_ACES_SRGB_DISPLAY_PRI = GetColorSpace(TONEMAPPER_GAMUT_sRGB_D65);
-	const FColorSpace UE_ACES_SRGB_LIMITING_PRI = GetColorSpace(TONEMAPPER_GAMUT_sRGB_D65);
-
-	float3 OutputReferredLinearsRGBColor = outputTransform(aces, UE_ACES_SDR_Y_MIN, UE_ACES_SDR_Y_MID, UE_ACES_SDR_Y_MAX, UE_ACES_SRGB_DISPLAY_PRI, UE_ACES_SRGB_LIMITING_PRI, UE_ACES_EOTF,
-														   UE_ACES_SURROUND, UE_ACES_STRETCH_BLACK, UE_ACES_D60_SIM, UE_ACES_LEGAL_RANGE);
-
-	return OutputReferredLinearsRGBColor / UE_ACES_SDR_Y_MAX;
-}
-
-float3 ACESOutputTransformsRGBD65( float3 SceneReferredLinearsRGBColor )
-{
-	const float3x3 sRGB_2_AP0 = mul( XYZ_2_AP0_MAT, mul( D65_2_D60_CAT, sRGB_2_XYZ_MAT ) );
-
-	return ACESOutputTransformsRGBD65( SceneReferredLinearsRGBColor, sRGB_2_AP0 );
-}
-
-float3 InverseACESOutputTransformsRGBD65( float3 OutputReferredLinearsRGBColor )
-{
-	const float3x3 AP0_2_sRGB = mul( XYZ_2_sRGB_MAT, mul( D60_2_D65_CAT, AP0_2_XYZ_MAT ) );
-	static const float Gain = 2.22; // chosen such that outputTransform(0.18 * Gain) / UE_ACES_SDR_Y_MAX = 0.18
-
-	const FColorSpace UE_ACES_SRGB_DISPLAY_PRI = GetColorSpace(TONEMAPPER_GAMUT_sRGB_D65);
-	const FColorSpace UE_ACES_SRGB_LIMITING_PRI = GetColorSpace(TONEMAPPER_GAMUT_sRGB_D65);
-
-	OutputReferredLinearsRGBColor *= UE_ACES_SDR_Y_MAX;
-	float3 aces = invOutputTransform ( OutputReferredLinearsRGBColor, UE_ACES_SDR_Y_MIN, UE_ACES_SDR_Y_MID, UE_ACES_SDR_Y_MAX, UE_ACES_SRGB_DISPLAY_PRI, UE_ACES_SRGB_LIMITING_PRI, UE_ACES_EOTF,
-									   UE_ACES_SURROUND, UE_ACES_STRETCH_BLACK, UE_ACES_D60_SIM, UE_ACES_LEGAL_RANGE);
-
-	float3 SceneReferredLinearsRGBColor = mul( AP0_2_sRGB, aces ) / Gain;
-
-	return SceneReferredLinearsRGBColor;
-}
+static const float UE_SCRGB_WHITE_NITS = 80.0; // ScRGB D65 white luminance (cd/m^2)
 
 struct FACESTonemapParams
 {
@@ -483,29 +435,6 @@ float3x3 OuputInverseGamutMappingMatrix( uint OutputGamut )
 	return GamutMappingMatrix;
 }
 
-float3 ST2084ToScRGB(float3 ColorAP1, uint OutputDevice, float MaxODTNits)
-{
-	// Nvidia HDR encoding - Remove PQ, convert to linear scRGB
-	const float3x3 AP1_2_sRGB = mul(XYZ_2_sRGB_MAT, mul(D60_2_D65_CAT, AP1_2_XYZ_MAT));
-	const float WhitePoint = 80.f;
-
-
-	float MinODTNits = UE_ACES_HDR_Y_MIN;
-	float3 OutColor = ST2084ToLinear(ColorAP1);
-		
-	OutColor = clamp(OutColor, MinODTNits, MaxODTNits);
-	OutColor.x = Y_2_linCV(OutColor.x, MaxODTNits, MinODTNits);
-	OutColor.y = Y_2_linCV(OutColor.y, MaxODTNits, MinODTNits);
-	OutColor.z = Y_2_linCV(OutColor.z, MaxODTNits, MinODTNits);
-
-	float scRGBScale = MaxODTNits / WhitePoint;
-	OutColor = mul(AP1_2_sRGB, OutColor) * scRGBScale;
-
-	return OutColor;
-}
-
-
-
 // Accurate for 1000K < Temp < 15000K
 // [Krystek 1985, "An algorithm to calculate correlated colour temperature"]
 float2 PlanckianLocusChromaticity( float Temp )
diff --git a/Engine/Source/Runtime/Renderer/Private/PostProcess/PostProcessCombineLUTs.cpp b/Engine/Source/Runtime/Renderer/Private/PostProcess/PostProcessCombineLUTs.cpp
--- a/Engine/Source/Runtime/Renderer/Private/PostProcess/PostProcessCombineLUTs.cpp
+++ b/Engine/Source/Runtime/Renderer/Private/PostProcess/PostProcessCombineLUTs.cpp
@@ -492,8 +492,13 @@ FRDGTextureRef AddCombineLUTPass(FRDGBuilder& GraphBuilder, const FViewInfo& Vie
 	const bool bUseComputePass = View.bUseComputePasses;
 
 	const bool bUseVolumeTextureLUT = PipelineVolumeTextureLUTSupportGuaranteedAtRuntime(View.GetShaderPlatform());
-
-	const bool bUseFloatOutput = ViewFamily.SceneCaptureSource == SCS_FinalColorHDR || ViewFamily.SceneCaptureSource == SCS_FinalToneCurveHDR;
+	
+	const EDisplayOutputFormat OutputDevice = ViewFamily.RenderTarget->GetDisplayOutputFormat();
+	// Note: We also enforce floats with HDR ScRGB (editor only) to preserve negative values.
+	const bool bUseFloatOutput = ViewFamily.SceneCaptureSource == SCS_FinalColorHDR ||
+		ViewFamily.SceneCaptureSource == SCS_FinalToneCurveHDR ||
+		OutputDevice == EDisplayOutputFormat::HDR_ACES_1000nit_ScRGB ||
+		OutputDevice == EDisplayOutputFormat::HDR_ACES_2000nit_ScRGB;
 
 	// Attempt to register the persistent view LUT texture.
 	const int32 LUTSize = CVarLUTSize->GetInt();
