Removed workaround related to unsafe transform matrices causing issues on some platforms since it is no longer necessary.

diff --git a/Engine/Shaders/Private/RayTracing/RayTracingInstanceBufferUtil.usf b/Engine/Shaders/Private/RayTracing/RayTracingInstanceBufferUtil.usf
--- a/Engine/Shaders/Private/RayTracing/RayTracingInstanceBufferUtil.usf
+++ b/Engine/Shaders/Private/RayTracing/RayTracingInstanceBufferUtil.usf
@@ -105,11 +105,6 @@ bool CullInstance(FInstanceSceneData InstanceSceneData, FPrimitiveSceneData Prim
 
 #endif // GPU_CULLING
 
-bool IsTransformElementUnsafe(float Element)
-{
-	return isnan(Element) || isinf(Element) || Element > 1e20f;
-}
-
 void ProcessInstance(FRayTracingInstanceDescriptor InputDesc)
 {
 	if (InputDesc.OutputDescriptorIndex >= NumInstances)
@@ -176,23 +171,6 @@ void ProcessInstance(FRayTracingInstanceDescriptor InputDesc)
 		BlasAddress = AccelerationStructureAddresses.Load2(InputDesc.AccelerationStructureIndex * 8);
 	}
 
-	// temporary workaround for crash on some platforms
-	if(IsTransformElementUnsafe(LocalToTranslatedWorld[0][0])
-		|| IsTransformElementUnsafe(LocalToTranslatedWorld[0][1])
-		|| IsTransformElementUnsafe(LocalToTranslatedWorld[0][2])
-		|| IsTransformElementUnsafe(LocalToTranslatedWorld[0][3])
-		|| IsTransformElementUnsafe(LocalToTranslatedWorld[1][0])
-		|| IsTransformElementUnsafe(LocalToTranslatedWorld[1][1])
-		|| IsTransformElementUnsafe(LocalToTranslatedWorld[1][2])
-		|| IsTransformElementUnsafe(LocalToTranslatedWorld[1][3])
-		|| IsTransformElementUnsafe(LocalToTranslatedWorld[2][0])
-		|| IsTransformElementUnsafe(LocalToTranslatedWorld[2][1])
-		|| IsTransformElementUnsafe(LocalToTranslatedWorld[2][2])
-		|| IsTransformElementUnsafe(LocalToTranslatedWorld[2][3]))
-	{
-		BlasAddress = 0;
-	}
-
 	OutPlatformInstanceDescriptors[InputDesc.OutputDescriptorIndex] = BuildPlatformRayTracingInstanceDesc(
 		InputDesc.InstanceMaskAndFlags & 0xFF,
 		InputDesc.InstanceId,
