Substrate - fix for subsurface profile disapearing from shading model node since MFP is not specified in this case.

#tests PC editor dx12
#rnx

diff --git a/Engine/Shaders/Private/Substrate/Substrate.ush b/Engine/Shaders/Private/Substrate/Substrate.ush
--- a/Engine/Shaders/Private/Substrate/Substrate.ush
+++ b/Engine/Shaders/Private/Substrate/Substrate.ush
@@ -1521,8 +1521,9 @@ FSubstrateData InternalGetSubstrateSlabBSDF(
 	{
 		SSSType = SSS_TYPE_DIFFUSION_PROFILE;
 	}
-	else if (!bMFPPlugged)
+	else if (SSSType != SSS_TYPE_DIFFUSION_PROFILE && !bMFPPlugged)
 	{
+		// If MFP is 0, then skip SSS, unless this is a Diffusion profile which has MFP in the profile itself.
 		SSSType = SSS_TYPE_NONE;
 	}
 
