Substrate - fix for subsurface profile disapearing from shading model node since MFP is not specified in this case.

#tests PC editor dx12
#rnx

diff --git a/Engine/Shaders/Private/Substrate/Substrate.ush b/Engine/Shaders/Private/Substrate/Substrate.ush
--- a/Engine/Shaders/Private/Substrate/Substrate.ush
+++ b/Engine/Shaders/Private/Substrate/Substrate.ush
@@ -1521,7 +1521,7 @@ FSubstrateData InternalGetSubstrateSlabBSDF(
 	{
 		SSSType = SSS_TYPE_DIFFUSION_PROFILE;
 	}
-	else if (SSSType != SSS_TYPE_DIFFUSION_PROFILE && !bMFPPlugged)
+	else if (!bMFPPlugged)
 	{
 		// If MFP is 0, then skip SSS, unless this is a Diffusion profile which has MFP in the profile itself.
 		SSSType = SSS_TYPE_NONE;
diff --git a/Engine/Shaders/Private/Substrate/SubstrateLegacyConversion.ush b/Engine/Shaders/Private/Substrate/SubstrateLegacyConversion.ush
--- a/Engine/Shaders/Private/Substrate/SubstrateLegacyConversion.ush
+++ b/Engine/Shaders/Private/Substrate/SubstrateLegacyConversion.ush
@@ -336,7 +336,7 @@ FSubstrateData SubstrateConvertLegacyMaterialDynamic(
 			LegacySlab.SSSMFPScale = Opacity;
 		}
 		LegacySlab.SSSPhaseAniso = 0.93f;
-		LegacySlab.SSSType = SSS_TYPE_DIFFUSION_PROFILE;
+		LegacySlab.SSSType = SSS_TYPE_DIFFUSION;	//  SSS_TYPE_DIFFUSION_PROFILE will be selected in InternalGetSubstrateSlabBSDF based on SSSProfileID and SSSMFPScale.
 		LegacySlab.Emissive = Emissive;
 
 		FinalWeight = 1.0f;
