Path Tracer: Optimize callable shader version of volumetric clouds

Instead of  invoking the callable shader once per raymarch step (which is very slow), invoke it only once per ray for the ray interval that overlaps with the cloud volume.

This means some code duplication of the core sampling loop, but allows the procedural logic to be executed only once.

diff --git a/Engine/Shaders/Private/PathTracing/Material/PathTracingMaterialSampling.ush b/Engine/Shaders/Private/PathTracing/Material/PathTracingMaterialSampling.ush
--- a/Engine/Shaders/Private/PathTracing/Material/PathTracingMaterialSampling.ush
+++ b/Engine/Shaders/Private/PathTracing/Material/PathTracingMaterialSampling.ush
@@ -121,19 +121,17 @@ float3 EstimateMaterialAlbedo(FPathTracingPayload Payload)
 
 #if !SIMPLIFIED_MATERIAL_SHADER
 
-#define ROUGHNESS_WEIGHT(r)		(saturate((r - max(MinRoughness, PathRoughness)) / (MaxRoughness - max(MinRoughness, PathRoughness))))
+#define ROUGHNESS_WEIGHT(r)		DenoiserRoughnessWeight(r, PathRoughness)
 
 void AccumulateAlbedoNormal(inout FPathTracingPayload Payload, float3 PathThroughput, float PathRoughness, inout float3 Albedo, inout float3 Normal)
 {
-	const float MaxRoughness = 0.15; // don't bother past this level
-	const float MinRoughness = 0.10; // capture anything below this level
-	if (PathRoughness >= MaxRoughness)
+	float AOVWeight = DenoiserAOVWeight(PathRoughness);
+	if (AOVWeight == 0)
 	{
-		// no point tracking past a certain roughness
+		// no contribution
 		return;
 	}
-	// figure out how much the current path should contribute to the denoiser aovs
-	PathThroughput *= 1 - saturate((PathRoughness - MinRoughness) / (MaxRoughness - MinRoughness));
+	PathThroughput *= AOVWeight;
 #if PATHTRACING_SUBSTRATE_PAYLOAD
 	PathThroughput *= Payload.WeightV * Payload.BSDFOpacity * lerp(1.0, Payload.TransmittanceN, Payload.CoverageAboveAlongN);
 #endif
@@ -261,15 +259,7 @@ void AccumulateAlbedoNormal(inout FPathTracingPayload Payload, float3 PathThroug
 
 void AccumulateAlbedo(float3 SigmaS, float3 PathThroughput, float PathRoughness, inout float3 Albedo)
 {
-	const float MaxRoughness = 0.15; // don't bother past this level
-	const float MinRoughness = 0.10; // capture anything below this level
-	if (PathRoughness >= MaxRoughness)
-	{
-		// no point tracking past a certain roughness
-		return;
-	}
-	// figure out how much the current path should contribute to the denoiser aovs
-	PathThroughput *= 1 - saturate((PathRoughness - MinRoughness) / (MaxRoughness - MinRoughness));
+	PathThroughput *= DenoiserAOVWeight(PathRoughness);
 	// Volume acts as diffuse scattering, so make it contribute fully (like DiffuseColor above)
 	Albedo += SigmaS * PathThroughput;
 }
diff --git a/Engine/Shaders/Private/PathTracing/PathTracingBuildCloud.usf b/Engine/Shaders/Private/PathTracing/PathTracingBuildCloud.usf
--- a/Engine/Shaders/Private/PathTracing/PathTracingBuildCloud.usf
+++ b/Engine/Shaders/Private/PathTracing/PathTracingBuildCloud.usf
@@ -25,6 +25,8 @@ void PathTracingBuildCloudCS(uint3 DispatchThreadId : SV_DispatchThreadID)
 		return;
 	}
 
+	ResolvedView = ResolveView();
+
 	RandomSequence RandSequence;
 	uint Seed = DispatchThreadId.x + CloudMapResolution.x * (DispatchThreadId.y + CloudMapResolution.y * DispatchThreadId.z); // unique seed per voxel
 	RandomSequence_Initialize(RandSequence, Seed, TemporalSeed);
diff --git a/Engine/Shaders/Private/PathTracing/PathTracingBuildCloudAccelerationMap.usf b/Engine/Shaders/Private/PathTracing/PathTracingBuildCloudAccelerationMap.usf
--- a/Engine/Shaders/Private/PathTracing/PathTracingBuildCloudAccelerationMap.usf
+++ b/Engine/Shaders/Private/PathTracing/PathTracingBuildCloudAccelerationMap.usf
@@ -24,6 +24,8 @@ void PathTracingBuildCloudAccelerationMapCS(uint2 DispatchThreadId : SV_Dispatch
 		return;
 	}
 
+	ResolvedView = ResolveView();
+
 	// get min and max altitude of the clouds
 	const float T = CloudLayerTopKm;
 	const float B = CloudLayerBotKm;
@@ -60,23 +62,23 @@ void PathTracingBuildCloudAccelerationMapCS(uint2 DispatchThreadId : SV_Dispatch
 		float Z = ClipPos.z - CloudClipRadiusKm; // remember height of the point (before rotation, relative to center)
 		float H = length(ClipPos);
 
-		float3 WorldPos = mul(ClipPos, GetCloudClipBasis()) + CloudClipCenterKm; // TangentToWorld
-
 		uint4 Data = 0;
+		float3 Extinction = 0;
 		if (UseCloudMap)
 		{
 			const float3 UVW = GetUVWFromCloudClipPos(ClipPos);
 			const int3 GridCoord = int3(floor(float3(CloudMapResolution) * UVW));
 			Data = CloudMap.Load(int4(GridCoord, 0));
+			Extinction = DecodeCloudResult(Data).Extinction;
 		}
 		else
 		{
+			float3 WorldPos = mul(ClipPos, GetCloudClipBasis()) + CloudClipCenterKm; // TangentToWorld
 			FDFVector3 AbsoluteWorldPosition = DFMultiply(WorldPos, SKY_UNIT_TO_CM);
 			// run the material and encode the result, which includes the effect that quantization will have during rendering
-			Data = EncodeCloudResult(SampleCloudMaterial(AbsoluteWorldPosition, H, B, T));
-
+			Extinction = SampleCloudMaterial(AbsoluteWorldPosition, H, B, T).Extinction;
 		}
-		const float3 Extinction = DecodeCloudResult(Data).Extinction;
+		
 		const float Density   = max3(Extinction.x, Extinction.y, Extinction.z);
 		const float DensityLo = min3(Extinction.x, Extinction.y, Extinction.z);
 
diff --git a/Engine/Shaders/Private/PathTracing/PathTracingCommon.ush b/Engine/Shaders/Private/PathTracing/PathTracingCommon.ush
--- a/Engine/Shaders/Private/PathTracing/PathTracingCommon.ush
+++ b/Engine/Shaders/Private/PathTracing/PathTracingCommon.ush
@@ -7,6 +7,8 @@
 #pragma once
 
 #include "Material/PathTracingMaterialCommon.ush"
+#include "Utilities/PathTracingRandomSequence.ush"
+#include "Utilities/PathTracingRIS.ush"
 #include "../RayTracing/RayTracingCommon.ush"
 #include "../OctahedralCommon.ush"
 #include "../DoubleFloat.ush"
@@ -45,6 +47,24 @@ HLSL_STATIC_ASSERT(SHADINGMODELID_SOLID_GLASS <= SHADINGMODELID_MASK, "");
 
 #define PATH_TRACING_ABSORPTION_SCALE  0.01        // inverse distance at which we hit the transmission color (Beer's law)
 
+// Utilities to help weight the contribution of the albedo/normal to the denoising buffers
+// We want cases like low roughness mirrors to "pass through" the signal
+static const float DenoiserMaxRoughness = 0.15f; // don't bother past this level
+static const float DenoiserMinRoughness = 0.10f; // capture anything below this level
+
+// Given the current path roughness, how much contribution should the current hit make to the denoiser AOVs?
+float DenoiserAOVWeight(float PathRoughness)
+{
+	return 1 - saturate((PathRoughness - DenoiserMinRoughness) / (DenoiserMaxRoughness - DenoiserMinRoughness));
+}
+// For a given surface of roughness R, how much should its own contribution be counted
+// This heuristic lets us skip the contribution a mirror's albedo would normally make because we will be seeing "through" it to the next hit
+float DenoiserRoughnessWeight(float Roughness, float PathRoughness)
+{
+	float MinRoughness = max(DenoiserMinRoughness, PathRoughness);
+	return saturate((Roughness - MinRoughness) / (DenoiserMaxRoughness - MinRoughness));
+}
+
 // This function is meant for the solid glass case so we can store the transmittance as a color in [0,1]
 float3 PathTracingGlassTransmittanceToExtinction(float3 TransmittanceColor)
 {
@@ -411,7 +431,9 @@ struct FPackedPathTracingPayload : FMinimalPayload
 		PackedData[8] = asuint(SceneZ);
 	}
 
-	// Volumetric callable shaders
+
+#if CLOUD_CALLABLE_SHADERS_PER_POINT
+	// Volumetric callable shaders -- used when calling per point
 	void SetVolumetricCallableShaderInput(float3 TranslatedWorldPos, float CloudHeightKm)
 	{
 		PackedData[0] = asuint(TranslatedWorldPos.x);
@@ -435,6 +457,155 @@ struct FPackedPathTracingPayload : FMinimalPayload
 
 	uint4 GetVolumetricCallableShaderOutput() { return uint4(PackedData[0], PackedData[1], PackedData[2], PackedData[3]); }
 
+#else
+
+	// Volumetric callable shaders
+	void SetVolumetricCallableShaderInput(float3 Origin, float3 Direction, float TMin, float TMax, RandomSequence RandSeq, uint Flags, float3 Throughput, float CloudDensity)
+	{
+		HitT = TMin;
+		PackedData[0]  = asuint(TMax);
+		PackedData[1]  = asuint(Origin.x);
+		PackedData[2]  = asuint(Origin.y);
+		PackedData[3]  = asuint(Origin.z);
+		PackedData[4]  = asuint(Direction.x);
+		PackedData[5]  = asuint(Direction.y);
+		PackedData[6]  = asuint(Direction.z);
+		PackedData[7]  = RandSeq.SampleIndex;
+		PackedData[8]  = RandSeq.SampleSeed;
+		PackedData[9]  = (Flags & 0xFFFF) | (f32tof16(CloudDensity) << 16);
+		PackedData[10] = asuint(Throughput.x);
+		PackedData[11] = asuint(Throughput.y);
+		PackedData[12] = asuint(Throughput.z);
+	}
+
+	void SetVolumetricCallableShaderRISContext(FRISContext SurfaceContext)
+	{
+		PackedData[13] = asuint(SurfaceContext.WeightSum);
+		PackedData[14] = asuint(SurfaceContext.RandSample);
+	}
+
+	FRayDesc GetVolumetricCallableShaderInputRay()
+	{
+		FRayDesc Ray;
+		Ray.Origin    = asfloat(uint3(PackedData[1], PackedData[2], PackedData[3]));
+		Ray.Direction = asfloat(uint3(PackedData[4], PackedData[5], PackedData[6]));
+		Ray.TMin      = HitT;
+		Ray.TMax      = asfloat(PackedData[0]);
+		return Ray;
+	}
+
+	RandomSequence GetVolumetricCallableShaderInputRandSequence()
+	{
+		RandomSequence RandSeq;
+		RandSeq.SampleIndex = PackedData[7];
+		RandSeq.SampleSeed  = PackedData[8];
+		return RandSeq;
+	}
+
+	uint GetVolumetricCallableShaderInputFlags()
+	{
+		return PackedData[9] & 0xFFFF;
+	}
+
+	float3 GetVolumetricCallableShaderInputThroughput()
+	{
+		return asfloat(uint3(PackedData[10], PackedData[11], PackedData[12]));
+	}
+
+	float GetVolumetricCallableShaderInputCloudDensity()
+	{
+		return f16tof32(PackedData[9] >> 16);
+	}
+
+	FRISContext GetVolumetricCallableShaderRISContext()
+	{
+		FRISContext Context;
+		Context.WeightSum  = asfloat(PackedData[13]);
+		Context.RandSample = asfloat(PackedData[14]);
+		return Context;
+	}
+
+	void SetVolumetricCallableShaderOutput(RandomSequence RandSeq, float3 Throughput, float3 Albedo)
+	{
+		PackedData[0] = RandSeq.SampleIndex;
+		PackedData[1] = RandSeq.SampleSeed;
+		PackedData[2] = PackFloat2ToUInt(Throughput.xy);
+		PackedData[3] = PackFloat2ToUInt(float2(Throughput.z, Albedo.x));
+		PackedData[4] = PackFloat2ToUInt(Albedo.yz);
+	}
+
+	void SetVolumetricCallableShaderOutputSample(float SampleT, float3 AlbedoRayleigh, float3 AlbedoHG1, float PhaseG, float3 AlbedoHG2, float3 DualPhaseData, float3 PayloadThroughput)
+	{
+		HitT = SampleT;
+		PackedData[5]  = PackFloat2ToUInt(PayloadThroughput.xy);
+		PackedData[6]  = PackFloat2ToUInt(float2(PayloadThroughput.z, PhaseG));
+		PackedData[7]  = PayloadEncodeLDRColor(AlbedoRayleigh);
+		PackedData[8]  = PayloadEncodeLDRColor(AlbedoHG1);
+		PackedData[9]  = PayloadEncodeLDRColor(AlbedoHG2);
+		PackedData[10] = PayloadEncodeLDRColor(DualPhaseData * float3(0.5, 0.5, 1.0) + float3(0.5, 0.5, 0.0));
+	}
+
+	void SetVolumetricCallableShaderOutputRadiance(float3 Radiance, float Alpha)
+	{
+		PackedData[11] = PayloadEncodeHDRColor(Radiance);
+		PackedData[12] = asuint(Alpha);
+	}
+
+	RandomSequence GetVolumetricCallableShaderOutputRandSeq()
+	{
+		RandomSequence RandSeq;
+		RandSeq.SampleIndex = PackedData[0];
+		RandSeq.SampleSeed  = PackedData[1];
+		return RandSeq;
+	}
+
+	float3 GetVolumetricCallableShaderOutputThroughput()
+	{
+		return float3(UnpackFloat2FromUInt(PackedData[2]), UnpackFloat2FromUInt(PackedData[3]).x);
+	}
+
+	float3 GetVolumetricCallableShaderOutputAlbedo()
+	{
+		return float3(UnpackFloat2FromUInt(PackedData[3]).y, UnpackFloat2FromUInt(PackedData[4]));
+	}
+
+	float3 GetVolumetricCallableShaderOutputPayloadThroughput()
+	{
+		return float3(UnpackFloat2FromUInt(PackedData[5]), UnpackFloat2FromUInt(PackedData[6]).x);
+	}
+
+	float3 GetVolumetricCallableShaderOutputRayleighWeight()
+	{
+		return PayloadDecodeLDRColor(PackedData[7]);
+	}
+
+	float4 GetVolumetricCallableShaderOutputHG()
+	{
+		return float4(PayloadDecodeLDRColor(PackedData[8]), UnpackFloat2FromUInt(PackedData[6]).y);
+	}
+
+	float3 GetVolumetricCallableShaderOutputDualHGWeight()
+	{
+		return PayloadDecodeLDRColor(PackedData[9]);
+	}
+
+	float3 GetVolumetricCallableShaderOutputDualHGPhaseData()
+	{
+		return PayloadDecodeLDRColor(PackedData[10]) * float3(2, 2, 1) - float3(1, 1, 0);
+	}
+
+	float3 GetVolumetricCallableShaderOutputRadiance()
+	{
+		return PayloadDecodeHDRColor(PackedData[11]);
+	}
+
+	float GetVolumetricCallableShaderOutputAlpha()
+	{
+		return asfloat(PackedData[12]);
+	}
+
+#endif 
+
 #endif // IS_PAYLOAD_ENABLED(RT_PAYLOAD_TYPE_GPULIGHTMASS)
 };
 
diff --git a/Engine/Shaders/Private/PathTracing/PathTracingCore.ush b/Engine/Shaders/Private/PathTracing/PathTracingCore.ush
--- a/Engine/Shaders/Private/PathTracing/PathTracingCore.ush
+++ b/Engine/Shaders/Private/PathTracing/PathTracingCore.ush
@@ -728,12 +728,64 @@ FPackedPathTracingPayload TraceTransparentRay(inout FPathState PathState, int Bo
 					(Interval.VolumeMask & (PATH_TRACER_VOLUME_ENABLE_CLOUDS | PATH_TRACER_VOLUME_ENABLE_HETEROGENEOUS_VOLUMES)) == 0;
 				// are there any holdouts on the current segment?
 				const bool bHaveHoldouts =
-					((Interval.VolumeMask << PATH_TRACER_MAX_VOLUMES) & VolumeFlags & (
-						PATH_TRACER_VOLUME_HOLDOUT_ATMOSPHERE |
-						PATH_TRACER_VOLUME_HOLDOUT_CLOUDS |
-						PATH_TRACER_VOLUME_HOLDOUT_FOG |
-						PATH_TRACER_VOLUME_HOLDOUT_HETEROGENEOUS_VOLUMES)) != 0;
+					(Interval.VolumeMask << PATH_TRACER_MAX_VOLUMES) & VolumeFlags & PATH_TRACER_VOLUME_HOLDOUT_MASK != 0;
 
+#if CLOUD_CALLABLE_SHADERS_PER_POINT == 0 && PATH_TRACER_USE_CLOUD_SHADER
+				if (Interval.VolumeMask & PATH_TRACER_VOLUME_ENABLE_CLOUDS)
+				{
+					// let callable shader handle cloud sampling
+					FPackedPathTracingPayload CloudPayload = (FPackedPathTracingPayload) 0;
+					const bool bNeedSample = !bLastBounce && PathState.ShouldAccumulateVolume();
+					CloudPayload.SetVolumetricCallableShaderInput(
+						PathState.Ray.Origin,
+						PathState.Ray.Direction,
+						Interval.VolumeTMin,
+						Interval.VolumeTMax,
+						PathState.RandSequence,
+						Interval.VolumeMask |
+						  (VolumeFlags & ~PATH_TRACER_VOLUME_ENABLE_MASK) |
+						  (bIsCameraRay ? PATH_TRACER_VOLUME_CALLABLE_FLAGS_IS_CAMERA_RAY : 0) |
+						  (bNeedSample  ? PATH_TRACER_VOLUME_CALLABLE_FLAGS_GET_SAMPLE    : 0),
+						PathThroughput,
+						Interval.CloudDensity
+					);
+					CloudPayload.SetVolumetricCallableShaderRISContext(HitSample);
+					CallShader(CloudCallableShaderId, CloudPayload);
+					PathState.RandSequence = CloudPayload.GetVolumetricCallableShaderOutputRandSeq();
+					PathThroughput = CloudPayload.GetVolumetricCallableShaderOutputThroughput();
+					PathState.SetAlbedoNormal(
+						PathState.GetAlbedo() + CloudPayload.GetVolumetricCallableShaderOutputAlbedo() * DenoiserAOVWeight(PathState.GetPathRoughness()),
+						PathState.GetNormal()
+					);
+					HitSample = CloudPayload.GetVolumetricCallableShaderRISContext();
+					AccumulateRadiance(PathState.Radiance, CloudPayload.GetVolumetricCallableShaderOutputRadiance(), bIsCameraRay);
+					if (bIsCameraRay)
+					{
+						PathState.Alpha += CloudPayload.GetVolumetricCallableShaderOutputAlpha();
+					}
+					if (CloudPayload.HitT > 0)
+					{
+						// If we registered a hit, transfer it to the currently sampled payload
+						PayloadThroughput = CloudPayload.GetVolumetricCallableShaderOutputPayloadThroughput();
+						FPathTracingPayload Result = (FPathTracingPayload)0; // clear all fields
+						Result.HitT = CloudPayload.HitT;
+						Result.TranslatedWorldPos = PathState.Ray.Origin + PathState.Ray.Direction * CloudPayload.HitT;
+						Result.ShadingModelID = SHADINGMODELID_MEDIUM;
+						Result.BSDFOpacity = 1.0;
+						Result.PrimitiveLightingChannelMask = 7;
+						Result.SetFrontFace();
+						Result.SetBaseColor(CloudPayload.GetVolumetricCallableShaderOutputRayleighWeight());
+						float4 HGData = CloudPayload.GetVolumetricCallableShaderOutputHG();
+						Result.SetHG(HGData.xyz, HGData.w);
+						Result.SetDualHG(
+							CloudPayload.GetVolumetricCallableShaderOutputDualHGWeight(),
+							CloudPayload.GetVolumetricCallableShaderOutputDualHGPhaseData()
+						);
+						Payload = PackPathTracingPayload(Result);
+					}
+				}
+				else
+#endif
 				// reasons to enter the ray marching loop:
 				//  - we can't use analytic transmittance (need ray marching for throughput or emission)
 				//  - we are a camera ray and some volumes are flagged as holdouts
diff --git a/Engine/Shaders/Private/PathTracing/PathTracingVolumetricCloudMaterialShader.usf b/Engine/Shaders/Private/PathTracing/PathTracingVolumetricCloudMaterialShader.usf
--- a/Engine/Shaders/Private/PathTracing/PathTracingVolumetricCloudMaterialShader.usf
+++ b/Engine/Shaders/Private/PathTracing/PathTracingVolumetricCloudMaterialShader.usf
@@ -5,19 +5,60 @@
 #include "/Engine/Private/Common.ush"
 #include "/Engine/Private/RayTracing/RayTracingCommon.ush"
 #include "/Engine/Private/PathTracing/PathTracingShaderUtils.ush"
+// Provide aliases for the parameters that now come through a shader buffer 
+
+#define FogDensity 		PathTracingCloudParameters.FogDensity
+#define FogHeight 		PathTracingCloudParameters.FogHeight
+#define FogFalloff		PathTracingCloudParameters.FogFalloff
+#define FogAlbedo		PathTracingCloudParameters.FogAlbedo
+#define FogPhaseG		PathTracingCloudParameters.FogPhaseG
+#define FogCenter		PathTracingCloudParameters.FogCenter
+#define FogMinZ			PathTracingCloudParameters.FogMinZ
+#define FogMaxZ			PathTracingCloudParameters.FogMaxZ
+#define FogRadius 		PathTracingCloudParameters.FogRadius
+#define FogFalloffClamp	PathTracingCloudParameters.FogFalloffClamp
+
+#define PlanetCenterTranslatedWorldHi PathTracingCloudParameters.PlanetCenterTranslatedWorldHi
+#define PlanetCenterTranslatedWorldLo PathTracingCloudParameters.PlanetCenterTranslatedWorldLo
+#define Atmosphere 		PathTracingCloudParameters
+
 
-#if CLOUD_LAYER_PIXEL_SHADER
 #include "./Volume/PathTracingCloudsMaterialCommon.ush"
-#endif
+#include "./Volume/PathTracingFog.ush"
+#include "./Volume/PathTracingAtmosphere.ush"
+#include "./Volume/PathTracingVolumeSampling.ush"
 
 #include "./PathTracingCommon.ush"
 
+uint MaxRaymarchSteps;
+
+
+FVolumeShadedResult CloudsGetDensity(float3 TranslatedWorldPos)
+{
+	FDFVector3 AbsoluteWorldPosition = DFFastSubtract(TranslatedWorldPos, ResolvedView.PreViewTranslation);
+
+	// Measure HeightKm as the distance to the planet center
+	const FDFVector3 OC = DFMultiply(DFSubtract(TranslatedWorldPos, MakeDFVector3(PlanetCenterTranslatedWorldHi, PlanetCenterTranslatedWorldLo)), CM_TO_SKY_UNIT);
+	const float H = DFLengthDemote(OC); // HeightKm
+	float B = PathTracingCloudParameters.CloudLayerBotKm;
+	float T = PathTracingCloudParameters.CloudLayerTopKm;
+
+	FSampleCloudMaterialResult CloudResult = SampleCloudMaterial(AbsoluteWorldPosition, H, B, T);
+	FVolumeShadedResult Result = (FVolumeShadedResult)0;
+	Result.SigmaT        = CloudResult.Extinction;
+	Result.SigmaSDualHG  = CloudResult.Albedo * Result.SigmaT;
+	Result.Emission      = CloudResult.EmissiveColor;
+	Result.DualPhaseData = float3(CloudResult.PhaseG1, CloudResult.PhaseG2, CloudResult.PhaseBlend);
+	return Result;
+}
 
 RAY_TRACING_ENTRY_CALLABLE(PathTracingVolumetricCloudMaterialShader, FPackedPathTracingPayload, PackedPayload)
 {
-#if CLOUD_LAYER_PIXEL_SHADER
 	ResolvedView = ResolveView();
 
+#if CLOUD_CALLABLE_SHADERS_PER_POINT
+	// Simple model, the Raygen contains all integrator logic and we just get called once per step to figure out the local cloud density
+
 	float3 TranslatedWorldPos = PackedPayload.GetVolumetricCallableShaderInputTranslatedWorldPos();
 	FDFVector3 AbsoluteWorldPosition = DFFastSubtract(TranslatedWorldPos, ResolvedView.PreViewTranslation);
 	float H = PackedPayload.GetVolumetricCallableShaderInputCloudHeightKm();
@@ -27,7 +68,163 @@ RAY_TRACING_ENTRY_CALLABLE(PathTracingVolumetricCloudMaterialShader, FPackedPath
 	FSampleCloudMaterialResult Result = SampleCloudMaterial(AbsoluteWorldPosition, H, B, T);
 	PackedPayload.SetVolumetricCallableShaderOutput(EncodeCloudResult(Result));
 #else
-	// for default shader - just flush outputs
-	PackedPayload.SetVolumetricCallableShaderOutput(0);
-#endif
+	// Advanced model - the ray marching happens inside here so that we only have to pay for the cost of callable shading _once_
+
+	RandomSequence RandSeq = PackedPayload.GetVolumetricCallableShaderInputRandSequence();
+	FRayDesc Ray = PackedPayload.GetVolumetricCallableShaderInputRay();
+	float3 PathThroughput = PackedPayload.GetVolumetricCallableShaderInputThroughput();
+	float CloudDensity = PackedPayload.GetVolumetricCallableShaderInputCloudDensity();
+	uint Flags = PackedPayload.GetVolumetricCallableShaderInputFlags();
+
+	if (Flags & PATH_TRACER_VOLUME_CALLABLE_FLAGS_TRANSMITTANCE) // TODO: evaluate if it would be better to make a second shader for this (so that we don't get two copies of the shader eval in inner loops)
+	{
+		// Transmittance loop for the cloud shader only
+		float TMin = Ray.TMin;
+		float TMax = Ray.TMax;
+		float SigmaBar = CloudDensity, InvSigmaBar = rcp(SigmaBar);
+		float3 Throughput = PathThroughput;
+		for (int Step = 0; Step < PathTracingCloudParameters.MaxRaymarchSteps; Step++)
+		{
+			/* take stochastic steps along the ray to estimate transmittance (null scattering) */
+			float3 ColorChannelPdf = Throughput;													
+			/* Sample the distance to the next interaction */										
+			float RandValue = RandomSequence_GenerateSample1D(RandSeq);						
+			float DeltaT = -log(1 - RandValue) * InvSigmaBar;
+			if (TMin + DeltaT < TMax)
+			{
+				TMin += DeltaT;
+				/* our ray marching step stayed inside the atmo and is still in front of the next hit */
+				float3 WorldPos = Ray.Origin + TMin * Ray.Direction;
+				/* clamp to make sure we never exceed the majorant (should not be the case, but need to avoid any possible numerical issues) */
+				float3 SigmaT = min(CloudsGetDensity(WorldPos).SigmaT, SigmaBar);
+				/* keep tracing through the volume */
+				Throughput *= 1.0 - SigmaT * InvSigmaBar;
+				/* early termination - fades off from 0.002 to 0.001 and then will be 0 */
+				Throughput *= saturate(1000.0 * max3(Throughput.x, Throughput.y, Throughput.z) - 1.0);
+				if (!any(Throughput > 0))
+				{
+					break;
+				}
+			}
+			else
+			{
+				/* exit the ray marching loop */
+				break;
+			}
+		}
+
+		PackedPayload = (FPackedPathTracingPayload)0;
+		PackedPayload.SetVolumetricCallableShaderOutput(RandSeq, Throughput, 0.0 /* Albedo not used for transmittance */);
+	}
+	else
+	{
+		FRISContext HitSample = PackedPayload.GetVolumetricCallableShaderRISContext();
+		// Loop through combined density for sampling
+		const bool bIsCameraRay = (Flags & PATH_TRACER_VOLUME_CALLABLE_FLAGS_IS_CAMERA_RAY) != 0;
+
+		// zero out the payload so we can write our output to it
+		PackedPayload = (FPackedPathTracingPayload) 0;
+		PackedPayload.HitT = -1.0; // invalid sample until we are sure we got one
+
+		float3 VolumeRadiance = 0;
+		float  VolumeAlpha = 0;
+		float3 VolumeAlbedo = 0;
+		// Limit number of steps to prevent timeouts
+		// TODO: This biases the result! Is there a better way to limit the number of steps?
+		for (int Step = 0; Step < PathTracingCloudParameters.MaxRaymarchSteps; Step++)
+		{
+			float3 SigmaBar = CloudDensity;
+			if (Flags & PATH_TRACER_VOLUME_ENABLE_ATMOSPHERE)
+			{
+				SigmaBar += AtmosphereGetDensityBounds(Ray.Origin, Ray.Direction, Ray.TMin, Ray.TMax).SigmaMax;
+			}
+			if (Flags & PATH_TRACER_VOLUME_ENABLE_FOG)
+			{
+				SigmaBar += FogGetDensityBounds(Ray.Origin, Ray.Direction, Ray.TMin, Ray.TMax).SigmaMax;
+			}
+
+			// Interval is only crossing atmo/clouds/fog
+			float RandValue = RandomSequence_GenerateSample1D(RandSeq);
+			float Distance = SampleSpectralTransmittance(RandValue, SigmaBar, PathThroughput);
+			if (Distance < 0.0)
+			{
+				// no energy left
+				break;
+			}
+
+			if (Ray.TMin + Distance < Ray.TMax)
+			{
+				float4 Evaluation = EvaluateSpectralTransmittanceHit(Distance, SigmaBar, PathThroughput);
+				PathThroughput *= Evaluation.xyz;
+				Ray.TMin += Distance;
+
+				// find out how much volume exists at the current point
+				float3 Ro = Ray.Origin;
+				float3 Rd = Ray.Direction;
+				float3 TranslatedWorldPos = Ro + Ray.TMin * Rd;
+				FVolumeShadedResult Result = (FVolumeShadedResult)0; 
+				MergeVolumeShadedResult(Result, CloudsGetDensity(TranslatedWorldPos), bIsCameraRay && (Flags & PATH_TRACER_VOLUME_HOLDOUT_CLOUDS) != 0);
+				if (Flags & PATH_TRACER_VOLUME_ENABLE_ATMOSPHERE)
+				{
+					MergeVolumeShadedResult(Result, AtmosphereGetDensity(TranslatedWorldPos), bIsCameraRay && (Flags & PATH_TRACER_VOLUME_HOLDOUT_ATMOSPHERE) != 0);
+				}
+				if (Flags & PATH_TRACER_VOLUME_ENABLE_FOG)
+				{
+					MergeVolumeShadedResult(Result, FogGetDensity(TranslatedWorldPos), bIsCameraRay && (Flags & PATH_TRACER_VOLUME_HOLDOUT_FOG) != 0);
+				}
+
+				// clamp to make sure we never exceed the majorant (should not be the case, but need to avoid any possible numerical issues)
+				float3 SigmaT = min(Result.SigmaT, SigmaBar);
+				float3 SigmaN = SigmaBar - SigmaT;
+
+				float3 SigmaH = min(Result.SigmaH, SigmaT);
+				VolumeAlpha += Luminance(PathThroughput * (SigmaT - SigmaH));
+
+				// If this ray wants to get candidate samples for scatter, process them here
+				if (Flags & PATH_TRACER_VOLUME_CALLABLE_FLAGS_GET_SAMPLE)
+				{
+					float3 SigmaS = min(Result.SigmaSRayleigh + Result.SigmaSHG + Result.SigmaSDualHG, SigmaT);
+
+					// accumulate a signal for the denoiser
+					float3 Contrib = PathThroughput * SigmaS;
+					VolumeAlbedo += Contrib;
+					float SelectionWeight = max3(Contrib.x, Contrib.y, Contrib.z);
+					if (HitSample.Accept(SelectionWeight))
+					{
+						// stash this hit for next time
+						float3 PayloadThroughput = Contrib / SelectionWeight;
+						const float3 RayleighWeight = select(SigmaS > 0.0, Result.SigmaSRayleigh / SigmaS, 0.0);
+						const float3 HGWeight       = select(SigmaS > 0.0, Result.SigmaSHG       / SigmaS, 0.0);
+						const float3 DualHGWeight   = select(SigmaS > 0.0, Result.SigmaSDualHG   / SigmaS, 0.0);
+						PackedPayload.SetVolumetricCallableShaderOutputSample(Ray.TMin, RayleighWeight, HGWeight, Result.PhaseG, DualHGWeight, Result.DualPhaseData, PayloadThroughput);
+					}
+				}
+
+				VolumeRadiance += Result.Emission * PathThroughput;
+				// keep tracing through the volume
+				PathThroughput *= SigmaN;
+				// early termination - fades off from 0.002 to 0.001 and then will be 0
+				PathThroughput *= saturate(10000.0 * max3(PathThroughput.x, PathThroughput.y, PathThroughput.z) - 1.0);
+				if (!any(PathThroughput > 0))
+				{
+					break;
+				}
+			}
+			else
+			{
+				// update the path throughput, knowing that we escaped the medium
+				// exit the ray marching loop
+				float4 Evaluation = EvaluateSpectralTransmittanceMiss(Ray.TMax - Ray.TMin, SigmaBar, PathThroughput);
+				PathThroughput *= Evaluation.xyz;
+				break;
+			}
+		}
+
+
+		// Write state back here
+		PackedPayload.SetVolumetricCallableShaderOutputRadiance(VolumeRadiance, VolumeAlpha);
+		PackedPayload.SetVolumetricCallableShaderOutput(RandSeq, PathThroughput, VolumeAlbedo);
+		PackedPayload.SetVolumetricCallableShaderRISContext(HitSample);
+	}
+#endif // CLOUD_CALLABLE_SHADERS_PER_POINT
 }
diff --git a/Engine/Shaders/Private/PathTracing/Volume/PathTracingAtmosphere.ush b/Engine/Shaders/Private/PathTracing/Volume/PathTracingAtmosphere.ush
--- a/Engine/Shaders/Private/PathTracing/Volume/PathTracingAtmosphere.ush
+++ b/Engine/Shaders/Private/PathTracing/Volume/PathTracingAtmosphere.ush
@@ -11,11 +11,12 @@
 
 // Shader parameters
 // NOTE: Atmosphere is a global struct (not directly defined here)
+#ifndef PlanetCenterTranslatedWorldHi
 float3 PlanetCenterTranslatedWorldHi;
 float3 PlanetCenterTranslatedWorldLo;
 Texture2D<float3> AtmosphereOpticalDepthLUT;
 SamplerState      AtmosphereOpticalDepthLUTSampler;
-
+#endif
 
 #if PLANET_ISECT_USES_DOUBLE_WORD_MATH
 
diff --git a/Engine/Shaders/Private/PathTracing/Volume/PathTracingClouds.ush b/Engine/Shaders/Private/PathTracing/Volume/PathTracingClouds.ush
--- a/Engine/Shaders/Private/PathTracing/Volume/PathTracingClouds.ush
+++ b/Engine/Shaders/Private/PathTracing/Volume/PathTracingClouds.ush
@@ -185,8 +185,6 @@ FVolumeShadedResult CloudsGetDensity(float3 TranslatedWorldPos)
 {
 	FVolumeShadedResult Result = (FVolumeShadedResult)0;
 
-	float3 WorldPos = DFMultiplyDemote(DFFastSubtract(TranslatedWorldPos, PrimaryView.PreViewTranslation), CM_TO_SKY_UNIT);
-
 	if ((VolumeFlags & PATH_TRACER_VOLUME_VIEW_CLOUD_MAP_ACCEL) != 0)
 	{
 		// debug preview path to visualize the cloud accel map
@@ -203,6 +201,7 @@ FVolumeShadedResult CloudsGetDensity(float3 TranslatedWorldPos)
 	}
 
 	uint4 Data = 0;
+#if CLOUD_CALLABLE_SHADERS_PER_POINT
 #if PATH_TRACER_USE_CLOUD_SHADER
 	{
 		// Reference path -- invoke the actual callable shader
@@ -223,11 +222,12 @@ FVolumeShadedResult CloudsGetDensity(float3 TranslatedWorldPos)
 		Data = CloudMap.Load(int4(GridCoord, 0));
 	}
 #endif
+#endif
+
 	const FSampleCloudMaterialResult CloudResult = DecodeCloudResult(Data);
 	Result.SigmaT        = CloudResult.Extinction;
 	Result.SigmaSDualHG  = CloudResult.Albedo * Result.SigmaT;
 	Result.Emission      = CloudResult.EmissiveColor;
 	Result.DualPhaseData = float3(CloudResult.PhaseG1, CloudResult.PhaseG2, CloudResult.PhaseBlend);
 	return Result;
-
 }
diff --git a/Engine/Shaders/Private/PathTracing/Volume/PathTracingFog.ush b/Engine/Shaders/Private/PathTracing/Volume/PathTracingFog.ush
--- a/Engine/Shaders/Private/PathTracing/Volume/PathTracingFog.ush
+++ b/Engine/Shaders/Private/PathTracing/Volume/PathTracingFog.ush
@@ -3,6 +3,7 @@
 #pragma once
 
 // FogParameters
+#ifndef FogDensity
 float2 FogDensity;
 float2 FogHeight;
 float2 FogFalloff;
@@ -13,6 +14,7 @@ float FogMinZ;
 float FogMaxZ;
 float FogRadius;
 float FogFalloffClamp;
+#endif
 
 FVolumeIntersection FogIntersect(float3 Origin, float3 Direction, float TMin, float TMax)
 {
diff --git a/Engine/Shaders/Private/PathTracing/Volume/PathTracingVolume.ush b/Engine/Shaders/Private/PathTracing/Volume/PathTracingVolume.ush
--- a/Engine/Shaders/Private/PathTracing/Volume/PathTracingVolume.ush
+++ b/Engine/Shaders/Private/PathTracing/Volume/PathTracingVolume.ush
@@ -47,16 +47,12 @@ FPackedPathTracingPayload VolumeGetBlockerHit(float3 Origin, float3 Direction, f
 // Given an input ray, figure out the bounds on density. The T interval may be smaller than the one returned by VolumeIntersect
 FVolumeDensityBounds VolumeGetDensityBounds(float3 Origin, float3 Direction, FVolumeIntersectionInterval Interval)
 {
-	FVolumeDensityBounds Result = CreateVolumeDensityBound(0, 0);
+	// CloudDensity will only be non-zero if the Interval overlaps with Clouds
+	FVolumeDensityBounds Result = CreateVolumeDensityBound(0, Interval.CloudDensity);
 	if (Interval.VolumeMask & PATH_TRACER_VOLUME_ENABLE_ATMOSPHERE)
 	{
 		MergeVolumeDensityBounds(Result, AtmosphereGetDensityBounds(Origin, Direction, Interval.VolumeTMin, Interval.VolumeTMax));
 	}
-	if (Interval.VolumeMask & PATH_TRACER_VOLUME_ENABLE_CLOUDS)
-	{
-		// We tracked the max density during intersection, so it is directly available here
-		MergeVolumeDensityBounds(Result, CreateVolumeDensityBound(0, Interval.CloudDensity));
-	}
 	if (Interval.VolumeMask & PATH_TRACER_VOLUME_ENABLE_FOG)
 	{
 		MergeVolumeDensityBounds(Result, FogGetDensityBounds(Origin, Direction, Interval.VolumeTMin, Interval.VolumeTMax));
@@ -176,7 +172,15 @@ float3 VolumeGetTransmittance(float3 StartThroughput, float3 Origin, float3 Dire
 		float TMin = Interval.VolumeTMin;
 		float TMax = Interval.VolumeTMax;
 		float SigmaBar = Interval.CloudDensity;
+#if CLOUD_CALLABLE_SHADERS_PER_POINT == 0 && PATH_TRACER_USE_CLOUD_SHADER
+		FPackedPathTracingPayload CloudPayload = (FPackedPathTracingPayload) 0;
+		CloudPayload.SetVolumetricCallableShaderInput(Origin, Direction, TMin, TMax, RandSequence, PATH_TRACER_VOLUME_CALLABLE_FLAGS_TRANSMITTANCE, Throughput, SigmaBar);
+		CallShader(CloudCallableShaderId, CloudPayload);
+		RandSequence = CloudPayload.GetVolumetricCallableShaderOutputRandSeq();
+		Throughput = CloudPayload.GetVolumetricCallableShaderOutputThroughput();
+#else
 		PATH_TRACER_REFERENCE_TRANSMITTANCE_LOOP_CONSTANT_BOUNDS(CloudsGetDensity);
+#endif
 	}
 	if (Interval.VolumeMask & PATH_TRACER_VOLUME_ENABLE_FOG)
 	{
@@ -222,7 +226,15 @@ float3 VolumeGetTransmittance(float3 StartThroughput, float3 Origin, float3 Dire
 		float TMin = IntersectionList.VolumeTMin[VOLUMEID_CLOUDS];
 		float TMax = IntersectionList.VolumeTMax[VOLUMEID_CLOUDS];
 		float SigmaBar = IntersectionList.CloudDensity;
+#if CLOUD_CALLABLE_SHADERS_PER_POINT == 0 && PATH_TRACER_USE_CLOUD_SHADER
+		FPackedPathTracingPayload CloudPayload = (FPackedPathTracingPayload) 0;
+		CloudPayload.SetVolumetricCallableShaderInput(Origin, Direction, TMin, TMax, RandSequence, PATH_TRACER_VOLUME_CALLABLE_FLAGS_TRANSMITTANCE, Throughput, SigmaBar);
+		CallShader(CloudCallableShaderId, CloudPayload);
+		RandSequence = CloudPayload.GetVolumetricCallableShaderOutputRandSeq();
+		Throughput = CloudPayload.GetVolumetricCallableShaderOutputThroughput();
+#else
 		PATH_TRACER_REFERENCE_TRANSMITTANCE_LOOP_CONSTANT_BOUNDS(CloudsGetDensity);
+#endif
 	}
 	if (IntersectionList.VolumeMask & PATH_TRACER_VOLUME_ENABLE_FOG)
 	{
diff --git a/Engine/Shaders/Private/PathTracing/Volume/PathTracingVolumeCommon.ush b/Engine/Shaders/Private/PathTracing/Volume/PathTracingVolumeCommon.ush
--- a/Engine/Shaders/Private/PathTracing/Volume/PathTracingVolumeCommon.ush
+++ b/Engine/Shaders/Private/PathTracing/Volume/PathTracingVolumeCommon.ush
@@ -126,7 +126,7 @@ struct FVolumeIntersectionList
 				else if (Result.VolumeTMin < TMin)
 				{
 					// what we have so far _starts_ before the current segment
-					// so either the two segments are disjoint (Result.y is closer)
+					// so either the two segments are disjoint (Result.VolumeTMax is closer)
 					// or the two overlap, and we need to stop at the next transition line
 					Result.VolumeTMax = min(Result.VolumeTMax, TMin);
 				}
@@ -140,7 +140,6 @@ struct FVolumeIntersectionList
 		}
 		// now that we established the distance, figure out which IDs overlap with the resolved segment
 		Result.VolumeMask = 0;
-		Result.CloudDensity = CloudDensity;
 		for (int Index2 = 0; Index2 < PATH_TRACER_MAX_VOLUMES; Index2++)
 		{
 			uint Mask = PATH_TRACER_VOLUME_ENABLE_BIT << Index2;
@@ -154,6 +153,8 @@ struct FVolumeIntersectionList
 				Result.VolumeMask |= Mask;
 			}
 		}
+		// Transfer CloudDensity to Interval if this interval includes clouds
+		Result.CloudDensity = Result.VolumeMask & PATH_TRACER_VOLUME_ENABLE_BIT ? CloudDensity : 0; 
 		return Result;
 	}
 
diff --git a/Engine/Shaders/Shared/PathTracingDefinitions.h b/Engine/Shaders/Shared/PathTracingDefinitions.h
--- a/Engine/Shaders/Shared/PathTracingDefinitions.h
+++ b/Engine/Shaders/Shared/PathTracingDefinitions.h
@@ -58,17 +58,27 @@
 #define PATH_TRACER_VOLUME_ENABLE_CLOUDS                    (PATH_TRACER_VOLUME_ENABLE_BIT << VOLUMEID_CLOUDS)
 #define PATH_TRACER_VOLUME_ENABLE_FOG               		(PATH_TRACER_VOLUME_ENABLE_BIT << VOLUMEID_FOG)
 #define PATH_TRACER_VOLUME_ENABLE_HETEROGENEOUS_VOLUMES 	(PATH_TRACER_VOLUME_ENABLE_BIT << VOLUMEID_HETEROGENEOUS_VOLUMES)
+#define PATH_TRACER_VOLUME_ENABLE_MASK                      ((1u << PATH_TRACER_MAX_VOLUMES) - 1)
 
 #define PATH_TRACER_VOLUME_HOLDOUT_BIT						(PATH_TRACER_VOLUME_ENABLE_BIT << PATH_TRACER_MAX_VOLUMES)
 #define PATH_TRACER_VOLUME_HOLDOUT_ATMOSPHERE  				(PATH_TRACER_VOLUME_HOLDOUT_BIT << VOLUMEID_ATMOSPHERE)
 #define PATH_TRACER_VOLUME_HOLDOUT_CLOUDS   				(PATH_TRACER_VOLUME_HOLDOUT_BIT << VOLUMEID_CLOUDS)
 #define PATH_TRACER_VOLUME_HOLDOUT_FOG               		(PATH_TRACER_VOLUME_HOLDOUT_BIT << VOLUMEID_FOG)
 #define PATH_TRACER_VOLUME_HOLDOUT_HETEROGENEOUS_VOLUMES	(PATH_TRACER_VOLUME_HOLDOUT_BIT << VOLUMEID_HETEROGENEOUS_VOLUMES)
+#define PATH_TRACER_VOLUME_HOLDOUT_MASK                     (PATH_TRACER_VOLUME_ENABLE_MASK << PATH_TRACER_MAX_VOLUMES)
 
 #define PATH_TRACER_VOLUME_USE_ANALYTIC_TRANSMITTANCE       (PATH_TRACER_VOLUME_HOLDOUT_BIT << PATH_TRACER_MAX_VOLUMES)
 #define PATH_TRACER_VOLUME_USE_CLOUD_MAP                    (PATH_TRACER_VOLUME_USE_ANALYTIC_TRANSMITTANCE << 1)
 #define PATH_TRACER_VOLUME_VIEW_CLOUD_MAP_ACCEL             (PATH_TRACER_VOLUME_USE_ANALYTIC_TRANSMITTANCE << 2)
 
+// These flags are reserved for passing information to the cloud callable shader
+#define PATH_TRACER_VOLUME_CALLABLE_FLAGS_TRANSMITTANCE     (PATH_TRACER_VOLUME_VIEW_CLOUD_MAP_ACCEL << 1) // only compute transmittance for the cloud portion
+#define PATH_TRACER_VOLUME_CALLABLE_FLAGS_GET_SAMPLE        (PATH_TRACER_VOLUME_VIEW_CLOUD_MAP_ACCEL << 2) // perform RIS sampling for a scatter location (else do transmittance/emission/alpha calc only)
+#define PATH_TRACER_VOLUME_CALLABLE_FLAGS_IS_CAMERA_RAY     (PATH_TRACER_VOLUME_VIEW_CLOUD_MAP_ACCEL << 3) // this is a camera ray (so apply holdouts if enabled)
+
+
+#define CLOUD_CALLABLE_SHADERS_PER_POINT 0 // TODO: remove the per-point codepath
+
 
 // Constants related to debugging
 #define PATH_TRACER_DEBUG_VIZ_RADIANCE					0 
diff --git a/Engine/Source/Runtime/Renderer/Private/PathTracing.cpp b/Engine/Source/Runtime/Renderer/Private/PathTracing.cpp
--- a/Engine/Source/Runtime/Renderer/Private/PathTracing.cpp
+++ b/Engine/Source/Runtime/Renderer/Private/PathTracing.cpp
@@ -1114,6 +1114,32 @@ static FPathTracingFogParameters PrepareFogParameters(const FViewInfo& View, con
 	return Parameters;
 }
 
+// Take a double as input as modify it to a FDFScalar that represents the same value.
+// TODO: Do we still need this? Could unify with FDFScalar C++ type
+static void SplitDouble(double x, float* hi, float* lo)
+{
+	const double SPLIT = 134217729.0; // 2^27+1
+	double temp = SPLIT * x;
+	*hi = static_cast<float>(temp - (temp - x));
+	*lo = static_cast<float>(x - *hi);
+}
+
+static void PreparePlanetCenter(const FViewInfo& View, const FSkyAtmosphereRenderSceneInfo* SkyAtmosphereSceneInfo, FVector3f* PlanetCenterTranslatedWorldHi, FVector3f* PlanetCenterTranslatedWorldLo)
+{
+	if (SkyAtmosphereSceneInfo != nullptr)
+	{
+		FVector PlanetCenterTranslatedWorld = SkyAtmosphereSceneInfo->GetSkyAtmosphereSceneProxy().GetAtmosphereSetup().PlanetCenterKm * double(FAtmosphereSetup::SkyUnitToCm) + View.ViewMatrices.GetPreViewTranslation();
+		SplitDouble(PlanetCenterTranslatedWorld.X, &PlanetCenterTranslatedWorldHi->X, &PlanetCenterTranslatedWorldLo->X);
+		SplitDouble(PlanetCenterTranslatedWorld.Y, &PlanetCenterTranslatedWorldHi->Y, &PlanetCenterTranslatedWorldLo->Y);
+		SplitDouble(PlanetCenterTranslatedWorld.Z, &PlanetCenterTranslatedWorldHi->Z, &PlanetCenterTranslatedWorldLo->Z);
+	}
+	else
+	{
+		*PlanetCenterTranslatedWorldHi = FVector3f(0);
+		*PlanetCenterTranslatedWorldLo = FVector3f(0);
+	}
+}
+
 
 BEGIN_SHADER_PARAMETER_STRUCT(FPathTracingCloudParameters, )
 	// coordinate frame for the cloud acceleration map
@@ -1136,6 +1162,11 @@ END_SHADER_PARAMETER_STRUCT()
 
 BEGIN_GLOBAL_SHADER_PARAMETER_STRUCT(FPathTracingCloudParameterGlobals,)
 	SHADER_PARAMETER_STRUCT_INCLUDE(FPathTracingCloudParameters, CloudParameters)
+	SHADER_PARAMETER_STRUCT_INCLUDE(FPathTracingFogParameters, FogParameters)
+	SHADER_PARAMETER_STRUCT_INCLUDE(FAtmosphereUniformShaderParameters, AtmosphereParameters)
+	SHADER_PARAMETER(FVector3f, PlanetCenterTranslatedWorldHi)
+	SHADER_PARAMETER(FVector3f, PlanetCenterTranslatedWorldLo)
+	SHADER_PARAMETER(uint32, MaxRaymarchSteps)
 END_GLOBAL_SHADER_PARAMETER_STRUCT()
 
 IMPLEMENT_GLOBAL_SHADER_PARAMETER_STRUCT(FPathTracingCloudParameterGlobals, "PathTracingCloudParameters");
@@ -1745,6 +1776,18 @@ void PreparePathTracingCloudMaterial(FScene* Scene, TArrayView<FViewInfo> Views)
 
 		FPathTracingCloudParameterGlobals Params = {};
 		Params.CloudParameters = PrepareCloudParameters(Scene, View, nullptr);
+		if (!Scene->ExponentialFogs.IsEmpty())
+		{
+			Params.FogParameters = PrepareFogParameters(View, Scene->ExponentialFogs[0]);
+		}
+		if (Scene->GetSkyAtmosphereSceneInfo() != nullptr)
+		{
+			Params.AtmosphereParameters = *Scene->GetSkyAtmosphereSceneInfo()->GetAtmosphereShaderParameters();
+		}
+
+		PreparePlanetCenter(View, Scene->GetSkyAtmosphereSceneInfo(), &Params.PlanetCenterTranslatedWorldHi, &Params.PlanetCenterTranslatedWorldLo);
+		Params.MaxRaymarchSteps = CVarPathTracingMaxRaymarchSteps.GetValueOnRenderThread();
+
 		TUniformBufferRef<FPathTracingCloudParameterGlobals> CloudParametersUB = CreateUniformBufferImmediate(
 			Params,
 			EUniformBufferUsage::UniformBuffer_SingleFrame
@@ -3151,14 +3194,6 @@ uint32 FSceneViewState::GetPathTracingSampleCount() const {
 	return State ? State->LastConfig.PathTracingData.MaxSamples : 0;
 }
 
-static void SplitDouble(double x, float* hi, float* lo)
-{
-	const double SPLIT = 134217729.0; // 2^27+1
-	double temp = SPLIT * x;
-	*hi = static_cast<float>(temp - (temp - x));
-	*lo = static_cast<float>(x - *hi);
-}
-
 #if WITH_MGPU
 BEGIN_SHADER_PARAMETER_STRUCT(FMGPUTransferParameters, )
 	RDG_TEXTURE_ACCESS(InputTexture, ERHIAccess::CopySrc)
@@ -3917,10 +3952,7 @@ void FDeferredShadingSceneRenderer::RenderPathTracing(
 							else if ((Config.PathTracingData.VolumeFlags & PATH_TRACER_VOLUME_ENABLE_ATMOSPHERE) != 0)
 							{
 								PassParameters->Atmosphere = Scene->GetSkyAtmosphereSceneInfo()->GetAtmosphereUniformBuffer();
-								FVector PlanetCenterTranslatedWorld = Scene->GetSkyAtmosphereSceneInfo()->GetSkyAtmosphereSceneProxy().GetAtmosphereSetup().PlanetCenterKm * double(FAtmosphereSetup::SkyUnitToCm) + View.ViewMatrices.GetPreViewTranslation();
-								SplitDouble(PlanetCenterTranslatedWorld.X, &PassParameters->PlanetCenterTranslatedWorldHi.X, &PassParameters->PlanetCenterTranslatedWorldLo.X);
-								SplitDouble(PlanetCenterTranslatedWorld.Y, &PassParameters->PlanetCenterTranslatedWorldHi.Y, &PassParameters->PlanetCenterTranslatedWorldLo.Y);
-								SplitDouble(PlanetCenterTranslatedWorld.Z, &PassParameters->PlanetCenterTranslatedWorldHi.Z, &PassParameters->PlanetCenterTranslatedWorldLo.Z);
+								PreparePlanetCenter(View, Scene->GetSkyAtmosphereSceneInfo(), &PassParameters->PlanetCenterTranslatedWorldHi, &PassParameters->PlanetCenterTranslatedWorldLo);
 							}
 							else
 							{
