Contact Shadows: Fixed contact shadow screen traces marching across view boundaries, causing issues in splitscreen.

diff --git a/Engine/Shaders/Private/ScreenSpaceShadowRayCast.ush b/Engine/Shaders/Private/ScreenSpaceShadowRayCast.ush
--- a/Engine/Shaders/Private/ScreenSpaceShadowRayCast.ush
+++ b/Engine/Shaders/Private/ScreenSpaceShadowRayCast.ush
@@ -11,53 +11,50 @@ float CastScreenSpaceShadowRay(
 	float Dither, float CompareToleranceScale, bool bHairNoShadowLight,
 	out float2 HitUV)
 {
-	float4 RayStartClip = mul(float4(RayOriginTranslatedWorld, 1), View.TranslatedWorldToClip);
-	float4 RayDirClip = mul(float4(RayDirection * RayLength, 0), View.TranslatedWorldToClip);
-	float4 RayEndClip = RayStartClip + RayDirClip;
+	const float4 RayStartClip = mul(float4(RayOriginTranslatedWorld, 1), View.TranslatedWorldToClip);
+	const float4 RayDirClip = mul(float4(RayDirection * RayLength, 0), View.TranslatedWorldToClip);
+	const float4 RayEndClip = RayStartClip + RayDirClip;
 
-	float3 RayStartScreen = RayStartClip.xyz / RayStartClip.w;
-	float3 RayEndScreen = RayEndClip.xyz / RayEndClip.w;
+	const float3 RayStart = RayStartClip.xyz / RayStartClip.w;
+	const float3 RayEnd = RayEndClip.xyz / RayEndClip.w;
 
-	float3 RayStepScreen = RayEndScreen - RayStartScreen;
+	const float3 RayStep = RayEnd - RayStart;
 
-	float3 RayStartUVz = float3(RayStartScreen.xy * View.ScreenPositionScaleBias.xy + View.ScreenPositionScaleBias.wz, RayStartScreen.z);
-	float3 RayStepUVz = float3(RayStepScreen.xy * View.ScreenPositionScaleBias.xy, RayStepScreen.z);
-
-	float4 RayDepthClip = RayStartClip + mul(float4(0, 0, RayLength, 0), View.ViewToClip);
-	float3 RayDepthScreen = RayDepthClip.xyz / RayDepthClip.w;
+	const float4 RayDepthClip = RayStartClip + mul(float4(0, 0, RayLength, 0), View.ViewToClip);
+	const float3 RayDepth = RayDepthClip.xyz / RayDepthClip.w;
 
 	const float StepOffset = Dither - 0.5f;
 	const float Step = 1.0 / NumSteps;
 
-	const float CompareTolerance = abs(RayDepthScreen.z - RayStartScreen.z) * Step * CompareToleranceScale;
+	const float CompareTolerance = abs(RayDepth.z - RayStart.z) * Step * CompareToleranceScale;
 
 	float SampleTime = StepOffset * Step + Step;
 
-	float FirstHitTime = -1.0;
-
-	const float StartDepth = LookupDeviceZ(RayStartUVz.xy);
+	const float StartDepth = LookupDeviceZ(RayStart.xy * View.ScreenPositionScaleBias.xy + View.ScreenPositionScaleBias.wz);
 
 	UNROLL
 	for (int i = 0; i < NumSteps; i++)
 	{
-		float3 SampleUVz = RayStartUVz + RayStepUVz * SampleTime;
-		float SampleDepth = LookupDeviceZ(SampleUVz.xy);
+		// SamplePos is in NDC space of the current view
+		const float3 SamplePos = RayStart + RayStep * SampleTime;
+		const float2 SampleUV = SamplePos.xy * View.ScreenPositionScaleBias.xy + View.ScreenPositionScaleBias.wz;
+		const float SampleDepth = LookupDeviceZ(SampleUV);
 
 		// Avoid self-intersection with the start pixel (exact comparison due to point sampling depth buffer)
 		// Exception is made for hair for occluding transmitted light with non-shadow casting light
 		if (SampleDepth != StartDepth || bHairNoShadowLight)
 		{
-			float DepthDiff = SampleUVz.z - SampleDepth;
-			bool Hit = abs(DepthDiff + CompareTolerance) < CompareTolerance;
+			const float DepthDiff = SamplePos.z - SampleDepth;
+			const bool bHit = abs(DepthDiff + CompareTolerance) < CompareTolerance;
 
-			if (Hit)
+			if (bHit)
 			{
-				HitUV = SampleUVz.xy;
+				HitUV = SampleUV;
 
-				// Off screen masking
-				bool bValidUV = all(and (0.0 < SampleUVz.xy, SampleUVz.xy < 1.0));
+				// Off screen masking, check NDC position against NDC boundary [-1,1]
+				bool bValidPos = all(and(-1.0 < SamplePos.xy, SamplePos.xy < 1.0));
 
-				return bValidUV ? (RayLength * SampleTime) : -1.0;
+				return bValidPos ? (RayLength * SampleTime) : -1.0;
 			}
 		}
 
