Fix cook warnings in cloud shaders.

#rnx

diff --git a/Engine/Shaders/Private/PathTracing/PathTracingCommon.ush b/Engine/Shaders/Private/PathTracing/PathTracingCommon.ush
--- a/Engine/Shaders/Private/PathTracing/PathTracingCommon.ush
+++ b/Engine/Shaders/Private/PathTracing/PathTracingCommon.ush
@@ -484,6 +484,7 @@ struct FPackedPathTracingPayload : FMinimalPayload
 		PackedData[14] = asuint(SurfaceContext.RandSample);
 	}
 
+#if !COMPUTESHADER  // FRayDesc only defined for non compute shaders or inline ray tracing
 	FRayDesc GetVolumetricCallableShaderInputRay()
 	{
 		FRayDesc Ray;
@@ -493,6 +494,7 @@ struct FPackedPathTracingPayload : FMinimalPayload
 		Ray.TMax      = asfloat(PackedData[0]);
 		return Ray;
 	}
+#endif
 
 	RandomSequence GetVolumetricCallableShaderInputRandSequence()
 	{
