Replaced ComputeHistoryVolumeUV_DEPRECATED(...) with ComputeHistoryVolumeUVFromTranslatedPos(...)

diff --git a/Engine/Shaders/Private/HeightFogCommon.ush b/Engine/Shaders/Private/HeightFogCommon.ush
--- a/Engine/Shaders/Private/HeightFogCommon.ush
+++ b/Engine/Shaders/Private/HeightFogCommon.ush
@@ -523,9 +523,9 @@ float3 ComputeHistoryVolumeUVFromNDC(float4 NDCPosition)
 	float3 VolumeUV = float3(NDCPosition.xy * float2(.5f, -.5f) + .5f, ComputeNormalizedZSliceFromDepth(NDCPosition.w));
 	return min(VolumeUV * float3(View.VolumetricFogViewGridUVToPrevViewRectUV.xy * View.VolumetricFogPrevViewGridRectUVToResourceUV.xy, 1.0), float3(View.VolumetricFogPrevUVMaxForTemporalBlend, 1.0));
 }
-float3 ComputeHistoryVolumeUV_DEPRECATED(float3 WorldPosition, float4x4 WorldToClip)
+float3 ComputeHistoryVolumeUVFromTranslatedPos(float3 TranslatedWorldPosition, float4x4 TranslatedWorldToClip)
 {
-	float4 NDCPosition = mul(float4(WorldPosition, 1), WorldToClip);
+	float4 NDCPosition = mul(float4(TranslatedWorldPosition, 1), TranslatedWorldToClip);
 	return ComputeHistoryVolumeUVFromNDC(NDCPosition);
 }
 float3 ComputeHistoryVolumeUV(FDFVector3 WorldPosition, FDFInverseMatrix WorldToClip)
diff --git a/Engine/Shaders/Private/VolumetricFog.usf b/Engine/Shaders/Private/VolumetricFog.usf
--- a/Engine/Shaders/Private/VolumetricFog.usf
+++ b/Engine/Shaders/Private/VolumetricFog.usf
@@ -380,7 +380,7 @@ float4 InjectShadowedLocalLightCommon(uint3 GridCoordinate)
 		uint NumSuperSamples = 1;
 
 #if USE_TEMPORAL_REPROJECTION
-		float3 HistoryUV = ComputeHistoryVolumeUV_DEPRECATED(ComputeCellTranslatedWorldPosition(GridCoordinate, .5f), UnjitteredPrevTranslatedWorldToClip);
+		float3 HistoryUV = ComputeHistoryVolumeUVFromTranslatedPos(ComputeCellTranslatedWorldPosition(GridCoordinate, .5f), UnjitteredPrevTranslatedWorldToClip);
 		float HistoryAlpha = HistoryWeight;
 
 		FLATTEN
@@ -774,7 +774,7 @@ void LightScatteringCS(
 	uint NumSuperSamples = 1;
 
 #if USE_TEMPORAL_REPROJECTION
-	float3 HistoryUV = ComputeHistoryVolumeUV_DEPRECATED(ComputeCellTranslatedWorldPosition(GridCoordinate, .5f), UnjitteredPrevTranslatedWorldToClip);
+	float3 HistoryUV = ComputeHistoryVolumeUVFromTranslatedPos(ComputeCellTranslatedWorldPosition(GridCoordinate, .5f), UnjitteredPrevTranslatedWorldToClip);
 	bool bHasValidHistory = true;
 #endif
 
@@ -810,7 +810,7 @@ void LightScatteringCS(
 #if USE_TEMPORAL_REPROJECTION
 	float HistoryAlpha = HistoryWeight;
 
-	// We need to test with View.VolumetricFogPrevUVMaxForTemporalBlend for HistoryUV because that is what we clamp with in ComputeHistoryVolumeUV_DEPRECATED.
+	// We need to test with View.VolumetricFogPrevUVMaxForTemporalBlend for HistoryUV because that is what we clamp with in ComputeHistoryVolumeUVFromTranslatedPos.
 	FLATTEN
 	if (any(HistoryUV < 0) || any(HistoryUV >= float3(View.VolumetricFogPrevUVMaxForTemporalBlend, 1.0f)) || !bHasValidHistory)
 	{
