Path Tracer: Fix missing background light visibility

A recent CL changed the default camera ray's TMax value to FLOAT_MAX instead of POSITIVE_INFINITY. I've updated the path tracer to use this convention throughout for consistency.

I also added an explicit FLOAT_MAX define in Common.ush for clarity as well as a RAY_DEFAULT_T_MAX meant for infinite rays (mostly used by the path tracer).

diff --git a/Engine/Plugins/Experimental/GPULightmass/Shaders/Private/LightmapPathTracing.usf b/Engine/Plugins/Experimental/GPULightmass/Shaders/Private/LightmapPathTracing.usf
--- a/Engine/Plugins/Experimental/GPULightmass/Shaders/Private/LightmapPathTracing.usf
+++ b/Engine/Plugins/Experimental/GPULightmass/Shaders/Private/LightmapPathTracing.usf
@@ -697,7 +697,7 @@ void PathTracingKernel(
 		Ray.Origin = Payload.TranslatedWorldPos;
 		Ray.Direction = MaterialSample.Direction;
 		Ray.TMin = 0;
-		Ray.TMax = POSITIVE_INFINITY;
+		Ray.TMax = RAY_DEFAULT_T_MAX;
 		ApplyRayBias(Ray.Origin, Payload.HitT, MaterialSample.PositionBiasSign * Payload.WorldGeoNormal);
 
 		if (ApproximateCaustics)
diff --git a/Engine/Shaders/Private/Common.ush b/Engine/Shaders/Private/Common.ush
--- a/Engine/Shaders/Private/Common.ush
+++ b/Engine/Shaders/Private/Common.ush
@@ -100,6 +100,7 @@ FloatDeriv4 ConstructFloatDeriv4(float4 InValue, float4 InDdx, float4 InDdy)
 }
 
 
+#define FLOAT_MAX         (asfloat(0x7F7FFFFF)) // Largest representable finite float (~3.4028235e38)
 #define POSITIVE_INFINITY (asfloat(0x7F800000))
 #define NEGATIVE_INFINITY (asfloat(0xFF800000))
 
diff --git a/Engine/Shaders/Private/PathTracing/Light/PathTracingDirectionalLight.ush b/Engine/Shaders/Private/PathTracing/Light/PathTracingDirectionalLight.ush
--- a/Engine/Shaders/Private/PathTracing/Light/PathTracingDirectionalLight.ush
+++ b/Engine/Shaders/Private/PathTracing/Light/PathTracingDirectionalLight.ush
@@ -10,7 +10,7 @@ DirectionalLight.usf: Light sampling functions for Directional light implementat
 
 FLightHit DirectionalLight_TraceLight(FRayDesc Ray, int LightId)
 {
-	if (Ray.TMax == POSITIVE_INFINITY) // we hit the background
+	if (Ray.TMax == RAY_DEFAULT_T_MAX) // we hit the background
 	{
 		float3 LightDirection = GetNormal(LightId);
 		float SinThetaMax = GetRadius(LightId); // Sin(Angle / 2)
@@ -27,7 +27,7 @@ FLightHit DirectionalLight_TraceLight(FRayDesc Ray, int LightId)
 			// that the overall illumination stays constant as the angle changes.
 			float SolidAngle = 2 * PI * OneMinusCosThetaMax;
 			float3 Radiance = GetColor(LightId) / SolidAngle;
-			return CreateLightHit(Radiance, 1.0 / SolidAngle, POSITIVE_INFINITY);
+			return CreateLightHit(Radiance, 1.0 / SolidAngle, RAY_DEFAULT_T_MAX);
 		}
 	}
 	return NullLightHit();
@@ -44,7 +44,7 @@ FLightSample DirectionalLight_SampleLight(
 	float3 Direction = TangentToWorld(DirAndPdf.xyz, GetNormal(LightId));
 
 	// Because the light is normalized by the solid angle, the radiance/pdf ratio is just the color
-	return CreateLightSample(GetColor(LightId), DirAndPdf.w, Direction, POSITIVE_INFINITY);
+	return CreateLightSample(GetColor(LightId), DirAndPdf.w, Direction, RAY_DEFAULT_T_MAX);
 }
 
 float DirectionalLight_EstimateLight(
diff --git a/Engine/Shaders/Private/PathTracing/Light/PathTracingSkyLight.ush b/Engine/Shaders/Private/PathTracing/Light/PathTracingSkyLight.ush
--- a/Engine/Shaders/Private/PathTracing/Light/PathTracingSkyLight.ush
+++ b/Engine/Shaders/Private/PathTracing/Light/PathTracingSkyLight.ush
@@ -30,10 +30,10 @@ float SkyLight_EstimateLight(
 
 FLightHit SkyLight_TraceLight(FRayDesc Ray, int LightId)
 {
-	if (Ray.TMax == POSITIVE_INFINITY)
+	if (Ray.TMax == RAY_DEFAULT_T_MAX)
 	{
 		float4 Result = SkyLight_EvalLight(Ray.Direction);
-		return CreateLightHit(Result.xyz, Result.w, POSITIVE_INFINITY);
+		return CreateLightHit(Result.xyz, Result.w, RAY_DEFAULT_T_MAX);
 	}
 	return NullLightHit();
 }
@@ -46,7 +46,7 @@ FLightSample SkyLight_SampleLight(
 {
 #if PATHTRACING_SKY_MIS == 0
 	FSkyLightSample Sample = SkyLight_SampleLight(RandSample);
-	return CreateLightSample(Sample.Radiance / Sample.Pdf, Sample.Pdf, Sample.Direction, POSITIVE_INFINITY);
+	return CreateLightSample(Sample.Radiance / Sample.Pdf, Sample.Pdf, Sample.Direction, RAY_DEFAULT_T_MAX);
 #else
 	// account for the fact that with MIS compensation, we could have a 0 probability of choosing the sky
 	float SkyLightSamplingStrategyPdf = SkyLight_Estimate() > 0.0 ? 0.5 : 0.0;
@@ -84,7 +84,7 @@ FLightSample SkyLight_SampleLight(
 
 	float Pdf = lerp(CosinePdf, SkyLightPdf, SkyLightSamplingStrategyPdf);
 
-	return CreateLightSample(SkyLightRadiance / Pdf, Pdf, Direction, POSITIVE_INFINITY);
+	return CreateLightSample(SkyLightRadiance / Pdf, Pdf, Direction, RAY_DEFAULT_T_MAX);
 #endif
 }
 
diff --git a/Engine/Shaders/Private/PathTracing/PathTracing.usf b/Engine/Shaders/Private/PathTracing/PathTracing.usf
--- a/Engine/Shaders/Private/PathTracing/PathTracing.usf
+++ b/Engine/Shaders/Private/PathTracing/PathTracing.usf
@@ -58,7 +58,7 @@ FPathState LoadPathStateData(uint Index)
 	Output.Ray.Origin = Data.RayOrigin;
 	Output.Ray.Direction = Data.RayDirection;
 	Output.Ray.TMin = 0;
-	Output.Ray.TMax = POSITIVE_INFINITY;
+	Output.Ray.TMax = RAY_DEFAULT_T_MAX;
 	Output.PathThroughput = Data.PathThroughput;
 	Output.PackedRoughnessSigma = Data.PackedRoughnessSigma;
 	return Output;
diff --git a/Engine/Shaders/Private/PathTracing/PathTracingCore.ush b/Engine/Shaders/Private/PathTracing/PathTracingCore.ush
--- a/Engine/Shaders/Private/PathTracing/PathTracingCore.ush
+++ b/Engine/Shaders/Private/PathTracing/PathTracingCore.ush
@@ -677,7 +677,7 @@ FPackedPathTracingPayload TraceTransparentRay(inout FPathState PathState, int Bo
 				LightHitPayload.ShadingModelID = SHADINGMODELID_UNLIT;
 				// light hits are fully transparent to match how the light loop works (the light surface does not cast a shadow)
 				LightHitPayload.BSDFOpacity = 1.0;
-				LightHitPayload.TransparencyColor = LightHit.HitT == POSITIVE_INFINITY ? 1.0 - BackgroundAlpha : 1.0;
+				LightHitPayload.TransparencyColor = LightHit.HitT == RAY_DEFAULT_T_MAX ? 1.0 - BackgroundAlpha : 1.0;
 				LightHitPayload.SetFrontFace();
 				PackedPayload = PackPathTracingPayload(LightHitPayload);
 			}
@@ -954,7 +954,7 @@ FPackedPathTracingPayload TraceTransparentRay(inout FPathState PathState, int Bo
 			AccumulateRadiance(PathState.Radiance, PathState.PathThroughput * PackedPayload.UnpackRadiance(), bIsCameraRay);
 		}
 
-		if (PackedPayload.HitT == POSITIVE_INFINITY)
+		if (PackedPayload.HitT == RAY_DEFAULT_T_MAX)
 		{
 			// if our hit was against an infinite light, exit now
 			break;
@@ -1822,7 +1822,7 @@ bool PathTracingKernel(inout FPathState PathState, int Bounce)
 						MaterialRay.Origin = TranslatedWorldPos;
 						MaterialRay.Direction = MaterialSample.Direction;
 						MaterialRay.TMin = 0.0;
-						MaterialRay.TMax = POSITIVE_INFINITY;
+						MaterialRay.TMax = RAY_DEFAULT_T_MAX;
 						FLightHit LightResult = TraceLight(MaterialRay, LightSampleLightId);
 						if (LightResult.IsHit())
 						{
@@ -2036,7 +2036,7 @@ bool PathTracingKernel(inout FPathState PathState, int Bounce)
 	PathState.Ray.Origin = Payload.TranslatedWorldPos;
 	PathState.Ray.Direction = MaterialSample.Direction;
 	PathState.Ray.TMin = 0;
-	PathState.Ray.TMax = POSITIVE_INFINITY;
+	PathState.Ray.TMax = RAY_DEFAULT_T_MAX;
 	ApplyRayBias(PathState.Ray.Origin, Payload.HitT, MaterialSample.PositionBiasSign * Payload.WorldGeoNormal);
 
 	// enlarge roughness based on the chosen lobe roughness
diff --git a/Engine/Shaders/Private/PathTracing/PathTracingDebug.usf b/Engine/Shaders/Private/PathTracing/PathTracingDebug.usf
--- a/Engine/Shaders/Private/PathTracing/PathTracingDebug.usf
+++ b/Engine/Shaders/Private/PathTracing/PathTracingDebug.usf
@@ -343,7 +343,7 @@ PathTracingResult PathTracingDebugCore(FRayDesc Ray, uint2 LaunchIndex)
 					ShadowRay.Origin = TranslatedWorldPos;
 					ShadowRay.Direction = MaterialSample.Direction;
 					ShadowRay.TMin = 0;
-					ShadowRay.TMax = POSITIVE_INFINITY;
+					ShadowRay.TMax = RAY_DEFAULT_T_MAX;
 
 					ApplyRayBias(ShadowRay.Origin, PackedPayload.HitT, PackedPayload.UnpackWorldGeoNormal());
 
diff --git a/Engine/Shaders/Private/PathTracing/PathTracingInitExtinctionCoefficient.usf b/Engine/Shaders/Private/PathTracing/PathTracingInitExtinctionCoefficient.usf
--- a/Engine/Shaders/Private/PathTracing/PathTracingInitExtinctionCoefficient.usf
+++ b/Engine/Shaders/Private/PathTracing/PathTracingInitExtinctionCoefficient.usf
@@ -16,7 +16,7 @@ RAY_TRACING_ENTRY_RAYGEN(PathTracingInitExtinctionCoefficientRG)
 	UpRay.Origin = PrimaryView.TranslatedWorldCameraOrigin;
 	UpRay.Direction = float3(0, 0, 1);
 	UpRay.TMin = 0.0;
-	UpRay.TMax = POSITIVE_INFINITY;
+	UpRay.TMax = RAY_DEFAULT_T_MAX;
 		
 	float3 StartingSigmaT = 0.0;
 	for (;;)
diff --git a/Engine/Shaders/Private/RayTracing/RayTracingCommon.ush b/Engine/Shaders/Private/RayTracing/RayTracingCommon.ush
--- a/Engine/Shaders/Private/RayTracing/RayTracingCommon.ush
+++ b/Engine/Shaders/Private/RayTracing/RayTracingCommon.ush
@@ -469,6 +469,9 @@ float3 UnpackNormalFromUInt32(uint PackedNormal)
 	return WorldNormal;
 }
 
+// Recommended default value for TMax for rays that should traverse the largest possible distance
+#define RAY_DEFAULT_T_MAX (FLOAT_MAX)
+
 #if !COMPUTESHADER || PLATFORM_SUPPORTS_INLINE_RAY_TRACING
 
 struct FRayDesc
@@ -517,7 +520,7 @@ FRayDesc CreatePrimaryRay(float2 UV)
 	// extract far plane (ortho only, so ignore distinction between T and Z)
 	float4 Back = mul(float4(0, 0, 0, 1), View.ClipToView);
 	Ray.TMin = 0;
-	Ray.TMax = Back.w > 0 ? (Back.z / Back.w - View.NearPlane) : asfloat(0x7f800000 - 1);
+	Ray.TMax = Back.w > 0 ? (Back.z / Back.w - View.NearPlane) : RAY_DEFAULT_T_MAX;
 
 	return Ray;
 }
