c

#12934

#tests Test velocidty output in a PPM
[RN] PR #12934: Fix Incorrect Vsualization of Motion Vectors for Static Objects

diff --git a/Engine/Shaders/Private/MaterialTemplate.ush b/Engine/Shaders/Private/MaterialTemplate.ush
--- a/Engine/Shaders/Private/MaterialTemplate.ush
+++ b/Engine/Shaders/Private/MaterialTemplate.ush
@@ -3122,10 +3122,11 @@ float2 PostProcessVelocityLookup(float Depth, float2 UV)
 	}
 	else
 	{
-		float4 ThisClip = float4( UV, Depth, 1 );
+		float2 ScreenPos = ViewportUVToScreenPos(UV);
+		float4 ThisClip = float4( ScreenPos, Depth, 1 );
 		float4 PrevClip = mul( ThisClip, View.ClipToPrevClip );
 		float2 PrevScreen = PrevClip.xy / PrevClip.w;
-		Velocity = UV - PrevScreen;
+		Velocity = ScreenPos - PrevScreen;
 	}
 
 	return Velocity;
