Support virtual texture sampling and feedback in light functions.
There is potentially an issue when baking to the light atlas, which still happens only once and without feedback. But the light atlas is 128 sized by default which is the same size as our default root virtual texture page.

diff --git a/Engine/Shaders/Private/LightFunctionCommon.ush b/Engine/Shaders/Private/LightFunctionCommon.ush
--- a/Engine/Shaders/Private/LightFunctionCommon.ush
+++ b/Engine/Shaders/Private/LightFunctionCommon.ush
@@ -22,7 +22,7 @@ float4x4 LightFunctionTranslatedWorldToLight;
  * Calculates the light function color with the given light vector. 
  * LightVector is the vector from the light to the position being shaded, in the light's coordinate space.
  */
-float3 GetLightFunctionColor(float3 LightVector, float3 TranslatedWorldSpace
+float3 GetLightFunctionColor(float3 LightVector, float3 TranslatedWorldSpace, in out FMaterialPixelParameters MaterialParameters
 #if LIGHT_FUNCTION_ATLAS
 	, float2 UVs
 #endif
@@ -58,7 +58,6 @@ float3 GetLightFunctionColor(float3 LightVector, float3 TranslatedWorldSpace
 
 #endif // LIGHT_FUNCTION_ATLAS
 
-	FMaterialPixelParameters MaterialParameters = MakeInitializedMaterialPixelParameters();
 #if NUM_MATERIAL_TEXCOORDS
 	for(int CoordinateIndex = 0;CoordinateIndex < NUM_MATERIAL_TEXCOORDS;CoordinateIndex++)
 	{
@@ -96,3 +95,17 @@ float3 GetLightFunctionColor(float3 LightVector, float3 TranslatedWorldSpace
 
 #endif // SUBSTRATE_ENABLED
 }
+
+float3 GetLightFunctionColor(float3 LightVector, float3 TranslatedWorldSpace
+#if LIGHT_FUNCTION_ATLAS
+	, float2 UVs
+#endif
+)
+{
+	FMaterialPixelParameters MaterialParameters = MakeInitializedMaterialPixelParameters();
+	return GetLightFunctionColor(LightVector, TranslatedWorldSpace, MaterialParameters
+#if LIGHT_FUNCTION_ATLAS
+	, UVs
+#endif
+	);
+}
diff --git a/Engine/Shaders/Private/LightFunctionPixelShader.usf b/Engine/Shaders/Private/LightFunctionPixelShader.usf
--- a/Engine/Shaders/Private/LightFunctionPixelShader.usf
+++ b/Engine/Shaders/Private/LightFunctionPixelShader.usf
@@ -51,9 +51,15 @@ void Main(
 	// Calculate radial view distance for stable fading
 	float ViewDistance = GetDistanceToCameraFromViewVector(PrimaryView.TranslatedWorldCameraOrigin - TranslatedWorldPosition);
 
+	FMaterialPixelParameters MaterialParameters = MakeInitializedMaterialPixelParameters();
+
+#if MATERIAL_VIRTUALTEXTURE_FEEDBACK
+	InitializeVirtualTextureFeedback(MaterialParameters.VirtualTextureFeedback);
+#endif
+
 	float GreyScale;
 	{
-		float3 Color = GetLightFunctionColor(LightVector, TranslatedWorldPosition);
+		float3 Color = GetLightFunctionColor(LightVector, TranslatedWorldPosition, MaterialParameters);
 		GreyScale = dot(Color, .3333f);
 	}
 
@@ -74,4 +80,8 @@ void Main(
 	// Light function shadows write to the blue channel.
 	OutColor = lerp(float4(1, 1, EncodedLightAttenuation.xx), EncodedLightAttenuation.xxxx, LightFunctionParameters2.z);
 #endif
+
+#if MATERIAL_VIRTUALTEXTURE_FEEDBACK
+	FinalizeVirtualTextureFeedback(MaterialParameters.VirtualTextureFeedback, SvPosition, View.VTFeedbackBuffer);
+#endif
 }
diff --git a/Engine/Source/Runtime/Engine/Private/Materials/HLSLMaterialTranslator.cpp b/Engine/Source/Runtime/Engine/Private/Materials/HLSLMaterialTranslator.cpp
--- a/Engine/Source/Runtime/Engine/Private/Materials/HLSLMaterialTranslator.cpp
+++ b/Engine/Source/Runtime/Engine/Private/Materials/HLSLMaterialTranslator.cpp
@@ -7191,6 +7191,7 @@ int32 FHLSLMaterialTranslator::TextureSample(
 			case MD_DeferredDecal:
 			case MD_PostProcess:
 			case MD_UI:
+			case MD_LightFunction:
 				break;
 			default:
 				return Errorf(TEXT("Sampling a virtual texture is not supported for the %s material domain."), *MaterialDomainString(Domain));
