[DOF] Petzval optimization in scatter pass. Precalculate the stretching transform in the vertex shader.

diff --git a/Engine/Shaders/Private/DiaphragmDOF/DOFHybridScatterPixelShader.usf b/Engine/Shaders/Private/DiaphragmDOF/DOFHybridScatterPixelShader.usf
--- a/Engine/Shaders/Private/DiaphragmDOF/DOFHybridScatterPixelShader.usf
+++ b/Engine/Shaders/Private/DiaphragmDOF/DOFHybridScatterPixelShader.usf
@@ -27,41 +27,6 @@ Texture2D BokehLUT;
 float4 ScatterOcclusionSize;
 Texture2D ScatterOcclusion;
 
-float Petzval;
-float PetzvalFalloffPower;
-float2 PetzvalExclusionBoxExtents;
-float PetzvalExclusionBoxRadius;
-
-float2 ApplyPetzval(float2 BokehCenter, float2 SamplePos)
-{
-#if CONFIG_PETZVAL
-	if (Petzval != 0)
-	{
-		const float Radius = PetzvalExclusionBoxRadius * min(PetzvalExclusionBoxExtents.x, PetzvalExclusionBoxExtents.y);
-		const float2 BoxToBokeh = sign(BokehCenter) * max(abs(BokehCenter) - PetzvalExclusionBoxExtents + Radius, 0.0);
-		float BokehDistance = max(0.0, length(BoxToBokeh));
-		const float2 Normal = BoxToBokeh / BokehDistance;
-		BokehDistance = max(0.0, BokehDistance - Radius);
-
-		if (BokehDistance > 0)
-		{
-			const float2 Tangent = float2(Normal.y, -Normal.x);
-
-			const float Intensity = pow(BokehDistance, PetzvalFalloffPower);
-			const float Stretch = 1 + abs(Petzval) * Intensity;
-
-			float2x2 ToSaggitalSpace = { Normal.x, Normal.y, Tangent.x, Tangent.y };
-			
-			float2 SagittalPos = mul(SamplePos, ToSaggitalSpace);
-			bool bMask = Petzval > 0;
-			SagittalPos.x *=  bMask ? Stretch : 1.0; // Sagittal
-			SagittalPos.y *= !bMask ? Stretch : 1.0; // Tangential
-			return mul(SagittalPos, transpose(ToSaggitalSpace));
-		}
-	}
-#endif //CONFIG_PETZVAL
-	return SamplePos;
-}
 
 //------------------------------------------------------- ENTRY POINT
 
@@ -75,6 +40,7 @@ void ScatterMainPS(
 	in nointerpolation float4 TaylorExpensionMAD0   : TEXCOORD5,
 	in nointerpolation float4 TaylorExpensionMAD1   : TEXCOORD6,
 #endif
+	in nointerpolation float2x2 PetzvalTransform : TEXCOORD7,
 	in float4 SvPosition : SV_Position,
 	out float4 OutColor : SV_Target0) 
 {
@@ -122,7 +88,7 @@ void ScatterMainPS(
 		// Pixel coordinate of the center of the bokeh.
 		float2 PixelOffset = PixelOffset0 - ScatteringScaling * kSquare2x2[i];
 		
-		PixelOffset = ApplyPetzval(BokehCenterClipPos, PixelOffset);
+		PixelOffset = mul(PixelOffset, PetzvalTransform);
 
 		PixelOffset.x *= rcp(CocInvSqueeze);
 
diff --git a/Engine/Shaders/Private/DiaphragmDOF/DOFHybridScatterVertexShader.usf b/Engine/Shaders/Private/DiaphragmDOF/DOFHybridScatterVertexShader.usf
--- a/Engine/Shaders/Private/DiaphragmDOF/DOFHybridScatterVertexShader.usf
+++ b/Engine/Shaders/Private/DiaphragmDOF/DOFHybridScatterVertexShader.usf
@@ -27,8 +27,48 @@ float ScatteringScaling;
 float CocRadiusToCircumscribedRadius;
 StructuredBuffer<float4> ScatterDrawList;
 
+float Petzval;
+float PetzvalFalloffPower;
+float2 PetzvalExclusionBoxExtents;
+float PetzvalExclusionBoxRadius;
 
-//------------------------------------------------------- ENTRY POINT
+//------------------------------------------------------- FUNCTIONS
+
+float2x2 CalcPetzvalTransform(float2 BokehCenter)
+{
+	float2x2 Transform = {1,0,0,1};
+#if CONFIG_PETZVAL
+	if (Petzval != 0)
+	{
+		const float Radius = PetzvalExclusionBoxRadius * min(PetzvalExclusionBoxExtents.x, PetzvalExclusionBoxExtents.y);
+		const float2 BoxToBokeh = sign(BokehCenter) * max(abs(BokehCenter) - PetzvalExclusionBoxExtents + Radius, 0.0);
+		float BokehDistance = max(0.0, length(BoxToBokeh));
+		const float2 Normal = BoxToBokeh / BokehDistance;
+		BokehDistance = max(0.0, BokehDistance - Radius);
+
+		if (BokehDistance > 0)
+		{
+			const float2 Tangent = float2(Normal.y, -Normal.x);
+
+			const float Intensity = pow(BokehDistance, PetzvalFalloffPower);
+			const float Stretch = 1 + abs(Petzval) * Intensity;
+
+			float2x2 ToSaggitalSpace = { Normal.x, Normal.y, Tangent.x, Tangent.y };
+
+			bool bMask = Petzval > 0;
+			float XStretch =  bMask ? Stretch : 1.0; // Sagittal
+			float YStretch = !bMask ? Stretch : 1.0; // Tangential
+			float2x2 StretchTransform = { 
+				XStretch, 0,
+				0, YStretch
+			};
+
+			Transform = mul(mul(ToSaggitalSpace, StretchTransform), transpose(ToSaggitalSpace));
+		}
+	}
+#endif //CONFIG_PETZVAL
+	return Transform;
+}
 
 void ScatterMainVS(
 	in uint VertexId : SV_VertexID,
@@ -42,6 +82,7 @@ void ScatterMainVS(
 	out nointerpolation float4 TaylorExpensionMAD0   : TEXCOORD5,
 	out nointerpolation float4 TaylorExpensionMAD1   : TEXCOORD6,
 #endif
+	out nointerpolation float2x2 PetzvalTransform : TEXCOORD7,
 	out float4 OutPosition : SV_POSITION)
 {
 	#if SCATTERING_GROUP_PACKING
@@ -121,6 +162,8 @@ void ScatterMainVS(
 		#endif
 	}
 
+	PetzvalTransform = CalcPetzvalTransform(BokehCenterClipPos);
+
 	#if BOKEH_DRAW_METHOD == BOKEH_DRAW_METHOD_TRIANGLE
 		float2 VertexPos = kSpriteVertices[SpriteVertexId];
 
