Path Tracer: Improve volumetric importance sampling around the edges of spots with inner cone angle

The spot falloff was not taken into account. We can do this by solving for the point which maximizes the dot product with the spot normal

diff --git a/Engine/Shaders/Private/PathTracing/Light/PathTracingPointLight.ush b/Engine/Shaders/Private/PathTracing/Light/PathTracingPointLight.ush
--- a/Engine/Shaders/Private/PathTracing/Light/PathTracingPointLight.ush
+++ b/Engine/Shaders/Private/PathTracing/Light/PathTracingPointLight.ush
@@ -177,6 +177,8 @@ FVolumeLightSampleSetup PointLight_PrepareLightVolumeSample(
 	float2 T = RaySphereOverlap(RayOrigin, RayDirection, Center, AttenuationRadius, TMin, TMax);
 	if (T.x < T.y)
 	{
+		float3 C = GetColor(LightId);
+		float LightImportance = max3(C.x, C.y, C.z);
 		if (IsSpotLight(LightId))
 		{
 			float3 N = GetNormal(LightId);
@@ -194,9 +196,22 @@ FVolumeLightSampleSetup PointLight_PrepareLightVolumeSample(
 				// hit the sphere but missed the cone
 				return NullVolumeLightSetup();
 			}
+
+			// find point on the ray which minimizes the dot product against the spot axis
+			// This is the largest that the spot attenuation can be for this ray -- this helps with sampling that goes through the edge of the cone when attenuation is strong
+			// Given the constants below, the dot product between a point Ro+t*Rd along the ray against the spot's cone axis can be solved by:
+			// Solve[D[(A - t*B)/Sqrt[C - 2*t*D + t^2], t] == 0, t]
+			// The solution of which is given below. If the denominator is negative, we are 
+			float3 Q = Center - RayOrigin;
+			float A = dot(Q, N);
+			float B = dot(RayDirection, N);
+			float C = dot(Q, Q);
+			float D = dot(Q, RayDirection);
+			float Denom = B * D - A;
+			float TClosest = (B * C - A * D) * rcp(Denom);
+			float3 Closest = (Q - RayDirection * TClosest) * sign(Denom);
+			LightImportance *= SpotAttenuation(normalize(Closest), N, GetCosConeAngles(LightId));
 		}
-		float3 C = GetColor(LightId);
-		float LightImportance = max3(C.x, C.y, C.z);
 		if ((SceneLights[LightId].Flags & PATHTRACER_FLAG_NON_INVERSE_SQUARE_FALLOFF_MASK) == 0)
 		{
 			// inverse square falloff requires equi-angular sampling
