Fix Mac Editor PSO creationg. In Metal, PS output must have as many color channel outputs as the attached render target even with a color write mask. Instead pda shader output from float2 to float4 and WriteMask to the RG channels via render state.

diff --git a/Engine/Shaders/Private/Landscape/LandscapeEditLayersHeightmaps.usf b/Engine/Shaders/Private/Landscape/LandscapeEditLayersHeightmaps.usf
--- a/Engine/Shaders/Private/Landscape/LandscapeEditLayersHeightmaps.usf
+++ b/Engine/Shaders/Private/Landscape/LandscapeEditLayersHeightmaps.usf
@@ -94,7 +94,7 @@ float4 ExtractSampleInfo(int2 InTextureCoordinates, float2 InPackedHeight)
 // Pixel shaders : 
 
 #if MERGE_EDIT_LAYER
-void MergeEditLayerPS(in float4 SVPos : SV_POSITION, out float2 OutColor : SV_Target0)
+void MergeEditLayerPS(in float4 SVPos : SV_POSITION, out float4 OutColor : SV_Target0)
 {
 	uint2 TextureCoordinates = floor(SVPos.xy);
 		
@@ -128,7 +128,7 @@ void MergeEditLayerPS(in float4 SVPos : SV_POSITION, out float2 OutColor : SV_Ta
 	
 	FinalHeight = clamp(FinalHeight, -32768.0, 32767.0);
 	FinalHeight += 32768.0f;
-	OutColor = PackHeight(FinalHeight);
+	OutColor = float4(PackHeight(FinalHeight), 0.0f, 0.0f);
 }
 #endif // MERGE_EDIT_LAYER
 
diff --git a/Engine/Source/Runtime/Landscape/Private/LandscapeEditLayers.cpp b/Engine/Source/Runtime/Landscape/Private/LandscapeEditLayers.cpp
--- a/Engine/Source/Runtime/Landscape/Private/LandscapeEditLayers.cpp
+++ b/Engine/Source/Runtime/Landscape/Private/LandscapeEditLayers.cpp
@@ -1699,7 +1699,8 @@ public:
 			MoveTemp(InRDGEventName),
 			PixelShader,
 			InParameters,
-			FIntRect(0, 0, InTextureSize.X, InTextureSize.Y));
+			FIntRect(0, 0, InTextureSize.X, InTextureSize.Y),
+			TStaticBlendStateWriteMask<CW_RG>::GetRHI());
 	}
 };
 
