Shader compilation error when shader used in PC preview on Mac (!VELOCITY_ENCODE_HAS_PIXEL_ANIMATION) due to uninitialized out parameter.

diff --git a/Engine/Shaders/Private/TemporalSuperResolution/TSRDepthVelocityAnalysis.ush b/Engine/Shaders/Private/TemporalSuperResolution/TSRDepthVelocityAnalysis.ush
--- a/Engine/Shaders/Private/TemporalSuperResolution/TSRDepthVelocityAnalysis.ush
+++ b/Engine/Shaders/Private/TemporalSuperResolution/TSRDepthVelocityAnalysis.ush
@@ -118,6 +118,8 @@ void FetchAndComputeScreenVelocity(
 
 	#if VELOCITY_ENCODE_HAS_PIXEL_ANIMATION
 		bHasPixelAnimation = DecodeHasPixelAnimationFromVelocityTexture(EncodedVelocity);
+	#else
+		bHasPixelAnimation = false;
 	#endif
 	
 	#if CONFIG_USE_SCENE_VELOCITY
