Reduce one barrier and improve the memory load for thin geometry detection pass. Get a small performance improvement of 0.01ms on console and pc.

#tests Test done on console and pc
#rnx

diff --git a/Engine/Shaders/Private/TemporalSuperResolution/TSRDetectThinGeometry.usf b/Engine/Shaders/Private/TemporalSuperResolution/TSRDetectThinGeometry.usf
--- a/Engine/Shaders/Private/TemporalSuperResolution/TSRDetectThinGeometry.usf
+++ b/Engine/Shaders/Private/TemporalSuperResolution/TSRDetectThinGeometry.usf
@@ -19,8 +19,6 @@
 
 #define EDGE_TYPE EDGE_HORIZONTAL_VERTICAL_DIAGNAL
 
-#define CONFIG_THIN_GEOMETRY_COVERAGE 1
-
 //------------------------------------------------------- CONSTANTS
 
 #if CONFIG_LDS_DESPILL
@@ -36,13 +34,10 @@ float ErrorMultiplier;
 float MaxRelaxationWeight;
 
 Texture2D<float> SceneDepthTexture;
+Texture2D<tsr_half> CurrentCoverageTexture;
+Texture2D<tsr_half> ReprojectedHistoryCoverageTexture;
 
-#if CONFIG_THIN_GEOMETRY_COVERAGE
-    Texture2D<tsr_half> CurrentCoverageTexture;
-    Texture2D<tsr_half> ReprojectedHistoryCoverageTexture;
-
-    RWTexture2DArray<tsr_half> HistoryCoverageOutput;
-#endif
+RWTexture2DArray<tsr_half> HistoryCoverageOutput;
 RWTexture2DArray<uint> R8Output;
 
 //------------------------------------------------------- GLOBALS
@@ -97,33 +92,30 @@ void FetchInput(
 	tsr_tensor_short2 PixelPos
 	,out tsr_tensor_half SceneDepth
 	,out tsr_tensor_half DeviceZThinGeometryDetectionError
-#if CONFIG_THIN_GEOMETRY_COVERAGE
 	,out tsr_tensor_half CurrentCoverage
 	,out tsr_tensor_half ReprojectedHistoryCoverage
-#endif
 	)
 {
-	tsr_tensor_half LumaAbsDelta;
-
+	tsr_tensor_float DepthCache;
 	UNROLL_N(SIMD_SIZE)
 	for (uint ElementIndex = 0; ElementIndex < SIMD_SIZE; ElementIndex++)
 	{
 		tsr_short2 ClampedFetchPixelPos = PixelPos.GetElement(ElementIndex);
+		CurrentCoverage.SetElement(ElementIndex, CurrentCoverageTexture[ClampedFetchPixelPos]);
+		ReprojectedHistoryCoverage.SetElement(ElementIndex, ReprojectedHistoryCoverageTexture[ClampedFetchPixelPos]);
+		DepthCache.SetElement(ElementIndex,SceneDepthTexture[ClampedFetchPixelPos]);
+	}
 
-        #if CONFIG_THIN_GEOMETRY_COVERAGE
-		    CurrentCoverage.SetElement(ElementIndex, CurrentCoverageTexture[ClampedFetchPixelPos]);
-		    ReprojectedHistoryCoverage.SetElement(ElementIndex, ReprojectedHistoryCoverageTexture[ClampedFetchPixelPos]);
-        #endif
-
-		float Depth = SceneDepthTexture[ClampedFetchPixelPos];
+	UNROLL_N(SIMD_SIZE)
+	for (uint ElementIndex = 0; ElementIndex < SIMD_SIZE; ElementIndex++)
+	{
+		float Depth = DepthCache.GetElement(ElementIndex);
 		float DeviceZError = ComputePixelDeviceZError(Depth);
 
 		SceneDepth.SetElement(ElementIndex, tsr_half(Depth));
 		DeviceZThinGeometryDetectionError.SetElement(ElementIndex, tsr_half(ErrorMultiplier * DeviceZError));
 	}
 
-	//Debug[0].SetComponent(1, LumaAbsDelta);
-
 } // FetchInput
 
 
@@ -223,7 +215,6 @@ tsr_tensor_uint EncodeThinGeometryWeight(tsr_tensor_half ThinGeometryEdge
 	tsr_tensor_half Weight
 )
 {
-#if CONFIG_THIN_GEOMETRY_COVERAGE
 	// 7bit   | 1bit
 	// weight | 0 : Edge line
 	//		  | 1 : Cluster hole region
@@ -233,10 +224,6 @@ tsr_tensor_uint EncodeThinGeometryWeight(tsr_tensor_half ThinGeometryEdge
 	tsr_tensor_uint Encoded = bit_shift_left(tsr_tensor_uint::CastFrom<tsr_half>(Weight * tsr_half(127.0)),tsr_tensor_uint::Const(0x1u));
 	Encoded = bit_or(Encoded,select(IsClusterHoleRegion,tsr_tensor_uint::Const(0x1u),tsr_tensor_uint::Const(0)));
 	return Encoded;
-#else
-	tsr_tensor_uint Encoded = bit_shift_left(tsr_tensor_uint::CastFrom<tsr_half>(ThinGeometryEdge * tsr_half(127.0)),tsr_tensor_uint::Const(0x1u));
-	return Encoded;
-#endif
 }
 
 //------------------------------------------------------- ENTRY POINT
@@ -266,17 +253,13 @@ void MainCS(
 	tsr_tensor_short2 PixelPos = ComputePixelPos();
 	tsr_tensor_half DeviceZ;
 	tsr_tensor_half DeviceZThinGeometryDetectionError;
-#if CONFIG_THIN_GEOMETRY_COVERAGE
 	tsr_tensor_half CurrentCoverage;
 	tsr_tensor_half ReprojectedHistoryCoverage;
-#endif
 	FetchInput(PixelPos,
 		DeviceZ, 
 		DeviceZThinGeometryDetectionError
-#if CONFIG_THIN_GEOMETRY_COVERAGE
 		,CurrentCoverage
 		,ReprojectedHistoryCoverage
-#endif
 	);
 	
 	//Debug[0].SetComponent(0, DeviceZ);
@@ -290,18 +273,11 @@ void MainCS(
 		//Debug[0].SetComponent(2, tsr_tensor_half(ThinGeometryEdge));
 	}
 
-	{
-		//dilate the edge with (0.8,1,0.8), such that, dilated edges has 0.8 for direct neighbor, 0.64 for diagnal, 1 for consecutive neighbors.
-		// This value are used to boost the weight with
-		// lerp(blend, 1.0, thinedgeblend)
-		ThinGeometryEdge = saturate(WeightedAvg3x3(ThinGeometryEdge, float3(0.8, 1.0, 0.8)));
-		//Debug[0].SetComponent(3, tsr_tensor_half(ThinGeometryEdge));
-	}
-
 	tsr_tensor_half HistoryRelaxationWeight = tsr_tensor_half::Const(0);
 	tsr_tensor_bool IsClusterHoleRegion = tsr_tensor_bool::Const(false);
 	tsr_tensor_half NewSampleWeight = tsr_tensor_half::Const(0.1);
-#if CONFIG_THIN_GEOMETRY_COVERAGE
+
+	// Thin geometry coverage estimation
 	{
 		// Student t test of current observation and history to check if they are of the same distribution.
 		tsr_tensor_bool bSameDistribution = tsr_tensor_bool::Const(false);
@@ -320,8 +296,8 @@ void MainCS(
 			T = abs(select( XAverage > tsr_tensor_half::Const(0), XAverage * rcp(S), tsr_tensor_half::Const(0)));
 			bSameDistribution = T < tsr_tensor_half::Const(20.754); //2.306 (p=0.05) * 9
 			
-			Debug[0].SetComponent(0,abs(T));
-			Debug[0].SetComponent(1,tsr_tensor_half::CastFrom<bool>(abs(T) > tsr_tensor_half::Const(20.754)));
+			//Debug[0].SetComponent(0,abs(T));
+			//Debug[0].SetComponent(1,tsr_tensor_half::CastFrom<bool>(abs(T) > tsr_tensor_half::Const(20.754)));
 		}
 
 		tsr_tensor_bool ValidCurrentThinGeometryRegion = tsr_tensor_bool::Const(false);
@@ -362,21 +338,21 @@ void MainCS(
 			tsr_tensor_bool bValidHistoryRelaxation = and(bThinGeometryHistoryRegion,and(ValidCurrentThinGeometryRegion, bSameDistribution));
 			HistoryRelaxationWeight = select(bValidHistoryRelaxation, HistoryRelaxationWeight, tsr_tensor_half::Const(0));
 			
-			//Diffuse a little bit to suppress the history clamping around the valid pixel and apply max relaxation weight.
-			HistoryRelaxationWeight = saturate(WeightedAvg3x3(HistoryRelaxationWeight, float3(0.8, 1.0, 0.8))*tsr_half(MaxRelaxationWeight));
-			HistoryRelaxationWeight = saturate(WeightedAvg3x3(HistoryRelaxationWeight, float3(0.5, 1.0, 0.5)));
+			//1. Diffuse a little bit to suppress the history clamping around the valid pixel and apply max relaxation weight.
+			//2. Dilate the edge with (0.8,1,0.8), such that, dilated edges has 0.8 for direct neighbor, 0.64 for diagnal, 1 for consecutive neighbors.
+			Deconcatenate(WeightedAvg3x3(Concatenate(ThinGeometryEdge,HistoryRelaxationWeight),float3(0.8, 1.0, 0.8)),ThinGeometryEdge,HistoryRelaxationWeight);
+			HistoryRelaxationWeight = saturate(WeightedAvg3x3(saturate(HistoryRelaxationWeight * tsr_half(MaxRelaxationWeight)), float3(0.5, 1.0, 0.5)));
 
 			// Not similar patches should reset relaxation weight to 0
 			// to use more strict history clamping.
 			HistoryRelaxationWeight = select(bSameDistribution, HistoryRelaxationWeight, tsr_tensor_half::Const(0));
 			
-			Debug[0].SetComponent(2,HistoryRelaxationWeight);
+			//Debug[0].SetComponent(2,HistoryRelaxationWeight);
 		}
 
 		// Valid pixels for cluster hole stability improvement.
 		IsClusterHoleRegion = ReprojectedHistoryCoverage > tsr_tensor_half::Const(half(0.99f/127.0f));
 	}
-#endif
 
 	tsr_tensor_uint EncodedThinGeometryEdge = EncodeThinGeometryWeight(
 		ThinGeometryEdge
@@ -392,10 +368,7 @@ void MainCS(
 		const tsr_short2 OutputPixelPos = ComputeElementOutputPixelPos(ElementIndex);
 		
 		R8Output[uint3(OutputPixelPos,ThinGeometryTextureIndex)] = EncodedThinGeometryEdge.GetElement(ElementIndex);//uint(ThinGeometryFlag * tsr_half(255.0));
-
-        #if CONFIG_THIN_GEOMETRY_COVERAGE
-		    HistoryCoverageOutput[uint3(OutputPixelPos,0)] = ReprojectedHistoryCoverage.GetElement(ElementIndex);
-		#endif
+		HistoryCoverageOutput[uint3(OutputPixelPos,0)] = ReprojectedHistoryCoverage.GetElement(ElementIndex);
 	}
 	
 	#if DEBUG_OUTPUT
