[DOF] Dual pixel vectorization optimization in scatter pass. On supported platforms, 2 pixels are processed together at a time using packed 16-bit float ops. This mitigates some of the ALU bottleneck and lowers interpolator count.

diff --git a/Engine/Shaders/Private/DiaphragmDOF/DOFHybridScatterPixelShader.usf b/Engine/Shaders/Private/DiaphragmDOF/DOFHybridScatterPixelShader.usf
--- a/Engine/Shaders/Private/DiaphragmDOF/DOFHybridScatterPixelShader.usf
+++ b/Engine/Shaders/Private/DiaphragmDOF/DOFHybridScatterPixelShader.usf
@@ -6,7 +6,7 @@
 
 #include "DOFHybridScatterCommon.ush"
 #include "../FastMath.ush"
-
+#include "/Engine/Public/DualPixelVectorization.ush"
 
 //------------------------------------------------------- COMPILE TIME CONFIG
 
@@ -28,19 +28,36 @@ float4 ScatterOcclusionSize;
 Texture2D ScatterOcclusion;
 
 
-//------------------------------------------------------- ENTRY POINT
+//------------------------------------------------------- FUNCTIONS
+
+
+half2 HalfSafeRcp(half2 Value)
+{
+	return half2(
+		Value.x > 0.0 ? rcp(Value.x) : 0.0,
+		Value.y > 0.0 ? rcp(Value.y) : 0.0
+	);
+}
+
+half2x2 dpv_mul_vec_mat(half2x2 InVector, half2x2 InMatrix)
+{
+	return dpv_interleave_registers(
+		mul(dpv_lo(InVector), InMatrix), 
+		mul(dpv_hi(InVector), InMatrix)
+	);
+}
 
 void ScatterMainPS(
 	in nointerpolation float4 BokehCenterScreenAndClipPos : TEXCOORD0,
 	in nointerpolation float4 ColorAndAbsCocRadius0 : TEXCOORD1,
 	in nointerpolation float4 ColorAndAbsCocRadius1 : TEXCOORD2,
-	in nointerpolation float4 ColorAndAbsCocRadius2 : TEXCOORD3,
-	in nointerpolation float4 ColorAndAbsCocRadius3 : TEXCOORD4,
 #if TAYLOR_EXPENDED_COC_INTERSECTION
-	in nointerpolation float4 TaylorExpensionMAD0   : TEXCOORD5,
-	in nointerpolation float4 TaylorExpensionMAD1   : TEXCOORD6,
+	in nointerpolation float4 TaylorExpensionMAD0   : TEXCOORD3,
+	in nointerpolation float4 TaylorExpensionMAD1   : TEXCOORD4,
+#endif
+#if CONFIG_PETZVAL
+	in nointerpolation float2 PetzvalTransformPacked : TEXCOORD5,
 #endif
-	in nointerpolation float2x2 PetzvalTransform : TEXCOORD7,
 	in float4 SvPosition : SV_Position,
 	out float4 OutColor : SV_Target0) 
 {
@@ -49,82 +66,117 @@ void ScatterMainPS(
 	const float2 BokehCenter0 = BokehCenterScreenAndClipPos.xy;
 	const float2 BokehCenterClipPos = BokehCenterScreenAndClipPos.zw;
 
-	float4 Color[PIXEL_COUNT_PER_SCATTER_GROUP] = {
-		float4(ColorAndAbsCocRadius0.rgb, 0),
-		float4(ColorAndAbsCocRadius1.rgb, 0),
-		float4(ColorAndAbsCocRadius2.rgb, 0),
-		float4(ColorAndAbsCocRadius3.rgb, 0),
-	};
-	
-	float AbsCocRadius[PIXEL_COUNT_PER_SCATTER_GROUP] = {
-		ColorAndAbsCocRadius0.w,
-		ColorAndAbsCocRadius1.w,
-		ColorAndAbsCocRadius2.w,
-		ColorAndAbsCocRadius3.w,
-	};
+	half3 Color[PIXEL_COUNT_PER_SCATTER_GROUP];
+	half AbsCocRadius[PIXEL_COUNT_PER_SCATTER_GROUP];
+
+	{
+		half2 Unpacked0;
+		half2 Unpacked1;
+
+		Unpacked0 = (half2)UnpackFloat2FromUInt(asuint(ColorAndAbsCocRadius0.x));
+		Unpacked1 = (half2)UnpackFloat2FromUInt(asuint(ColorAndAbsCocRadius0.y));
+		Color[0] = half3(Unpacked0.xy, Unpacked1.x);
+		AbsCocRadius[0] = Unpacked1.y;
+
+		Unpacked0 = (half2)UnpackFloat2FromUInt(asuint(ColorAndAbsCocRadius0.z));
+		Unpacked1 = (half2)UnpackFloat2FromUInt(asuint(ColorAndAbsCocRadius0.w));
+		Color[1] = half3(Unpacked0.xy, Unpacked1.x);
+		AbsCocRadius[1] = Unpacked1.y;
+
+		Unpacked0 = (half2)UnpackFloat2FromUInt(asuint(ColorAndAbsCocRadius1.x));
+		Unpacked1 = (half2)UnpackFloat2FromUInt(asuint(ColorAndAbsCocRadius1.y));
+		Color[2] = half3(Unpacked0.xy, Unpacked1.x);
+		AbsCocRadius[2] = Unpacked1.y;
+
+		Unpacked0 = (half2)UnpackFloat2FromUInt(asuint(ColorAndAbsCocRadius1.z));
+		Unpacked1 = (half2)UnpackFloat2FromUInt(asuint(ColorAndAbsCocRadius1.w));
+		Color[3] = half3(Unpacked0.xy, Unpacked1.x);
+		AbsCocRadius[3] = Unpacked1.y;
+	}
 	
 	#if TAYLOR_EXPENDED_COC_INTERSECTION
-		float2 TaylorExpension[4] = {
-			TaylorExpensionMAD0.xy,
-			TaylorExpensionMAD0.zw,
-			TaylorExpensionMAD1.xy,
-			TaylorExpensionMAD1.zw,
+		half2 TaylorExpension[4] = {
+			half2(TaylorExpensionMAD0.xy),
+			half2(TaylorExpensionMAD0.zw),
+			half2(TaylorExpensionMAD1.xy),
+			half2(TaylorExpensionMAD1.zw),
 		};
 	#endif
 
 	#if DIM_SCATTER_OCCLUSION
 		float4 ScatterOcclusionSample = ScatterOcclusion.SampleLevel(GlobalBilinearWrappedSampler, SvPosition.xy * ScatterOcclusionSize.zw, 0);
 
-		float CocRadiusAvg = ScatterOcclusionSample.x;
-		float CocRadiusVariance = ScatterOcclusionSample.y;
+		half CocRadiusAvg = half(ScatterOcclusionSample.x);
+		half CocRadiusVariance = half(ScatterOcclusionSample.y);
 	#endif
 	
-	float2 PixelOffset0 = SvPosition.xy - BokehCenter0;
+	half2 PixelOffset0 = half2(SvPosition.xy - BokehCenter0);
+
+#if CONFIG_PETZVAL
+	half2x2 PetzvalTransform = half2x2(
+		UnpackFloat2FromUInt(asuint(PetzvalTransformPacked.x)),
+		UnpackFloat2FromUInt(asuint(PetzvalTransformPacked.y))
+	);
+#endif
+
+	half CocSqueeze = half(rcp(CocInvSqueeze));
 
 	// Evaluate all bokeh of the scattering group.
+	// Dual Pixel Vectorization: evaluate 2 pixels at a time using packed 16-bit float instructions on supported platforms.
 	UNROLL
-	for (uint i = 0; i < PIXEL_COUNT_PER_SCATTER_GROUP; i++)
+	for (uint i = 0; i < PIXEL_COUNT_PER_SCATTER_GROUP; i += DPV_PIXEL_PER_LANE)
 	{
 		// Pixel coordinate of the center of the bokeh.
-		float2 PixelOffset = PixelOffset0 - ScatteringScaling * kSquare2x2[i];
-		
-		PixelOffset = mul(PixelOffset, PetzvalTransform);
+		half2x2 PixelLocalOffset = dpv_interleave_registers(half2(kSquare2x2[i]), half2(kSquare2x2[i + 1]));
+		half2x2 PixelOffset = dpv_sub(PixelOffset0, dpv_scale(half(ScatteringScaling), PixelLocalOffset));
+
+		half2 PixelAbsCocRadius = dpv_interleave_registers(AbsCocRadius[i], AbsCocRadius[i+1]);
+
+#if CONFIG_PETZVAL		
+		PixelOffset = dpv_mul_vec_mat(PixelOffset, PetzvalTransform);
+#endif
 
-		PixelOffset.x *= rcp(CocInvSqueeze);
+		PixelOffset = dpv_scale(PixelOffset, half2(CocSqueeze, 1));
 
-		float Intersection;
+		half2 Intersection;
 		#if DIM_BOKEH_SIMULATION
 		{
-			float PixelDistance = sqrtFast(dot(PixelOffset, PixelOffset));
-			float2 LookUpUV = PixelOffset * (CONFIG_POINT_MIRROR_BOKEH * (0.5 - 1.0 / float(BOKEH_LUT_SIZE)) * SafeRcp(PixelDistance));
+			half2 PixelDistance = dpv_length(PixelOffset);
+			half2x2 LookUpUV = dpv_scale(PixelOffset, half(CONFIG_POINT_MIRROR_BOKEH * (0.5 - 1.0 / float(BOKEH_LUT_SIZE))) * HalfSafeRcp(PixelDistance));
 
 			// TODO: might be faster without using view uniform buffer.
-			float4 LookupSample = BokehLUT.SampleLevel(GlobalBilinearWrappedSampler, LookUpUV, 0);
-			float CocRadiusToBokehEdgeFactor = LookupSample.x;
+			float4 LookupSample0 = BokehLUT.SampleLevel(GlobalBilinearWrappedSampler, dpv_lo(LookUpUV), 0);
+			float4 LookupSample1 = BokehLUT.SampleLevel(GlobalBilinearWrappedSampler, dpv_hi(LookUpUV), 0);
+			half2 CocRadiusToBokehEdgeFactor = dpv_interleave_registers(half(LookupSample0.x), half(LookupSample1.x));
 
-			Intersection = saturate(AbsCocRadius[i] * CocRadiusToBokehEdgeFactor - PixelDistance + ANTI_ALIASING_FEATHER_OFFSET);
+			Intersection = saturate(PixelAbsCocRadius * CocRadiusToBokehEdgeFactor - PixelDistance + ANTI_ALIASING_FEATHER_OFFSET);
 		}
 		#elif TAYLOR_EXPENDED_COC_INTERSECTION
 		{
-			Intersection = saturate(dot(PixelOffset, PixelOffset) * TaylorExpension[i].x + TaylorExpension[i].y);
+			half2 TaylorExpansionX = dpv_interleave_registers(TaylorExpension[i].x, TaylorExpension[i+1].x);
+			half2 TaylorExpansionY = dpv_interleave_registers(TaylorExpension[i].y, TaylorExpension[i+1].y);
+			Intersection = saturate(dpv_dot(PixelOffset, PixelOffset) * TaylorExpansionX + TaylorExpansionY);
 		}
 		#else
 		{
 			// Pretty simple sphere intersection.
-			float PixelDistance = sqrtFast(dot(PixelOffset, PixelOffset));
-			Intersection = saturate(AbsCocRadius[i] - PixelDistance + ANTI_ALIASING_FEATHER_OFFSET);
+			half2 PixelDistance = dpv_length(PixelOffset);
+			Intersection = saturate(PixelAbsCocRadius - PixelDistance + ANTI_ALIASING_FEATHER_OFFSET);
 		}
 		#endif
 
-		float Visibility = 1;
+		half2 Visibility = 1;
 		#if DIM_SCATTER_OCCLUSION
 		{
-			float CocRadiusComparison = max(AbsCocRadius[i] * (1 - FAST_GATHERING_COC_ERROR) - CocRadiusAvg, 0);
+			half2 CocRadiusComparison = max(PixelAbsCocRadius * half(1 - FAST_GATHERING_COC_ERROR) - CocRadiusAvg, (half2)0);
 			Visibility = CocRadiusVariance / (CocRadiusVariance + CocRadiusComparison * CocRadiusComparison);
 		}
 		#endif
 
-		OutColor += Color[i] * (Intersection * Visibility);
+		half3x2 PixelColor = dpv_interleave_registers(Color[i].rgb, Color[i + 1].rgb);
+		PixelColor = dpv_scale(PixelColor, Intersection * Visibility);
+		half3 ColorAcc = dpv_lo(PixelColor) + dpv_hi(PixelColor);
+		OutColor.rgb += ColorAcc;
 	}
 	
 	#if DEBUG_HYBRID_SCATTERING_SPRITE == 1
diff --git a/Engine/Shaders/Private/DiaphragmDOF/DOFHybridScatterVertexShader.usf b/Engine/Shaders/Private/DiaphragmDOF/DOFHybridScatterVertexShader.usf
--- a/Engine/Shaders/Private/DiaphragmDOF/DOFHybridScatterVertexShader.usf
+++ b/Engine/Shaders/Private/DiaphragmDOF/DOFHybridScatterVertexShader.usf
@@ -76,13 +76,13 @@ void ScatterMainVS(
 	out nointerpolation float4 BokehCenterScreenAndClipPos : TEXCOORD0,
 	out nointerpolation float4 ColorAndAbsCocRadius0 : TEXCOORD1,
 	out nointerpolation float4 ColorAndAbsCocRadius1 : TEXCOORD2,
-	out nointerpolation float4 ColorAndAbsCocRadius2 : TEXCOORD3,
-	out nointerpolation float4 ColorAndAbsCocRadius3 : TEXCOORD4,
 #if TAYLOR_EXPENDED_COC_INTERSECTION
-	out nointerpolation float4 TaylorExpensionMAD0   : TEXCOORD5,
-	out nointerpolation float4 TaylorExpensionMAD1   : TEXCOORD6,
+	out nointerpolation float4 TaylorExpensionMAD0   : TEXCOORD3,
+	out nointerpolation float4 TaylorExpensionMAD1   : TEXCOORD4,
+#endif
+#if CONFIG_PETZVAL
+	out nointerpolation float2 PetzvalTransformPacked : TEXCOORD5,
 #endif
-	out nointerpolation float2x2 PetzvalTransform : TEXCOORD7,
 	out float4 OutPosition : SV_POSITION)
 {
 	#if SCATTERING_GROUP_PACKING
@@ -144,10 +144,18 @@ void ScatterMainVS(
 	const float2 BokehCenterScreenPos = ScatteringScaling * ScatterDrawList[5 * GlobalSpriteId].xy;
 	const float2 BokehCenterClipPos = (BokehCenterScreenPos * ViewportSize.zw) * 2.0 - 1.0;
 	{
-		ColorAndAbsCocRadius0 = float4(Color[0].rgb, AbsCocRadius[0]);
-		ColorAndAbsCocRadius1 = float4(Color[1].rgb, AbsCocRadius[1]);
-		ColorAndAbsCocRadius2 = float4(Color[2].rgb, AbsCocRadius[2]);
-		ColorAndAbsCocRadius3 = float4(Color[3].rgb, AbsCocRadius[3]);
+		ColorAndAbsCocRadius0 = float4(
+			asfloat(PackFloat2ToUInt(Color[0].r, Color[0].g)), 
+			asfloat(PackFloat2ToUInt(Color[0].b, AbsCocRadius[0])),
+			asfloat(PackFloat2ToUInt(Color[1].r, Color[1].g)), 
+			asfloat(PackFloat2ToUInt(Color[1].b, AbsCocRadius[1]))
+		);
+		ColorAndAbsCocRadius1 = float4(
+			asfloat(PackFloat2ToUInt(Color[2].r, Color[2].g)), 
+			asfloat(PackFloat2ToUInt(Color[2].b, AbsCocRadius[2])),
+			asfloat(PackFloat2ToUInt(Color[3].r, Color[3].g)), 
+			asfloat(PackFloat2ToUInt(Color[3].b, AbsCocRadius[3]))
+		);
 	
 		BokehCenterScreenAndClipPos.xy = BokehCenterScreenPos;
 		BokehCenterScreenAndClipPos.zw = BokehCenterClipPos;
@@ -162,7 +170,14 @@ void ScatterMainVS(
 		#endif
 	}
 
-	PetzvalTransform = CalcPetzvalTransform(BokehCenterClipPos);
+#if CONFIG_PETZVAL
+	float2x2 PetzvalTransform = CalcPetzvalTransform(BokehCenterClipPos);
+	half4 PetzvalTransformHalf = half4(float4(PetzvalTransform));
+	PetzvalTransformPacked = float2(
+		asfloat(PackFloat2ToUInt(PetzvalTransformHalf[0], PetzvalTransformHalf[1])),
+		asfloat(PackFloat2ToUInt(PetzvalTransformHalf[2], PetzvalTransformHalf[3]))
+	);
+#endif
 
 	#if BOKEH_DRAW_METHOD == BOKEH_DRAW_METHOD_TRIANGLE
 		float2 VertexPos = kSpriteVertices[SpriteVertexId];
