PR #11882: Fixed strange stripes on the screen when using gtao on mobile.

#11882

diff --git a/Engine/Shaders/Private/PostProcessAmbientOcclusionMobile.usf b/Engine/Shaders/Private/PostProcessAmbientOcclusionMobile.usf
--- a/Engine/Shaders/Private/PostProcessAmbientOcclusionMobile.usf
+++ b/Engine/Shaders/Private/PostProcessAmbientOcclusionMobile.usf
@@ -322,8 +322,8 @@ half3 GetNormal(half2 UV, int2 ThreadPos, float3 ViewSpacePosMid)
 
 #else
 	// Get the normal derived from the depth buffer
-	half2 XOffset = half2(BufferSizeAndInvSize.z, 0.0f);
-	half2 YOffset = half2(0.0f, BufferSizeAndInvSize.w);
+	float2 XOffset = float2(BufferSizeAndInvSize.z, 0.0f);
+	float2 YOffset = float2(0.0f, BufferSizeAndInvSize.w);
 
 #if HORIZONSEARCH_INTEGRAL_PIXEL_SHADER || SPATIALFILTER_PIXEL_SHADER
 	float DeviceZ = GetDeviceZFromAOInput(UV);
