Substrate - fixed thin translucency for shacing model conversion node.

#tests PC editor dx12
#rnx

diff --git a/Engine/Shaders/Private/Substrate/SubstrateLegacyConversion.ush b/Engine/Shaders/Private/Substrate/SubstrateLegacyConversion.ush
--- a/Engine/Shaders/Private/Substrate/SubstrateLegacyConversion.ush
+++ b/Engine/Shaders/Private/Substrate/SubstrateLegacyConversion.ush
@@ -438,6 +438,7 @@ FSubstrateData SubstrateConvertLegacyMaterialDynamic(
 		LegacySlab.Emissive 		= Emissive;
 		LegacySlab.Thickness 		= Thickness;
 		LegacySlab.SharedLocalBasisIndex = SharedLocalBasisIndex;
+		LegacySlab.SSSType			= SSS_TYPE_SIMPLEVOLUME;	// This translucent explicitely use this mode to get colored translucency from SSSMFP.
 
 		// Opacity on the root node is the coverage of the disney model on top of the translucent part.
 		// Thin translucency model have a special surface coverage input.
