Reduce VGPR count of skin cache shaders by using half's for tangents.

diff --git a/Engine/Shaders/Private/GpuSkinCacheComputeShader.usf b/Engine/Shaders/Private/GpuSkinCacheComputeShader.usf
--- a/Engine/Shaders/Private/GpuSkinCacheComputeShader.usf
+++ b/Engine/Shaders/Private/GpuSkinCacheComputeShader.usf
@@ -41,7 +41,7 @@ void GetMorphBuffer(uint VertexIndex, out float3 DeltaPosition, out half3 DeltaT
 		MorphBuffer[Offset + 1],
 		MorphBuffer[Offset + 2]);
 
-	DeltaTangentZ = float3(
+	DeltaTangentZ = half3(
 		MorphBuffer[Offset + 3],
 		MorphBuffer[Offset + 4],
 		MorphBuffer[Offset + 5]);
@@ -99,8 +99,8 @@ void GetClothBuffer(uint VertexIndex, uint ClothIndexOffset, out FVertexTriangle
 
 struct FVertexUnpacked
 {
-	float4	TangentX;
-	float4	TangentZ;
+	half4	TangentX;
+	half4	TangentZ;
 
 #if !GPUSKIN_UNLIMITED_BONE_INFLUENCE
 	FGPUSkinIndexAndWeight IndexAndWeights;
@@ -127,15 +127,11 @@ struct FVertexFactoryIntermediates
 	float3 UnpackedPosition;
 	
 	// Tangent Basis (includes DeltaTangentZ if GPUSKIN_MORPH_BLEND)
-	float3x3 TangentToLocal; 
-
-	// Vertex Color
-	float4 Color;
+	half3x3 TangentToLocal; 
 
 #if GPUSKIN_APEX_CLOTH
     float3 SimulatedPosition;
 #endif
-
 };
 
 #if GPUSKIN_UNLIMITED_BONE_INFLUENCE
@@ -157,9 +153,6 @@ uint NumVertices;
 uint InputWeightStart;
 uint InputWeightStride;
 
-// see r.SkinCache.Debug, default is 1.0f
-float SkinCacheDebug;
-
 FVertexUnpacked UnpackedVertex(uint MorphIndex, uint ClothIndex, uint ClothIndexOffset, uint VertexIndex, uint WeightOffset)
 {
 	FVertexUnpacked Unpacked;
@@ -241,8 +234,8 @@ void SkinCacheUpdateBatchCS(uint3 GroupID : SV_GroupID,
 	Intermediates.UnpackedPosition.y = PositionInputBuffer[InputOffset * 3 + 1];
 	Intermediates.UnpackedPosition.z = PositionInputBuffer[InputOffset * 3 + 2];
 
-	float3 LocalTangentX = Unpacked.TangentX.xyz;
-	float3 LocalTangentZ = Unpacked.TangentZ.xyz;
+	half3 LocalTangentX = Unpacked.TangentX;
+	half3 LocalTangentZ = Unpacked.TangentZ;
 
 #if GPUSKIN_MORPH_BLEND
 	{
@@ -348,7 +341,7 @@ void SkinCacheUpdateBatchCS(uint3 GroupID : SV_GroupID,
 #endif // GPUSKIN_APEX_CLOTH
 #endif // GPUSKIN_MORPH_BLEND
 
-	float3 SPos = SkinPosition(Intermediates) * SkinCacheDebug;
+	float3 SPos = SkinPosition(Intermediates);
 
 #if GPUSKIN_APEX_CLOTH
 	// Transform simulated position from cloth space into local space and blend with skinned position
@@ -356,8 +349,8 @@ void SkinCacheUpdateBatchCS(uint3 GroupID : SV_GroupID,
     SPos = lerp(SPos.xyz, TransformedSimulatedPos.xyz, ClothWeight);
 #endif
 
-	float3 TangentX = normalize(mul((float3x3)Intermediates.BlendMatrix, LocalTangentX));
-	float3 TangentZ = normalize(mul((float3x3)Intermediates.BlendMatrix, LocalTangentZ));
+	half3 TangentX = normalize(mul((half3x3)Intermediates.BlendMatrix, LocalTangentX));
+	half3 TangentZ = normalize(mul((half3x3)Intermediates.BlendMatrix, LocalTangentZ));
 	
 #if GPUSKIN_APEX_CLOTH
     TangentX = ClothTangentX * ClothWeight + TangentX * (1.0f - ClothWeight);
@@ -368,8 +361,8 @@ void SkinCacheUpdateBatchCS(uint3 GroupID : SV_GroupID,
 	PositionBufferUAV[OutputOffset * 3 + 1] = SPos.y;
 	PositionBufferUAV[OutputOffset * 3 + 2] = SPos.z;
 
-	TangentBufferUAV[2 * OutputOffset + GPUSKIN_RWBUFFER_OFFSET_TANGENT_X] = TangentUnbias_SkinCache(float4(TangentX, Unpacked.TangentX.w));
-	TangentBufferUAV[2 * OutputOffset + GPUSKIN_RWBUFFER_OFFSET_TANGENT_Z] = TangentUnbias_SkinCache(float4(TangentZ, Unpacked.TangentZ.w));
+	TangentBufferUAV[2 * OutputOffset + GPUSKIN_RWBUFFER_OFFSET_TANGENT_X] = TangentUnbias_SkinCache(half4(TangentX, Unpacked.TangentX.w));
+	TangentBufferUAV[2 * OutputOffset + GPUSKIN_RWBUFFER_OFFSET_TANGENT_Z] = TangentUnbias_SkinCache(half4(TangentZ, Unpacked.TangentZ.w));
 
 #if GPUSKIN_DEBUG_PASSTHROUGH
 	// Passthrough debug code
@@ -379,4 +372,4 @@ void SkinCacheUpdateBatchCS(uint3 GroupID : SV_GroupID,
 	TangentBufferUAV[2 * OutputOffset + GPUSKIN_RWBUFFER_OFFSET_TANGENT_X] = TangentInputBuffer[2 * InputOffset + GPUSKIN_VB_OFFSET_TANGENT_X];
 	TangentBufferUAV[2 * OutputOffset + GPUSKIN_RWBUFFER_OFFSET_TANGENT_Z] = TangentInputBuffer[2 * InputOffset + GPUSKIN_VB_OFFSET_TANGENT_Z];
 #endif
-}
+}
\ No newline at end of file
diff --git a/Engine/Shaders/Private/GpuSkinCommon.ush b/Engine/Shaders/Private/GpuSkinCommon.ush
--- a/Engine/Shaders/Private/GpuSkinCommon.ush
+++ b/Engine/Shaders/Private/GpuSkinCommon.ush
@@ -67,20 +67,20 @@ uint4 UnpackU8x4(uint Packed)
 }
 
 // @return 0..1
-float4 UnpackU8x4N(uint Packed)
+half4 UnpackU8x4N(uint Packed)
 {
 	return UnpackU8x4(Packed) / 255.0f;
 }
 
 // @return -1..1
-float4 UnpackS8x4N(uint Packed)
+half4 UnpackS8x4N(uint Packed)
 {
 	return UnpackU8x4(Packed) / 127.5f - 1;
 }
 
 // @param Unpacked -1..1
 // @return layout 0xWWZZYYXX where 0..0xff is mapped to -1..1
-uint PackS8x4N(float4 Unpacked)
+uint PackS8x4N(half4 Unpacked)
 {
 	uint Packed = asuint(
 		(uint((Unpacked.x + 1.0f) * 127.499f) << 0) |
@@ -92,7 +92,7 @@ uint PackS8x4N(float4 Unpacked)
 
 // @param Unpacked -1..1
 // @return layout 0x00ZZYYXX where 0..0xff is mapped to -1..1
-uint PackS8x3N(float3 Unpacked)
+uint PackS8x3N(half3 Unpacked)
 {
 	uint Packed = asuint(
 		(uint((Unpacked.x + 1.0f) * 127.499f) << 0) |
@@ -130,11 +130,11 @@ int4 GetBlendIndices(Buffer<uint> InWeightBuffer, uint StreamOffset, uint ExtraI
 	return UnpackedBlendIndices;
 }
 
-float4 GetBlendWeights(Buffer<uint> InWeightBuffer, uint StreamOffset, uint ExtraInfluenceIndex, uint InNumBoneInfluences)
+half4 GetBlendWeights(Buffer<uint> InWeightBuffer, uint StreamOffset, uint ExtraInfluenceIndex, uint InNumBoneInfluences)
 {
 	StreamOffset += ((InNumBoneInfluences + 3) / 4) * (1 + GPUSKIN_BONE_INDEX_UINT16);
 
-	float4 UnpackedBlendWeights = 0;
+	half4 UnpackedBlendWeights = 0;
 #if GPUSKIN_BONE_WEIGHTS_UINT16
 	StreamOffset += ExtraInfluenceIndex * 2;
 	UnpackedBlendWeights.xy = UnpackUnorm2x16(InWeightBuffer[StreamOffset   ]);
@@ -142,7 +142,7 @@ float4 GetBlendWeights(Buffer<uint> InWeightBuffer, uint StreamOffset, uint Extr
 #else
 	StreamOffset += ExtraInfluenceIndex;
 	UnpackedBlendWeights = UnpackRGBA8(InWeightBuffer[StreamOffset]);
-#endif	
+#endif
 	return UnpackedBlendWeights;
 }
 
@@ -153,8 +153,8 @@ struct FGPUSkinIndexAndWeight
 {
 	int4 BlendIndices;
 	int4 BlendIndices2;
-	float4 BlendWeights;
-	float4 BlendWeights2;
+	half4 BlendWeights;
+	half4 BlendWeights2;
 };
 
 FGPUSkinIndexAndWeight GetBlendIndicesAndWeights(Buffer<uint> InWeightBuffer, uint InWeightOffset)
@@ -239,28 +239,28 @@ FBoneMatrix ComputeBoneMatrixWithUnlimitedInfluences(Buffer<float4> InBoneMatric
 		{
 			BoneIndex = InVertexWeights[BoneIndexOffset + 1] << 8 | BoneIndex;
 			//@todo-lh: Workaround to fix issue in SPIRVEmitter of DXC; this block must be inside the if branch
-			float BoneWeight;
+			half BoneWeight;
 			if (BoneWeightSize > 1)
 			{
-				BoneWeight = float(InVertexWeights[BoneWeightOffset + 1] << 8 | InVertexWeights[BoneWeightOffset]) / 65535.0;
+				BoneWeight = half(InVertexWeights[BoneWeightOffset + 1] << 8 | InVertexWeights[BoneWeightOffset]) / 65535.0;
 			}
 			else
 			{
-				BoneWeight = float(InVertexWeights[BoneWeightOffset]) / 255.0f;
+				BoneWeight = half(InVertexWeights[BoneWeightOffset]) / 255.0f;
 			}
 			BoneMatrix += BoneWeight * InternalGetBoneMatrix(InBoneMatrices, BoneIndex);
 		}
 		else
 		{
 			//@todo-lh: Workaround to fix issue in SPIRVEmitter of DXC; this block must be inside the if branch
-			float BoneWeight;
+			half BoneWeight;
 			if (BoneWeightSize > 1)
 			{
-				BoneWeight = float(InVertexWeights[BoneWeightOffset + 1] << 8 | InVertexWeights[BoneWeightOffset]) / 65535.0;
+				BoneWeight = half(InVertexWeights[BoneWeightOffset + 1] << 8 | InVertexWeights[BoneWeightOffset]) / 65535.0;
 			}
 			else
 			{
-				BoneWeight = float(InVertexWeights[BoneWeightOffset]) / 255.0f;
+				BoneWeight = half(InVertexWeights[BoneWeightOffset]) / 255.0f;
 			}
 			BoneMatrix += BoneWeight * InternalGetBoneMatrix(InBoneMatrices, BoneIndex);
 		}
diff --git a/Engine/Source/Runtime/Engine/Private/GPUSkinCache.cpp b/Engine/Source/Runtime/Engine/Private/GPUSkinCache.cpp
--- a/Engine/Source/Runtime/Engine/Private/GPUSkinCache.cpp
+++ b/Engine/Source/Runtime/Engine/Private/GPUSkinCache.cpp
@@ -28,6 +28,8 @@ GPUSkinCache.cpp: Performs skinning on a compute shader into a buffer to avoid v
 #include "Algo/Sort.h"
 #include "ComponentRecreateRenderStateContext.h"
 #include "RenderGraphUtils.h"
+#include "DataDrivenShaderPlatformInfo.h"
+#include "ShaderCompilerCore.h"
 
 DECLARE_DWORD_COUNTER_STAT(TEXT("Num Dispatches"), STAT_GPUSkinCache_NumDispatches, STATGROUP_GPUSkinCache);
 DECLARE_DWORD_COUNTER_STAT(TEXT("Num Buffers"), STAT_GPUSkinCache_NumBuffers, STATGROUP_GPUSkinCache);
@@ -1048,6 +1050,11 @@ public:
 		OutEnvironment.SetDefine(TEXT("GPUSKIN_RWBUFFER_OFFSET_TANGENT_Z"), FGPUSkinCache::RWTangentZOffsetInFloats);
 		OutEnvironment.SetDefine(TEXT("GPUSKIN_BONE_INDEX_UINT16"), bBoneIndex16);
 		OutEnvironment.SetDefine(TEXT("GPUSKIN_BONE_WEIGHTS_UINT16"), bBoneWeights16);
+
+		if (FDataDrivenShaderPlatformInfo::GetSupportsRealTypes(Parameters.Platform) == ERHIFeatureSupport::RuntimeGuaranteed)
+		{
+			OutEnvironment.CompilerFlags.Add(CFLAG_AllowRealTypes);
+		}
 	}
 
 	TGPUSkinCacheCS(const ShaderMetaType::CompiledShaderInitializerType& Initializer)
