[DOF] Fix bokeh squeeze regression in scatter pass pixel shader.

diff --git a/Engine/Shaders/Private/DiaphragmDOF/DOFHybridScatterPixelShader.usf b/Engine/Shaders/Private/DiaphragmDOF/DOFHybridScatterPixelShader.usf
--- a/Engine/Shaders/Private/DiaphragmDOF/DOFHybridScatterPixelShader.usf
+++ b/Engine/Shaders/Private/DiaphragmDOF/DOFHybridScatterPixelShader.usf
@@ -119,8 +119,6 @@ void ScatterMainPS(
 	);
 #endif
 
-	half CocSqueeze = half(rcp(CocInvSqueeze));
-
 	// Evaluate all bokeh of the scattering group.
 	// Dual Pixel Vectorization: evaluate 2 pixels at a time using packed 16-bit float instructions on supported platforms.
 	UNROLL
@@ -136,7 +134,7 @@ void ScatterMainPS(
 		PixelOffset = dpv_mul_vec_mat(PixelOffset, PetzvalTransform);
 #endif
 
-		PixelOffset = dpv_scale(PixelOffset, half2(CocSqueeze, 1));
+		PixelOffset[0] *= half(CocSqueeze);
 
 		half2 Intersection;
 		#if DIM_BOKEH_SIMULATION
