[Substrate] Change Substrate SSS diffusion fallback to none rather than wrap to minimize look diffferences.

SSS Diffusion is not supported on Forward/Translucent surface. Until now it was fallbacking onto Wrap, but this caused hue shift. So instead we now fallback onto none (i.e. diffuse look).

diff --git a/Engine/Shaders/Private/Substrate/Substrate.ush b/Engine/Shaders/Private/Substrate/Substrate.ush
--- a/Engine/Shaders/Private/Substrate/Substrate.ush
+++ b/Engine/Shaders/Private/Substrate/Substrate.ush
@@ -1505,11 +1505,11 @@ FSubstrateData InternalGetSubstrateSlabBSDF(
 	const bool bIsThin = bIsAtTheBottomOfTopology && (bMFPPlugged && bIsThinSurface); // Enable thin only if MFP is plugged, as otherwise there will be no back-tramission
 
 	// In forward mode (blend mode != Opaque/Masked), diffusion is not supported. 
-	// In such a case we fallback onto wrap mode, as look-wise, this is the closest one.
+	// In such a case we fallback onto no SSS (i.e., diffuse), as look-wise, this is the closest one.
 	#if SUBSTRATE_OPAQUE_MATERIAL == 0
 	if (SSSType == SSS_TYPE_DIFFUSION || SSSType == SSS_TYPE_DIFFUSION_PROFILE)
 	{
-		SSSType = SSS_TYPE_WRAP;
+		SSSType = SSS_TYPE_NONE;
 	}
 	#endif
 
