Metal - Apply depth bias to Basepass Pixel Shader that use OUTPUT_PIXEL_DEPTH_OFFSET
This fix a Black Hair Rendering corruption in some scenario

#tests local testing

diff --git a/Engine/Shaders/Private/BasePassPixelShader.usf b/Engine/Shaders/Private/BasePassPixelShader.usf
--- a/Engine/Shaders/Private/BasePassPixelShader.usf
+++ b/Engine/Shaders/Private/BasePassPixelShader.usf
@@ -900,6 +900,10 @@ void FPixelShaderInOut_MainPS(
 
 #if OUTPUT_PIXEL_DEPTH_OFFSET
 	ApplyPixelDepthOffsetForBasePass(MaterialParameters, PixelMaterialInputs, BasePassInterpolants, Out.Depth);
+	
+	#if APPLE_DEPTH_BIAS_HACK
+	Out.Depth -= APPLE_DEPTH_BIAS_VALUE;
+	#endif
 #endif
 
 	//Clip if the blend mode requires it.
