VulkanRHI: Fix Vulkan shader compilation: Move declaration before first use

#rnx

diff --git a/Engine/Shaders/Private/VirtualShadowMaps/VirtualShadowMapPageMarking.usf b/Engine/Shaders/Private/VirtualShadowMaps/VirtualShadowMapPageMarking.usf
--- a/Engine/Shaders/Private/VirtualShadowMaps/VirtualShadowMapPageMarking.usf
+++ b/Engine/Shaders/Private/VirtualShadowMaps/VirtualShadowMapPageMarking.usf
@@ -40,6 +40,10 @@ uint bCullBackfacingPixels;
 uint NumDirectionalLightSmInds;
 uint MipModeLocal;
 
+#ifdef GeneratePageFlagsFromFroxelsCS
+int PassId;
+#endif
+
 int GetBiasedClipmapLevel(FVirtualShadowMapProjectionShaderData BaseProjectionData, float3 TranslatedWorldPosition, float ExtraBias)
 {
 	float BiasedLevel = CalcAbsoluteClipmapLevel(BaseProjectionData, TranslatedWorldPosition) + BaseProjectionData.ResolutionLodBias + VirtualShadowMap.GlobalResolutionLodBias + ExtraBias;
@@ -743,7 +747,6 @@ FCulledLightsGridHeader VirtualShadowMapGetLightsGridForFroxel(int3 FroxelCoord,
 }
 
 uint bShouldMarkLocaLights;
-int PassId;
 float DebugRange;
 
 [numthreads(FROXEL_INDIRECT_ARG_WORKGROUP_SIZE, 1, 1)]
