[VSM] Removed dependency on "GlobalBilinearClampedSampler" as it is loaded from the View UB on some platforms.

diff --git a/Engine/Shaders/Private/VirtualShadowMaps/VirtualShadowMapPageAccessCommon.ush b/Engine/Shaders/Private/VirtualShadowMaps/VirtualShadowMapPageAccessCommon.ush
--- a/Engine/Shaders/Private/VirtualShadowMaps/VirtualShadowMapPageAccessCommon.ush
+++ b/Engine/Shaders/Private/VirtualShadowMaps/VirtualShadowMapPageAccessCommon.ush
@@ -362,4 +362,25 @@ uint GetVirtualShadowMapStaticArrayIndex()
 uint GetVirtualShadowMapHZBArrayIndex(bool bUseStaticOcclusion)
 {
 	return select(bUseStaticOcclusion, VirtualShadowMap.StaticHZBArrayIndex, 0);
-}
\ No newline at end of file
+}
+
+// Gather a 2x2 footprint from a page table texture.
+uint4 GatherPageTable(Texture2D<uint> PageTableTexture, uint2 TexelCoord, uint HMipLevel, float UVScaleFactor)
+{
+	uint2 TexelCoordMip = TexelCoord >> HMipLevel;
+#if COMPILER_SUPPORTS_GATHER_LOD_RED
+	// Offset to 2x2 footprint center and scale to UV space
+	float MipScale = float(1U << HMipLevel);
+	float2 UV = float2(TexelCoordMip + uint2(1U, 1U)) * VirtualShadowMap.PageTableTextureSizeInvSize.zw * MipScale * UVScaleFactor;
+	uint4 Result2x2 = PageTableTexture.GatherLODRed(VirtualShadowMap.PageTableSampler, UV, HMipLevel);
+#else
+	uint4 TexelRectMip = uint4(TexelCoordMip.xy, TexelCoordMip.xy + uint2(1U, 1U));
+	uint4 Result2x2 = uint4(
+		PageTableTexture.Load(uint3(TexelRectMip.xw, HMipLevel)),	// (-, +)
+		PageTableTexture.Load(uint3(TexelRectMip.zw, HMipLevel)),	// (+, +)
+		PageTableTexture.Load(uint3(TexelRectMip.zy, HMipLevel)),	// (+, -)
+		PageTableTexture.Load(uint3(TexelRectMip.xy, HMipLevel))	// (-, -)
+	);
+#endif
+	return Result2x2;
+}
diff --git a/Engine/Shaders/Private/VirtualShadowMaps/VirtualShadowMapPageManagement.usf b/Engine/Shaders/Private/VirtualShadowMaps/VirtualShadowMapPageManagement.usf
--- a/Engine/Shaders/Private/VirtualShadowMaps/VirtualShadowMapPageManagement.usf
+++ b/Engine/Shaders/Private/VirtualShadowMaps/VirtualShadowMapPageManagement.usf
@@ -153,20 +153,7 @@ void ProcessMaskMip4x4(RWTexture2D<uint> OutPageReceiverMaskMip, inout uint Mask
 // Gather a 2x2 footprint given the upper-left coordinate
 uint4 GatherReceiverMask(uint2 TexelCoord)
 {
-#if COMPILER_SUPPORTS_GATHER_LOD_RED
-	// Offset to 2x2 footprint center and scale to UV space
-	float2 UV = float2(TexelCoord + uint2(1U, 1U)) * VirtualShadowMap.PageTableTextureSizeInvSize.zw * 0.5f;
-	uint4 Flags2x2 = InPageReceiverMasks.GatherLODRed(GlobalBilinearClampedSampler, UV, 0);
-#else
-	uint4 TexelRect = uint4(TexelCoord.xy, TexelCoord.xy + uint2(1U, 1U));
-	uint4 Flags2x2 = uint4(
-		InPageReceiverMasks.Load(uint3(TexelRect.xw, 0)),	// (-, +)
-		InPageReceiverMasks.Load(uint3(TexelRect.zw, 0)),	// (+, +)
-		InPageReceiverMasks.Load(uint3(TexelRect.zy, 0)),	// (+, -)
-		InPageReceiverMasks.Load(uint3(TexelRect.xy, 0))		// (-, -)
-	);
-#endif
-	return Flags2x2;
+	return GatherPageTable(InPageReceiverMasks, TexelCoord, 0u, 0.5f);
 }
 
 /**
diff --git a/Engine/Shaders/Private/VirtualShadowMaps/VirtualShadowMapPageOverlap.ush b/Engine/Shaders/Private/VirtualShadowMaps/VirtualShadowMapPageOverlap.ush
--- a/Engine/Shaders/Private/VirtualShadowMaps/VirtualShadowMapPageOverlap.ush
+++ b/Engine/Shaders/Private/VirtualShadowMaps/VirtualShadowMapPageOverlap.ush
@@ -61,35 +61,11 @@ FScreenRect VirtualShadowMapGetUncachedScreenRect(FScreenRect ScreenRect, uint V
 	return VirtualShadowMapClipScreenRect(ScreenRect, VirtualShadowMapId, MipLevel, VirtualShadowMap.UncachedPageRectBounds);
 }
 
-
-// Gather a 2x2 footprint given the upper-left coordinate
-uint4 GatherPageTable(Texture2D<uint> PageTableTexture, uint2 TexelCoord, uint HMipLevel, float UVScaleFactor = 1.0f)
-{
-	uint2 TexelCoordMip = TexelCoord >> HMipLevel;
-#if COMPILER_SUPPORTS_GATHER_LOD_RED
-	// Offset to 2x2 footprint center and scale to UV space
-	float MipScale = float(1U << HMipLevel);
-	float2 UV = float2(TexelCoordMip + uint2(1U, 1U)) * VirtualShadowMap.PageTableTextureSizeInvSize.zw * MipScale * UVScaleFactor;
-	uint4 Result2x2 = PageTableTexture.GatherLODRed(GlobalBilinearClampedSampler, UV, HMipLevel);
-#else
-	uint4 TexelRectMip = uint4(TexelCoordMip.xy, TexelCoordMip.xy + uint2(1U, 1U));
-	uint4 Result2x2 = uint4(
-		PageTableTexture.Load(uint3(TexelRectMip.xw, HMipLevel)),	// (-, +)
-		PageTableTexture.Load(uint3(TexelRectMip.zw, HMipLevel)),	// (+, +)
-		PageTableTexture.Load(uint3(TexelRectMip.zy, HMipLevel)),	// (+, -)
-		PageTableTexture.Load(uint3(TexelRectMip.xy, HMipLevel))	// (-, -)
-	);
-#endif
-	return Result2x2;
-}
-
-// Gather a 2x2 footprint given the upper-left coordinate
 uint4 GatherPageFlags(uint2 TexelCoord, uint HMipLevel)
 {
 	return GatherPageTable(VirtualShadowMap.PageFlags, TexelCoord, HMipLevel, 1.0f);
 }
 
-// Gather a 2x2 footprint given the upper-left coordinate
 uint4 GatherReceiverMask(uint2 TexelCoord, uint HMipLevel)
 {
 	// Note: UV scaled by extra 0.5f to account for 2x size of mask texture
diff --git a/Engine/Source/Runtime/Renderer/Private/VirtualShadowMaps/VirtualShadowMapArray.cpp b/Engine/Source/Runtime/Renderer/Private/VirtualShadowMaps/VirtualShadowMapArray.cpp
--- a/Engine/Source/Runtime/Renderer/Private/VirtualShadowMaps/VirtualShadowMapArray.cpp
+++ b/Engine/Source/Runtime/Renderer/Private/VirtualShadowMaps/VirtualShadowMapArray.cpp
@@ -662,6 +662,7 @@ void FVirtualShadowMapArray::Initialize(
 
 	UniformParameters.SMRTHairRayCount = CVarSMRTRayCountHair.GetValueOnRenderThread();
 
+	UniformParameters.PageTableSampler = TStaticSamplerState<SF_Bilinear, AM_Clamp, AM_Clamp, AM_Clamp>::GetRHI();
 	// Reference dummy data in the UB initially
 	const uint32 DummyPageTableElement = 0xFFFFFFFF;
 	UniformParameters.PageTable = GSystemTextures.GetDefaultTexture2D(GraphBuilder, PF_R32_UINT, DummyPageTableElement);
diff --git a/Engine/Source/Runtime/Renderer/Private/VirtualShadowMaps/VirtualShadowMapArray.h b/Engine/Source/Runtime/Renderer/Private/VirtualShadowMaps/VirtualShadowMapArray.h
--- a/Engine/Source/Runtime/Renderer/Private/VirtualShadowMaps/VirtualShadowMapArray.h
+++ b/Engine/Source/Runtime/Renderer/Private/VirtualShadowMaps/VirtualShadowMapArray.h
@@ -197,6 +197,8 @@ BEGIN_GLOBAL_SHADER_PARAMETER_STRUCT(FVirtualShadowMapUniformParameters, )
 	SHADER_PARAMETER(uint32, SMRTHairRayCount)
 
 	SHADER_PARAMETER_RDG_BUFFER_SRV(ByteAddressBuffer, ProjectionData)
+
+	SHADER_PARAMETER_SAMPLER(SamplerState, PageTableSampler)
 	SHADER_PARAMETER_RDG_TEXTURE(Texture2D<uint>, PageTable)
 	SHADER_PARAMETER_RDG_TEXTURE(Texture2D<uint>, PageFlags)
 	SHADER_PARAMETER_RDG_TEXTURE(Texture2D<uint>, PageReceiverMasks)
diff --git a/Engine/Source/Runtime/Renderer/Private/VirtualShadowMaps/VirtualShadowMapCacheManager.cpp b/Engine/Source/Runtime/Renderer/Private/VirtualShadowMaps/VirtualShadowMapCacheManager.cpp
--- a/Engine/Source/Runtime/Renderer/Private/VirtualShadowMaps/VirtualShadowMapCacheManager.cpp
+++ b/Engine/Source/Runtime/Renderer/Private/VirtualShadowMaps/VirtualShadowMapCacheManager.cpp
@@ -1632,6 +1632,8 @@ FVirtualShadowMapArrayCacheManager::FInvalidationPassCommon FVirtualShadowMapArr
 		{
 			return GraphBuilder.CreateSRV(GraphBuilder.RegisterExternalBuffer(Buffer, Name));
 		};
+
+		UniformParameters->PageTableSampler = TStaticSamplerState<SF_Bilinear, AM_Clamp, AM_Clamp, AM_Clamp>::GetRHI();
 		UniformParameters->ProjectionData = RegExtCreateSrv(PrevBuffers.ProjectionData, TEXT("Shadow.Virtual.PrevProjectionData"));
 		UniformParameters->PageTable = GraphBuilder.RegisterExternalTexture(PrevBuffers.PageTable, TEXT("Shadow.Virtual.PrevPageTable"));
 		UniformParameters->PageFlags = GraphBuilder.RegisterExternalTexture(PrevBuffers.PageFlags, TEXT("Shadow.Virtual.PrevPageFlags"));
