[TSR] Improve thin geometry detection with depth map and add thin geometry visualization.

Coverage map should only consider pixels that has large depth change. If it is a plane or soft curved surface, no need to infer the coverage map. This leads to
1. Better partial coverage estimation of dense foliages.
2. Remove frost artifact on water if foliage is below. Foliages underwater can be removed and not affecting the water surface denoising.

#tests Done in both project and nanite research project

diff --git a/Engine/Shaders/Private/TemporalSuperResolution/TSRDetectThinGeometry.usf b/Engine/Shaders/Private/TemporalSuperResolution/TSRDetectThinGeometry.usf
--- a/Engine/Shaders/Private/TemporalSuperResolution/TSRDetectThinGeometry.usf
+++ b/Engine/Shaders/Private/TemporalSuperResolution/TSRDetectThinGeometry.usf
@@ -97,6 +97,7 @@ void FetchInput(
 	)
 {
 	tsr_tensor_float DepthCache;
+	tsr_tensor_bool bHasReprojectionOffset;
 	UNROLL_N(SIMD_SIZE)
 	for (uint ElementIndex = 0; ElementIndex < SIMD_SIZE; ElementIndex++)
 	{
@@ -104,6 +105,9 @@ void FetchInput(
 		CurrentCoverage.SetElement(ElementIndex, CurrentCoverageTexture[ClampedFetchPixelPos]);
 		ReprojectedHistoryCoverage.SetElement(ElementIndex, ReprojectedHistoryCoverageTexture[ClampedFetchPixelPos]);
 		DepthCache.SetElement(ElementIndex,SceneDepthTexture[ClampedFetchPixelPos]);
+
+		uint EncodedDilateMask = R8Output[int3(ClampedFetchPixelPos,0)];
+		bHasReprojectionOffset.SetElement(ElementIndex, (EncodedDilateMask & 0x80u) != 0u);
 	}
 
 	UNROLL_N(SIMD_SIZE)
@@ -116,6 +120,7 @@ void FetchInput(
 		DeviceZThinGeometryDetectionError.SetElement(ElementIndex, tsr_half(ErrorMultiplier * DeviceZError));
 	}
 
+	CurrentCoverage = select(bHasReprojectionOffset, CurrentCoverage, tsr_tensor_half::Const(0));
 } // FetchInput
 
 
diff --git a/Engine/Source/Runtime/Renderer/Private/PostProcess/TemporalSuperResolution.cpp b/Engine/Source/Runtime/Renderer/Private/PostProcess/TemporalSuperResolution.cpp
--- a/Engine/Source/Runtime/Renderer/Private/PostProcess/TemporalSuperResolution.cpp
+++ b/Engine/Source/Runtime/Renderer/Private/PostProcess/TemporalSuperResolution.cpp
@@ -3349,6 +3349,10 @@ FDefaultTemporalUpscaler::FOutputs AddTemporalSuperResolutionPasses(
 				Tiles[4 * 3 + 1] = Visualize(EVisualizeId::ReprojectionFieldSummary, TEXT("Reprojection Field"));
 				Tiles[4 * 3 + 1].Label = FString::Printf(TEXT("Reprojection Field (r.TSR.Visualize=%d)"), int32(EVisualizeId::ReprojectionFieldOverview));
 				Tiles[4 * 1 + 3] = Visualize(EVisualizeId::AntiFlickering, TEXT("Flickering Temporal Analysis"));
+				if (PassConfig.ThinGeometryDetectionEnable)
+				{
+					Tiles[4 * 2 + 3] = Visualize(EVisualizeId::ThinGeometry, TEXT("ThinGeometry"));
+				}
 			}
 			else if (Visualization == EVisualizeId::ReprojectionFieldOverview)
 			{
