Fix VS/PS linkage error when trying to use the shader complexity view mode in projects where vr.InstancedStereo is set to 1.

Full screen pixel shaders need to declare the FStereoVSOutput interpolator as well. We should move this before SvPosition in the ScreenPassVS output declaration, so that shaders which don't care about this can safely ignore it, but that's a larger change which is covered by a different ticket.

#rnx

diff --git a/Engine/Shaders/Private/ShaderComplexityApplyPixelShader.usf b/Engine/Shaders/Private/ShaderComplexityApplyPixelShader.usf
--- a/Engine/Shaders/Private/ShaderComplexityApplyPixelShader.usf
+++ b/Engine/Shaders/Private/ShaderComplexityApplyPixelShader.usf
@@ -138,6 +138,7 @@ half3 ColorizeComplexityFromMode(float C)
 
 void Main(
 	noperspective float4 UVAndScreenPos : TEXCOORD0,
+	FStereoVSOutput StereoOutput,
 	float4 SvPosition : SV_POSITION,
 	out float4 OutColor : SV_Target0)
 {
