Place Substrate roughness into the same .VGPR variable for all BSDF types.

This helps to reduce the instruction count/branching by 30% and improve the pass by ~10% in Lumen screen probe classification.

diff --git a/Engine/Shaders/Private/Substrate/Substrate.ush b/Engine/Shaders/Private/Substrate/Substrate.ush
--- a/Engine/Shaders/Private/Substrate/Substrate.ush
+++ b/Engine/Shaders/Private/Substrate/Substrate.ush
@@ -652,8 +652,8 @@ struct FSubstrateBSDF
 ///////////////////////////////////////////////////////////////////////////////
 // Hair BSDF members
 #define HAIR_BASECOLOR(X)				X.VGPRs[0].xyz
-#define HAIR_SCATTER(X)					X.VGPRs[0].w
-#define HAIR_ROUGHNESS(X)				X.VGPRs[1].x
+#define HAIR_ROUGHNESS(X)				X.VGPRs[0].w
+#define HAIR_SCATTER(X)					X.VGPRs[1].x
 #define HAIR_SPECULAR(X)				X.VGPRs[1].y
 #define HAIR_BACKLIT(X)					X.VGPRs[1].z
 #define HAIR_COMPLEXTRANSMITTANCE(X)	X.VGPRs[1].w
@@ -661,9 +661,9 @@ struct FSubstrateBSDF
 ///////////////////////////////////////////////////////////////////////////////
 // Eye BSDF members
 #define EYE_DIFFUSEALBEDO(X)			X.VGPRs[0].xyz
+#define EYE_ROUGHNESS(X)				X.VGPRs[0].w
 #define EYE_F0(X)						X.VGPRs[2].w
 #define EYE_F90(X)						X.VGPRs[3].w
-#define EYE_ROUGHNESS(X)				X.VGPRs[0].w
 #define EYE_IRISDISTANCE(X)				X.VGPRs[1].x
 #define EYE_IRISMASK(X)					X.VGPRs[1].y
 #define EYE_SSSPROFILEID(X)				X.VGPRs[1].z
@@ -673,14 +673,14 @@ struct FSubstrateBSDF
 ///////////////////////////////////////////////////////////////////////////////
 // Single Layer Water BSDF members
 #define SLW_BASECOLOR(X)				X.VGPRs[0].xyz
+#define SLW_ROUGHNESS(X)				X.VGPRs[0].w
 #define SLW_METALLIC(X)					X.VGPRs[1].x
 #define SLW_SPECULAR(X)					X.VGPRs[1].y
-#define SLW_ROUGHNESS(X)				X.VGPRs[1].z
 #define SLW_TOPMATERIALOPACITY(X)		X.VGPRs[1].w
 #if SUBSTRATE_INLINE_SINGLELAYERWATER
 #define SLW_WATERALBEDO(X)				X.InlineVGPRs[0].xyz
 #define SLW_WATEREXTINCTION(X)			X.InlineVGPRs[1].xyz
-#define SLW_WATERPHASEG(X)				X.InlineVGPRs[0].w
+#define SLW_WATERPHASEG(X)				X.InlineVGPRs[1].z
 #define SLW_COLORSCALEBEHINDWATER(X)	X.InlineVGPRs[2].xyz
 #endif
 
