[Metal] Fix HairStrandsMesh shader for Metal backend by using primary for-loop iterator instead of separate increments of 'WeightsOffset'.

This is an alternative workaround that would have otherwise been fixed by inserting a nested dummy loop.
This is a more elegant solution as it is equivalent to the code logic without introducing extra loop statements.

#rnx

diff --git a/Engine/Shaders/Private/HairStrands/HairStrandsMesh.usf b/Engine/Shaders/Private/HairStrands/HairStrandsMesh.usf
--- a/Engine/Shaders/Private/HairStrands/HairStrandsMesh.usf
+++ b/Engine/Shaders/Private/HairStrands/HairStrandsMesh.usf
@@ -113,11 +113,12 @@ void MainCS(uint3 DispatchThreadId : SV_DispatchThreadID)
 		return;
 	}
 
-	uint WeightsOffset = SampleIndex * HAIR_RBF_ENTRY_COUNT(MaxSampleCount);
+	uint WeightsBaseOffset = SampleIndex * HAIR_RBF_ENTRY_COUNT(MaxSampleCount);
 	float3 SampleDeformation = float3(0,0,0);
-	for(uint i = 0; i < MaxSampleCount; ++i, ++WeightsOffset)
+	for(uint i = 0; i < MaxSampleCount; ++i)
 	{
-		SampleDeformation += InterpolationWeightsBuffer[WeightsOffset] * (SampleDeformedPositionsBuffer[i] - SampleRestPositionsBuffer[i]).xyz;
+		[REDACTED_EMAIL] - Incrementing '++WeightsBaseOffset' instead of using primary iterator '+i' results in visual artifacts with Metal backend
+		SampleDeformation += InterpolationWeightsBuffer[WeightsBaseOffset + i] * (SampleDeformedPositionsBuffer[i] - SampleRestPositionsBuffer[i]).xyz;
 	}
 
 	OutSampleDeformationsBuffer[SampleIndex] = float4(SampleDeformation,1);
@@ -376,4 +377,4 @@ void MainCS(uint3 DispatchThreadId : SV_DispatchThreadID)
 	}
 }
 
-#endif // SHADER_RESOURCE_TRANSITION
\ No newline at end of file
+#endif // SHADER_RESOURCE_TRANSITION
