Substrate - fixed broken water look.

#tests PC editor dx12
#rnx

diff --git a/Engine/Shaders/Private/Substrate/Substrate.ush b/Engine/Shaders/Private/Substrate/Substrate.ush
--- a/Engine/Shaders/Private/Substrate/Substrate.ush
+++ b/Engine/Shaders/Private/Substrate/Substrate.ush
@@ -680,7 +680,7 @@ struct FSubstrateBSDF
 #if SUBSTRATE_INLINE_SINGLELAYERWATER
 #define SLW_WATERALBEDO(X)				X.InlineVGPRs[0].xyz
 #define SLW_WATEREXTINCTION(X)			X.InlineVGPRs[1].xyz
-#define SLW_WATERPHASEG(X)				X.InlineVGPRs[1].z
+#define SLW_WATERPHASEG(X)				X.InlineVGPRs[1].w
 #define SLW_COLORSCALEBEHINDWATER(X)	X.InlineVGPRs[2].xyz
 #endif
 
