Adds support for ACES 2.0 as an alternative to 1.3 in HDR mode.

* v1.3 remains the current default, but v2.0 can be activated with "r.HDR.Aces.Version 2"
* When compared against OpenColorIO, some differences are present due to our implementation always targetting AP1 as a display limiting color space. This should eventually be user-exposed as part of display characterization.
* The shader implementation is a direct port of aces-core ctls.
* The table resource implementation is based on OpenColorIO v2.4.1. These are global resources updated based on the HDR max luminance cvar.

diff --git a/Engine/Shaders/Private/ACES.ush b/Engine/Shaders/Private/ACES.ush
deleted file mode 100644
diff --git a/Engine/Shaders/Private/ACES/ACESCommon.ush b/Engine/Shaders/Private/ACES/ACESCommon.ush
new file mode 100644
--- /dev/null
+++ b/Engine/Shaders/Private/ACES/ACESCommon.ush
@@ -0,0 +1,300 @@
+// Copyright Epic Games, Inc. All Rights Reserved.
+
+#pragma once
+
+#define WhiteStandard uint
+static const WhiteStandard WhiteStandard_D65 = 0;
+static const WhiteStandard WhiteStandard_D60 = 1;
+static const WhiteStandard WhiteStandard_DCI = 2;
+
+struct FColorSpace
+{
+	WhiteStandard White;
+	float3x3 XYZtoRGB;
+	float3x3 RGBtoXYZ;
+};
+
+static const float3x3 AP0_2_XYZ_MAT =
+{
+	0.9525523959, 0.0000000000, 0.0000936786,
+	0.3439664498, 0.7281660966,-0.0721325464,
+	0.0000000000, 0.0000000000, 1.0088251844,
+};
+
+static const float3x3 XYZ_2_AP0_MAT =
+{
+	 1.0498110175, 0.0000000000,-0.0000974845,
+	-0.4959030231, 1.3733130458, 0.0982400361,
+	 0.0000000000, 0.0000000000, 0.9912520182,
+};
+
+static const float3x3 AP1_2_XYZ_MAT =
+{
+	 0.6624541811, 0.1340042065, 0.1561876870,
+	 0.2722287168, 0.6740817658, 0.0536895174,
+	-0.0055746495, 0.0040607335, 1.0103391003,
+};
+
+static const float3x3 XYZ_2_AP1_MAT =
+{
+	 1.6410233797, -0.3248032942, -0.2364246952,
+	-0.6636628587,  1.6153315917,  0.0167563477,
+	 0.0117218943, -0.0082844420,  0.9883948585,
+};
+
+static const float3x3 AP0_2_AP1_MAT = //mul( AP0_2_XYZ_MAT, XYZ_2_AP1_MAT );
+{
+	 1.4514393161, -0.2365107469, -0.2149285693,
+	-0.0765537734,  1.1762296998, -0.0996759264,
+	 0.0083161484, -0.0060324498,  0.9977163014,
+};
+
+static const float3x3 AP1_2_AP0_MAT = //mul( AP1_2_XYZ_MAT, XYZ_2_AP0_MAT );
+{
+	 0.6954522414,  0.1406786965,  0.1638690622,
+	 0.0447945634,  0.8596711185,  0.0955343182,
+	-0.0055258826,  0.0040252103,  1.0015006723,
+};
+
+static const float3 AP1_RGB2Y =
+{
+	0.2722287168, //AP1_2_XYZ_MAT[0][1],
+	0.6740817658, //AP1_2_XYZ_MAT[1][1],
+	0.0536895174, //AP1_2_XYZ_MAT[2][1]
+};
+
+// REC 709 primaries
+static const float3x3 XYZ_2_sRGB_MAT =
+{
+	 3.2409699419, -1.5373831776, -0.4986107603,
+	-0.9692436363,  1.8759675015,  0.0415550574,
+	 0.0556300797, -0.2039769589,  1.0569715142,
+};
+
+static const float3x3 sRGB_2_XYZ_MAT =
+{
+	0.4123907993, 0.3575843394, 0.1804807884,
+	0.2126390059, 0.7151686788, 0.0721923154,
+	0.0193308187, 0.1191947798, 0.9505321522,
+};
+
+// REC 2020 primaries
+static const float3x3 XYZ_2_Rec2020_MAT =
+{
+	 1.7166511880, -0.3556707838, -0.2533662814,
+	-0.6666843518,  1.6164812366,  0.0157685458,
+	 0.0176398574, -0.0427706133,  0.9421031212,
+};
+
+static const float3x3 Rec2020_2_XYZ_MAT =
+{
+	0.6369580483, 0.1446169036, 0.1688809752,
+	0.2627002120, 0.6779980715, 0.0593017165,
+	0.0000000000, 0.0280726930, 1.0609850577,
+};
+
+// P3, D65 primaries
+static const float3x3 XYZ_2_P3D65_MAT =
+{
+	 2.4934969119, -0.9313836179, -0.4027107845,
+	-0.8294889696,  1.7626640603,  0.0236246858,
+	 0.0358458302, -0.0761723893,  0.9568845240,
+};
+
+static const float3x3 P3D65_2_XYZ_MAT =
+{
+	0.4865709486, 0.2656676932, 0.1982172852,
+	0.2289745641, 0.6917385218, 0.0792869141,
+	0.0000000000, 0.0451133819, 1.0439443689,
+};
+
+// Bradford chromatic adaptation transforms between ACES white point (D60) and sRGB white point (D65)
+static const float3x3 D65_2_D60_CAT =
+{
+	 1.0130349146, 0.0061052578, -0.0149709436,
+	 0.0076982301, 0.9981633521, -0.0050320385,
+	-0.0028413174, 0.0046851567,  0.9245061375,
+};
+
+static const float3x3 D60_2_D65_CAT =
+{
+	 0.9872240087, -0.0061132286, 0.0159532883,
+	-0.0075983718,  1.0018614847, 0.0053300358,
+	 0.0030725771, -0.0050959615, 1.0816806031,
+};
+
+static const float3x3 CAM16_2_XYZ_MAT =
+{
+	 2.0512756811, -1.1400313439,  0.0887556628,
+	 0.4269389763,  0.7005835277, -0.1275225040,
+	-0.0174712779, -0.0384725929,  1.0589468739
+};
+
+static const float3x3 XYZ_2_CAM16_MAT =
+{
+	 0.3640744835,  0.5947008156, 0.04110127349,
+	-0.2222450987,  1.0738554823, 0.14794533610,
+	-0.0020676190,  0.0488260453, 0.95038755696
+};
+
+float rgb_2_saturation(float3 rgb)
+{
+	float minrgb = min(min(rgb.r, rgb.g), rgb.b);
+	float maxrgb = max(max(rgb.r, rgb.g), rgb.b);
+	return (max(maxrgb, 1e-10) - max(minrgb, 1e-10)) / max(maxrgb, 1e-2);
+}
+
+// Transformations from RGB to other color representations
+float rgb_2_hue(float3 rgb)
+{
+	// Returns a geometric hue angle in degrees (0-360) based on RGB values.
+	// For neutral colors, hue is undefined and the function will return a quiet NaN value.
+	float hue;
+	if (rgb[0] == rgb[1] && rgb[1] == rgb[2]) {
+		//hue = FLT_NAN; // RGB triplets where RGB are equal have an undefined hue
+		hue = 0;
+	}
+	else {
+		hue = (180. / PI) * atan2(sqrt(3.0) * (rgb[1] - rgb[2]), 2 * rgb[0] - rgb[1] - rgb[2]);
+	}
+
+	if (hue < 0.) hue = hue + 360;
+
+	return clamp(hue, 0, 360);
+}
+
+float center_hue(float hue, float centerH)
+{
+	float hueCentered = hue - centerH;
+	if (hueCentered < -180.) hueCentered += 360;
+	else if (hueCentered > 180.) hueCentered -= 360;
+	return hueCentered;
+}
+
+// Transformations between CIE XYZ tristimulus values and CIE x,y 
+// chromaticity coordinates
+float3 XYZ_2_xyY(float3 XYZ)
+{
+	float3 xyY;
+	float divisor = (XYZ[0] + XYZ[1] + XYZ[2]);
+	if (divisor == 0.) divisor = 1e-10;
+	xyY[0] = XYZ[0] / divisor;
+	xyY[1] = XYZ[1] / divisor;
+	xyY[2] = XYZ[1];
+
+	return xyY;
+}
+
+float3 xyY_2_XYZ(float3 xyY)
+{
+	float3 XYZ;
+	XYZ[0] = xyY[0] * xyY[2] / max(xyY[1], 1e-10);
+	XYZ[1] = xyY[2];
+	XYZ[2] = (1.0 - xyY[0] - xyY[1]) * xyY[2] / max(xyY[1], 1e-10);
+
+	return XYZ;
+}
+
+
+float3x3 ChromaticAdaptation( float2 src_xy, float2 dst_xy )
+{
+	// Von Kries chromatic adaptation 
+
+	// Bradford
+	const float3x3 ConeResponse =
+	{
+		 0.8951,  0.2664, -0.1614,
+		-0.7502,  1.7135,  0.0367,
+		 0.0389, -0.0685,  1.0296,
+	};
+	const float3x3 InvConeResponse =
+	{
+		 0.9869929, -0.1470543,  0.1599627,
+		 0.4323053,  0.5183603,  0.0492912,
+		-0.0085287,  0.0400428,  0.9684867,
+	};
+
+	float3 src_XYZ = xyY_2_XYZ( float3( src_xy, 1 ) );
+	float3 dst_XYZ = xyY_2_XYZ( float3( dst_xy, 1 ) );
+
+	float3 src_coneResp = mul( ConeResponse, src_XYZ );
+	float3 dst_coneResp = mul( ConeResponse, dst_XYZ );
+
+	float3x3 VonKriesMat =
+	{
+		{ dst_coneResp[0] / src_coneResp[0], 0.0, 0.0 },
+		{ 0.0, dst_coneResp[1] / src_coneResp[1], 0.0 },
+		{ 0.0, 0.0, dst_coneResp[2] / src_coneResp[2] }
+	};
+
+	return mul( InvConeResponse, mul( VonKriesMat, ConeResponse ) );
+}
+
+// ------- Glow module functions
+float glow_fwd(float ycIn, float glowGainIn, float glowMid)
+{
+	float glowGainOut;
+
+	if (ycIn <= 2. / 3. * glowMid) {
+		glowGainOut = glowGainIn;
+	}
+	else if (ycIn >= 2 * glowMid) {
+		glowGainOut = 0;
+	}
+	else {
+		glowGainOut = glowGainIn * (glowMid / ycIn - 0.5);
+	}
+
+	return glowGainOut;
+}
+
+float glow_inv(float ycOut, float glowGainIn, float glowMid)
+{
+	float glowGainOut;
+
+	if (ycOut <= ((1 + glowGainIn) * 2. / 3. * glowMid)) {
+		glowGainOut = -glowGainIn / (1 + glowGainIn);
+	}
+	else if (ycOut >= (2. * glowMid)) {
+		glowGainOut = 0.;
+	}
+	else {
+		glowGainOut = glowGainIn * (glowMid / ycOut - 1. / 2.) / (glowGainIn / 2. - 1.);
+	}
+
+	return glowGainOut;
+}
+
+float rgb_2_yc(float3 rgb, float ycRadiusWeight = 1.75)
+{
+	// Converts RGB to a luminance proxy, here called YC
+	// YC is ~ Y + K * Chroma
+	// Constant YC is a cone-shaped surface in RGB space, with the tip on the 
+	// neutral axis, towards white.
+	// YC is normalized: RGB 1 1 1 maps to YC = 1
+	//
+	// ycRadiusWeight defaults to 1.75, although can be overridden in function 
+	// call to rgb_2_yc
+	// ycRadiusWeight = 1 -> YC for pure cyan, magenta, yellow == YC for neutral 
+	// of same value
+	// ycRadiusWeight = 2 -> YC for pure red, green, blue  == YC for  neutral of 
+	// same value.
+
+	float r = rgb[0];
+	float g = rgb[1];
+	float b = rgb[2];
+
+	float chroma = sqrt(b * (b - g) + g * (g - r) + r * (r - b));
+
+	return (b + g + r + ycRadiusWeight * chroma) / 3.;
+}
+
+float sigmoid_shaper(float x)
+{
+	// Sigmoid function in the range 0 to 1 spanning -2 to +2.
+
+	float t = max(1 - abs(0.5 * x), 0);
+	float y = 1 + sign(x) * (1 - t * t);
+	return 0.5 * y;
+}
+
diff --git a/Engine/Shaders/Private/ACES/ACESDisplayEncoding.ush b/Engine/Shaders/Private/ACES/ACESDisplayEncoding.ush
new file mode 100644
--- /dev/null
+++ b/Engine/Shaders/Private/ACES/ACESDisplayEncoding.ush
@@ -0,0 +1,65 @@
+// Copyright Epic Games, Inc. All Rights Reserved.
+
+// <ACEStransformID>urn:ampas:aces:transformId:v2.0:Lib.Academy.DisplayEncoding.a2.v1</ACEStransformID>
+// <ACESuserName>Display Encoding Functions</ACESuserName>
+
+// 
+// Library File with functions used for the display encoding/decoding steps
+//
+
+#pragma once
+
+#include "ACESOutputTransform.ush"
+
+float3 apply_white_scale( float3 XYZluminance, 
+                            ODTParams PARAMS, 
+                            bool _invert)
+{
+    float3 RGB_w = mul( PARAMS.OUTPUT_XYZ_TO_RGB, PARAMS.XYZ_w_limit );
+    float3 RGB_w_f = (1./referenceLuminance) * RGB_w;
+    float scale = 1. / max( max(RGB_w_f[0], RGB_w_f[1]), RGB_w_f[2]); 
+    // scale factor is equal to 1/largestChannel
+
+    return (_invert ? (1./scale) : scale) * XYZluminance;
+}
+
+float3 clamp_zero_to_peakLuminance( float3 XYZ_in,
+                                    ODTParams PARAMS)
+{
+    float3 rgb = mul(PARAMS.LIMIT_XYZ_TO_RGB, XYZ_in);
+
+    // Clamp to relative peakLuminance in RGB space prior to white scaling
+    rgb = clamp(rgb, 0.0, PARAMS.peakLuminance / referenceLuminance);
+
+    float3 XYZ = mul(PARAMS.LIMIT_RGB_TO_XYZ, rgb);
+
+    return XYZ;
+}
+
+float3 white_limiting(float3 XYZ_in,
+                        ODTParams PARAMS,
+                        bool scale_white,
+                        bool _invert = false)
+{
+    // The white limiting step clamps the output to between zero and the peak
+    // luminance value so that values do not exceed the maximum value of the
+    // tonescale. 
+    // If the creative white differs from the calibration white of the display,
+    // unequal display code values will be required to produce the neutral of
+    // the creative white. Without scaling, one channel would hit the max value
+    // first while the other channels continue to increase, resulting in a hue
+    // shift. To avoid this, the white scaling finds the largest channel and
+    // applies a scale factor to force the point where this channel hits max to
+    // 1.0, assuring that all three channels "fit" within the peak value. In the
+    // inverse direction, the white scaling is removed.
+
+    float3 XYZ = XYZ_in;
+
+    // Clamp to peak luminance
+    XYZ = clamp_zero_to_peakLuminance(XYZ, PARAMS);
+
+    // White point scaling
+    XYZ = scale_white ? apply_white_scale(XYZ, PARAMS, _invert) : XYZ;
+
+    return XYZ;
+}
\ No newline at end of file
diff --git a/Engine/Shaders/Private/ACES/ACESOutputTransform.ush b/Engine/Shaders/Private/ACES/ACESOutputTransform.ush
new file mode 100644
--- /dev/null
+++ b/Engine/Shaders/Private/ACES/ACESOutputTransform.ush
@@ -0,0 +1,926 @@
+// Copyright Epic Games, Inc. All Rights Reserved.
+
+// <ACEStransformID>urn:ampas:aces:transformId:v2.0:Lib.Academy.OutputTransform.a2.v1</ACEStransformID>
+// <ACESuserName>Output Transform Functions</ACESuserName>
+
+//
+// Library File with functions and presets used for the forward and inverse output
+// transform
+//
+
+#pragma once
+
+#include "ACESUtilities.ush"
+#include "ACESTonescale.ush"
+
+// Gamut compression constants
+static const float smoothCusps = 0.12;
+static const float smoothM = 0.27;
+static const float cuspMidBlend = 1.3;
+
+static const float focusGainBlend = 0.3;
+static const float focusAdjustGain = 0.55;
+static const float focusDistance = 1.35;
+static const float focusDistanceScaling = 1.75;
+
+// Constant used in gamut mapping
+static const float compressionThreshold = 0.75;
+
+static const int tableSize = 360;       // add 1 extra entry at end that is to duplicate first entry for wrapped hue
+static const int additionalTableEntries = 2; // allots for extra entries to wrap the hues without special cases
+static const int totalTableSize = tableSize + additionalTableEntries;
+static const int baseIndex = 1; // array index for smallest hue, which is not necessarily a 0.0 hue angle
+
+static const float gammaMinimum = 0.0;
+static const float gammaMaximum = 5.0;
+static const float gammaSearchStep = 0.4;
+static const float gammaAccuracy = 1e-5;
+
+static const float referenceLuminance = 100.;
+
+// CAM Parameters
+static const float L_A = 100.;
+static const float Y_b = 20.;
+
+static const float ac_resp = 1.0;
+static const float ra = 2. * ac_resp;
+static const float ba = 0.05 + (2. - ra);
+
+static const float3 surround = float3(0.9, 0.59, 0.9);
+
+// Matrix for Hellwig inverse
+static const float3x3 panlrcm =
+{
+	460.0,  451.0,  288.0 ,
+	460.0, -891.0, -261.0 ,
+	460.0, -220.0, -6300.0
+};
+
+struct ODTParams
+{
+    float peakLuminance;
+
+    // Tonescale // Set via TSParams structure
+    float n_r; // normalized white in nits (what 1.0 should be)
+    float g;   // surround / contrast
+    float t_1; // shadow toe or flare/glare compensation
+    float c_t;
+    float s_2;
+    float u_2;
+    float m_2;
+
+    // Chroma Compression
+    float limitJmax;
+    float midJ;
+    float model_gamma;
+    float sat;
+    float sat_thr;
+    float compr;
+    float chromaCompressScale;
+
+    float focusDist;
+
+    // Limit
+    float3x3 LIMIT_RGB_TO_XYZ;
+    float3x3 LIMIT_XYZ_TO_RGB;
+    float3 XYZ_w_limit;
+
+    // Output
+    float3x3 OUTPUT_RGB_TO_XYZ;
+    float3x3 OUTPUT_XYZ_TO_RGB;
+    float3 XYZ_w_output;
+
+    float lowerHullGamma;
+
+    Texture2D<float> reach_table;
+    Texture2D<float3> gamut_table;
+    Texture2D<float> gamma_table;
+};
+
+float wrap_to_360(float hue)
+{
+    float y = fmod(hue, 360.);
+    if (y < 0.)
+    {
+        y = y + 360.;
+    }
+    return y;
+}
+
+int hue_position_in_uniform_table(float hue, int table_size)
+{
+    const float wrapped_hue = wrap_to_360(hue);
+    int result = (wrapped_hue / 360. * table_size);
+    return result;
+}
+
+int next_position_in_table(int entry, int table_size)
+{
+    int result = (entry + 1) % table_size;
+    return result;
+}
+
+float base_hue_for_position(int i_lo, int table_size)
+{
+    float result = i_lo * 360. / table_size;
+    return result;
+}
+
+// CAM Functions
+float3 post_adaptation_non_linear_response_compression_forward(float3 RGB, float F_L)
+{
+    float3 F_L_RGB = spow(F_L * abs(RGB) / 100., 0.42);
+    float3 RGB_c = (400. * copysign(1., RGB) * F_L_RGB) / (27.13 + F_L_RGB);
+
+    return RGB_c;
+}
+
+float3 post_adaptation_non_linear_response_compression_inverse(float3 RGB, float F_L)
+{
+    float3 RGB_p = sign(RGB) * 100. / F_L * spow((27.13 * abs(RGB)) / (400. - abs(RGB)), 1.0 / 0.42);
+
+    return RGB_p;
+}
+
+float3 XYZ_to_Hellwig2022_JMh(  float3 XYZ,
+                                float3 XYZ_w)
+{
+
+    float Y_w = XYZ_w[1];
+
+    // Step 0 - Converting CIE XYZ tristimulus values to sharpened RGB values
+    float3 RGB_w = mul(XYZ_2_CAM16_MAT, XYZ_w);
+
+    // Ignore degree of adaptation.
+    float D = 1.0;
+
+    // Viewing condition dependent parameters
+    float k = 1. / (5. * L_A + 1.);
+    float k4 = pow(k, 4.);
+    float F_L = 0.2 * k4 * (5. * L_A) + 0.1 * pow((1. - k4), 2.) * spow(5. * L_A, 1. / 3.);
+    float n = Y_b / Y_w;
+    float z = 1.48 + sqrt(n);
+
+    float3 D_RGB = D * Y_w / RGB_w + 1. - D;
+
+    float3 RGB_wc = D_RGB * RGB_w;
+
+    float3 RGB_aw = post_adaptation_non_linear_response_compression_forward(RGB_wc, F_L);
+    float A_w = ra * RGB_aw[0] + RGB_aw[1] + ba * RGB_aw[2];
+
+    // Step 1 - Converting CIE XYZ tristimulus values to sharpened RGB values
+    float3 RGB = mul(XYZ_2_CAM16_MAT, XYZ);
+
+    // Step 2
+    float3 RGB_c = D_RGB * RGB;
+
+    // Step 3 - apply  forward post-adaptation non-linear response compression
+    float3 RGB_a = post_adaptation_non_linear_response_compression_forward(RGB_c, F_L);
+
+    // Step 4 - Converting to preliminary cartesian  coordinates
+    float a = RGB_a[0] - 12. * RGB_a[1] / 11. + RGB_a[2] / 11.;
+    float b = (RGB_a[0] + RGB_a[1] - 2. * RGB_a[2]) / 9.;
+
+    // Computing the hue angle math h
+    // Unclear why this isn't matching the python version
+    float hr = atan2(b, a);
+    float h = wrap_to_360(radians_to_degrees(hr));
+
+    // Step 6 - Computing achromatic responses for the stimulus
+    float A = ra * RGB_a[0] + RGB_a[1] + ba * RGB_a[2];
+
+    // Step 7 - Computing the correlate of lightness, J
+    float J = 100. * pow(A / A_w, surround[1] * z);
+
+    // Step 9 - Computing the correlate of colourfulness, M
+    float M = (J == 0.0) ? 0.0 : 43. * surround[2] * sqrt(a * a + b * b);
+
+    float3 return_JMh = float3(J, M, h);
+
+    return return_JMh;
+}
+
+float3 Hellwig2022_JMh_to_XYZ(  float3 JMh,
+                                float3 XYZ_w)
+{
+    float J = JMh[0];
+    float M = JMh[1];
+    float h = JMh[2];
+
+    float Y_w = XYZ_w[1];
+
+    // Step 0 - Converting CIE XYZ tristimulus values to sharpened RGB values.
+    float3 RGB_w = mul(XYZ_2_CAM16_MAT, XYZ_w);
+
+    // Ignore degree of adaptation.
+    float D = 1.0;
+
+    // Viewing condition dependent parameters
+    float k = 1.0 / (5.0 * L_A + 1.0);
+    float k4 = pow(k, 4.0);
+    float F_L = 0.2 * k4 * (5.0 * L_A) + 0.1 * pow((1.0 - k4), 2.0) * spow(5.0 * L_A, 1.0 / 3.0);
+    float n = sdiv(Y_b, Y_w);
+    float z = 1.48 + sqrt(n);
+
+    float3 D_RGB = D * Y_w / RGB_w + 1.0 - D;
+
+    float3 RGB_wc = D_RGB * RGB_w;
+
+    float3 RGB_aw = post_adaptation_non_linear_response_compression_forward(RGB_wc, F_L);
+
+    float A_w = ra * RGB_aw[0] + RGB_aw[1] + ba * RGB_aw[2];
+
+    float hr = degrees_to_radians(h);
+
+    // No Helmholtz-Kohlrausch effect
+
+    // Computing achromatic respons A for the stimulus
+    float A = A_w * spow(J / 100.0, 1.0 / (surround[1] * z));
+
+    // Computing P_p_1 to P_p_2
+    float P_p_1 = 43.0 * surround[2];
+    float P_p_2 = A;
+
+    // Step 3 - Computing opponent colour dimensions a and b
+    float gamma = M / P_p_1;
+    float a = gamma * cos(hr);
+    float b = gamma * sin(hr);
+
+    // Step 4 - Applying post-adaptation non-linear response compression matrix
+    float3 vec = {P_p_2, a, b};
+    float3 RGB_a = (1.0 / 1403.0) * mul(panlrcm, vec);
+
+    // Step 5 - Applying the inverse post-adaptation non-linear respnose compression
+    float3 RGB_c = post_adaptation_non_linear_response_compression_inverse(RGB_a, F_L);
+
+    // Step 6
+    float3 RGB = RGB_c / D_RGB;
+
+    // Step 7
+    float3 XYZ = mul(CAM16_2_XYZ_MAT, RGB);
+
+    return XYZ;
+}
+
+float Hellwig_J_to_Y(float J,
+                     float surround = 0.59,
+                     float L_A = 100.,
+                     float Y_b = 20.)
+{
+    // Viewing conditions dependent parameters (could be pre-calculated)
+    float k = 1.0 / (5. * L_A + 1.);
+    float k4 = k * k * k * k;
+    float F_L = 0.2 * k4 * (5. * L_A) + 0.1 * pow((1. - k4), 2.) * pow(5. * L_A, 1. / 3.);
+    float n = Y_b / 100.;
+    float z = 1.48 + sqrt(n);
+    float F_L_W = pow(F_L, 0.42);
+    float A_w = (400. * F_L_W) / (27.13 + F_L_W);
+
+    float A = A_w * pow(abs(J) / 100., 1. / (surround * z));
+
+    return sign(J) * 100. / F_L * pow((27.13 * A) / (400.0 - A), 1. / 0.42);
+}
+
+float Y_to_Hellwig_J(float Y,
+                     float surround = 0.59,
+                     float L_A = 100.,
+                     float Y_b = 20.)
+{
+    // Viewing conditions dependent parameters (could be pre-calculated)
+    float k = 1.0 / (5. * L_A + 1.);
+    float k4 = k * k * k * k;
+    float F_L = 0.2 * k4 * (5. * L_A) + 0.1 * pow((1. - k4), 2.) * pow(5. * L_A, 1. / 3.);
+    float n = Y_b / 100.;
+    float z = 1.48 + sqrt(n);
+    float F_L_W = pow(F_L, 0.42);
+    float A_w = (400. * F_L_W) / (27.13 + F_L_W);
+
+    float F_L_Y = pow(F_L * abs(Y) / 100., 0.42);
+
+    return sign(Y) * 100. * pow(((400. * F_L_Y) / (27.13 + F_L_Y)) / A_w, surround * z);
+}
+
+float3 clamp_AP1(float3 AP1, float peakLuminance)
+{
+    const float upper_clamp_limit = 8. * (128 + 768 * (log(peakLuminance/100.)/log(10000./100.))); 
+    // limit to nice power of 2 (3 stops) above that needed to max out
+    // note the quantity (128 + ...) is the definition of r_hit from the tonescale constants
+
+    float3 ap1_clamped = clamp(AP1, 0., upper_clamp_limit);
+    float3 XYZ_clamped = mul(AP1_2_XYZ_MAT, ap1_clamped);
+
+    return XYZ_clamped;
+}
+
+float3 aces_to_JMh(float3 aces,
+                     float peakLuminance)
+{
+    // AP0 to AP1
+    float3 AP1 = mul(AP0_2_AP1_MAT, aces);
+
+    // Clamp to half float range
+    float3 XYZ = clamp_AP1(AP1, peakLuminance);
+
+    // XYZ to JMh
+    float3 RGB_w = float3(referenceLuminance, referenceLuminance, referenceLuminance);
+    float3 XYZ_w = mul(AP0_2_XYZ_MAT, RGB_w);
+
+    float3 XYZluminance = referenceLuminance * XYZ;
+    float3 JMh = XYZ_to_Hellwig2022_JMh(XYZluminance, XYZ_w);
+
+    return JMh;
+}
+
+ODTParams init_ODTParams(
+    float peakLuminance,
+    float3x3 LIMIT_RGB_TO_XYZ,
+    float3x3 LIMIT_XYZ_TO_RGB,
+    float3x3 OUTPUT_RGB_TO_XYZ,
+    float3x3 OUTPUT_XYZ_TO_RGB,
+    Texture2D<float> reach_table,
+    Texture2D<float3> gamut_table,
+    Texture2D<float> gamma_table
+)
+{
+    TSParams TSPARAMS = init_TSParams(peakLuminance);
+
+    float limitJmax = Y_to_Hellwig_J(peakLuminance);
+    float midJ = Y_to_Hellwig_J(TSPARAMS.c_t * 100.);
+
+    // Chroma compress presets
+    static const float chroma_compress = 2.4;
+    static const float chroma_compress_fact = 3.3;
+    static const float chroma_expand = 1.3;
+    static const float chroma_expand_fact = 0.69;
+    static const float chroma_expand_thr = 0.5;
+
+    // Calculated chroma compress variables
+    const float log_peak = log10(TSPARAMS.n / TSPARAMS.n_r);
+    const float compr = chroma_compress + (chroma_compress * chroma_compress_fact) * log_peak;
+    const float sat = max(0.2, chroma_expand - (chroma_expand * chroma_expand_fact) * log_peak);
+    const float sat_thr = chroma_expand_thr / TSPARAMS.n;
+    const float chromaCompressScale = pow(0.03379 * TSPARAMS.n, 0.30596) - 0.45135;
+
+    const float model_gamma = 1. / (surround[1] * (1.48 + sqrt(Y_b / L_A)));
+
+    const float focusDist = focusDistance +
+                            focusDistance * focusDistanceScaling * log_peak;
+
+    static const float3 RGB_w = float3(100., 100., 100.);
+
+    const float lowerHullGamma = 1.14 + 0.07 * log_peak;
+
+    // Limiting Primaries
+    float3 XYZ_w_limit = mul(LIMIT_RGB_TO_XYZ, RGB_w);
+
+    // Output / Encoding Primaries
+    float3 XYZ_w_output = mul(OUTPUT_RGB_TO_XYZ, RGB_w);
+
+    const ODTParams ODTPARAMS = {
+        peakLuminance,
+
+        // Tonescale
+        TSPARAMS.n_r, // normalized white in nits (what 1.0 should be)
+        TSPARAMS.g,   // surround / contrast
+        TSPARAMS.t_1,
+        TSPARAMS.c_t,
+        TSPARAMS.s_2,
+        TSPARAMS.u_2,
+        TSPARAMS.m_2,
+
+        // Chroma Compression
+        limitJmax,
+        midJ,
+        model_gamma,
+        sat,
+        sat_thr,
+        compr,
+        chromaCompressScale,
+
+        focusDist,
+
+        // Limit
+        LIMIT_RGB_TO_XYZ,
+        LIMIT_XYZ_TO_RGB,
+        XYZ_w_limit,
+
+        // Output
+        LIMIT_RGB_TO_XYZ,
+        OUTPUT_XYZ_TO_RGB,
+        XYZ_w_output,
+
+        lowerHullGamma, // lowerHullGamma
+
+        reach_table,
+        gamut_table,
+        gamma_table
+    };
+
+    return ODTPARAMS;
+}
+
+// Chroma compression
+//
+// Compresses colors inside the gamut with the aim for colorfulness to have an
+// appropriate rate of change from display black to display white, and from
+// achromatic outward to purer colors.
+//
+float chromaCompressionNorm(float h,
+                            ODTParams PARAMS)
+{
+
+    float hr = degrees_to_radians(h);
+
+    float a = cos(hr);
+    float b = sin(hr);
+    float cos_hr2 = a * a - b * b;
+    float sin_hr2 = 2.0 * a * b;
+    float cos_hr3 = 4.0 * a * a * a - 3.0 * a;
+    float sin_hr3 = 3.0 * b - 4.0 * b * b * b;
+
+    float M = 11.34072 * a +
+              16.46899 * cos_hr2 +
+              7.88380 * cos_hr3 +
+              14.66441 * b +
+              -6.37224 * sin_hr2 +
+              9.19364 * sin_hr3 +
+              77.12896;
+
+    return M * PARAMS.chromaCompressScale;
+}
+
+// A "toe" function that remaps the given value x between 0 and limit.
+// The k1 and k2 parameters change the size and shape of the toe.
+// https://www.desmos.com/calculator/6vplvw14ti
+float toe_fwd(float x,
+          float limit,
+          float k1_in,
+          float k2_in)
+{
+    if (x > limit)
+        return x;
+
+    const float k2 = max(k2_in, 0.001);
+    const float k1 = sqrt(k1_in * k1_in + k2 * k2);
+    const float k3 = (limit + k1) / (limit + k2);
+    const float minus_b = k3 * x - k1;
+    const float minus_c = k2 * k3 * x;
+
+    return 0.5 * (minus_b + sqrt(minus_b * minus_b + 4. * minus_c));
+}
+
+float toe_inv(float x,
+          float limit,
+          float k1_in,
+          float k2_in)
+{
+    if (x > limit)
+        return x;
+
+    float k2 = max(k2_in, 0.001);
+    float k1 = sqrt(k1_in * k1_in + k2 * k2);
+    float k3 = (limit + k1) / (limit + k2);
+
+    return (x * x + k1 * x) / (k3 * (x + k2));
+}
+
+float reachMFromTable(float h,
+                      Texture2D<float> reach_table)
+{
+    int i_lo = hue_position_in_uniform_table(h, tableSize);
+    int i_hi = next_position_in_table(i_lo, tableSize);
+
+    float t = (h - i_lo) / (i_hi - i_lo);
+
+    return lerp(reach_table.Load(int3(i_lo, 0, 0)), reach_table.Load(int3(i_hi, 0, 0)), t);
+}
+
+// In-gamut chroma compression
+//
+// Compresses colors inside the gamut with the aim for colorfulness to have an
+// appropriate rate of change from display black to display white, and from
+// achromatic outward to purer colors.
+float chromaCompression(float3 JMh,
+                        float origJ,
+                        ODTParams PARAMS,
+                        bool _invert = false)
+{
+    float J = JMh[0];
+    float M = JMh[1];
+    float h = JMh[2];
+
+    if (M == 0.0)
+    {
+        return M;
+    }
+
+    float nJ = J / PARAMS.limitJmax;
+    float snJ = max(0., 1. - nJ);
+    float Mnorm = chromaCompressionNorm(h, PARAMS);
+    float limit = pow(nJ, PARAMS.model_gamma) * reachMFromTable(h, PARAMS.reach_table) / Mnorm;
+
+    float toe_limit = limit - 0.001;
+    float toe_snJ_sat = snJ * PARAMS.sat;
+    float toe_sqrt_nJ_sat_thr = sqrt(nJ * nJ + PARAMS.sat_thr);
+    float toe_nJ_compr = nJ * PARAMS.compr;
+
+    if (!_invert) { // Forward chroma compression
+        // Rescaling of M with the tonescaled J to get the M to the same range as
+        // J after the tonescale.  The rescaling uses the Hellwig2022 model gamma to
+        // keep the M/J ratio correct (keeping the chromaticities constant).
+        M = M * pow(J / origJ, PARAMS.model_gamma);
+
+        // Normalize M with the rendering space cusp M
+        M = M / Mnorm;
+
+        // Expand the colorfulness by running the toe function in reverse.  The goal is to
+        // expand less saturated colors less and more saturated colors more.  The expansion
+        // increases saturation in the shadows and mid-tones but not in the highlights.
+        // The 0.001 offset starts the expansions slightly above zero.  The sat_thr makes
+        // the toe less aggressive near black to reduce the expansion of noise.
+        M = limit - toe_fwd(limit - M,
+                        toe_limit,
+                        toe_snJ_sat,
+                        toe_sqrt_nJ_sat_thr);
+
+        // Compress the colorfulness.  The goal is to compress less saturated colors more and
+        // more saturated colors less, especially in the highlights.  This step creates the
+        // saturation roll-off in the highlights, but attemps to preserve pure colors.  This
+        // mostly affects highlights and mid-tones, and does not compress shadows.
+        M = toe_fwd(M,
+                limit,
+                toe_nJ_compr,
+                snJ);
+
+        // Denormalize
+        M = M * Mnorm;
+    }
+    else
+    {
+        M = M / Mnorm;
+        M = toe_inv(M,
+                limit,
+                toe_nJ_compr,
+                snJ);
+        M = limit - toe_inv(limit - M,
+                        toe_limit,
+                        toe_snJ_sat,
+                        toe_sqrt_nJ_sat_thr);
+        M = M * Mnorm;
+        M = M * pow(J / origJ, -PARAMS.model_gamma);
+    }
+
+    return M;
+}
+
+float compressionFunction(float v,
+                          float thr,
+                          float lim,
+                          bool _invert = false)
+{
+    float s = (lim - thr) * (1.0 - thr) / (lim - 1.0);
+    float nd = (v - thr) / s;
+
+    float vCompressed;
+
+    if (_invert)
+    {
+        if (v < thr || lim <= 1.0001 || v > thr + s)
+        {
+            vCompressed = v;
+        }
+        else
+        {
+            vCompressed = thr + s * (-nd / (nd - 1));
+        }
+    }
+    else
+    {
+        if (v < thr || lim <= 1.0001)
+        {
+            vCompressed = v;
+        }
+        else
+        {
+            vCompressed = thr + s * nd / (1.0 + nd);
+        }
+    }
+
+    return vCompressed;
+}
+
+int midpoint(int low, int high)
+{
+    return (low + high) / 2;
+}
+
+float2 cuspFromTable(float h, Texture2D<float3> gamut_table)
+{
+    float3 lo;
+    float3 hi;
+
+    int low_i = 0;
+    int high_i = baseIndex + tableSize; // allowed as we have an extra entry in the table
+    int i = hue_position_in_uniform_table(h, tableSize) + baseIndex;
+
+    while (low_i + 1 < high_i)
+    {
+        if (h > gamut_table.Load(int3(i, 0, 0))[2])
+        {
+            low_i = i;
+        }
+        else
+        {
+            high_i = i;
+        }
+        i = midpoint(low_i, high_i);
+    }
+    lo = gamut_table.Load(int3(high_i - 1, 0, 0));
+    hi = gamut_table.Load(int3(high_i, 0, 0));
+
+    float t = (h - lo[2]) / (hi[2] - lo[2]);
+    float cuspJ = lerp(lo[0], hi[0], t);
+    float cuspM = lerp(lo[1], hi[1], t);
+
+    float2 cuspJM = float2(cuspJ, cuspM);
+
+    return cuspJM;
+}
+
+
+float3 tonemapAndCompress_fwd(float3 inputJMh,
+                                ODTParams PARAMS)
+{
+    float3 outputJMh;
+
+    float _linear = Hellwig_J_to_Y(inputJMh[0]) / referenceLuminance;
+
+    TSParams TSPARAMS = init_TSParams(PARAMS.peakLuminance);
+    float luminanceTS = tonescale_fwd(_linear, TSPARAMS);
+
+    float tonemappedJ = Y_to_Hellwig_J(luminanceTS);
+
+    float3 tonemappedJMh = float3(tonemappedJ, inputJMh[1], inputJMh[2]);
+
+    outputJMh = tonemappedJMh;
+    outputJMh[1] = chromaCompression(outputJMh,
+                                     inputJMh[0],
+                                     PARAMS,
+                                     false);
+
+    return outputJMh;
+}
+
+float3 JMh_to_output_XYZ(float3 JMh,
+                           ODTParams PARAMS)
+{
+    float3 XYZluminance = Hellwig2022_JMh_to_XYZ(JMh,
+                                                   PARAMS.XYZ_w_limit);
+
+    float3 XYZ = (1. / referenceLuminance) * XYZluminance;
+
+    return XYZ;
+}
+
+float hueDependentUpperHullGamma(float h,
+                                 Texture2D<float> gamma_table)
+{
+    uint3 gammaSize;// = uint3(totalTableSize, 1, 1);
+    gamma_table.GetDimensions(0, gammaSize.x, gammaSize.y, gammaSize.z);
+
+    const int i_lo = hue_position_in_uniform_table(h, tableSize) + baseIndex;
+    const int i_hi = next_position_in_table(i_lo, (int)gammaSize.x);
+
+    const float base_hue = base_hue_for_position(i_lo - baseIndex, tableSize);
+
+    const float t = wrap_to_360(h) - base_hue;
+
+    return lerp(gamma_table.Load(int3(i_lo, 0, 0)), gamma_table.Load(int3(i_hi, 0, 0)), t);
+}
+
+float getFocusGain(float J,
+                   float cuspJ,
+                   float limitJmax)
+{
+
+    float thr = lerp(cuspJ, limitJmax, focusGainBlend);
+    if (J > thr)
+    {
+        // Approximate inverse required above threshold
+        float gain = (limitJmax - thr) / max(0.0001, limitJmax - J);
+        return pow(log10(gain), 1. / focusAdjustGain) + 1.;
+    }
+    else
+    {
+        // Analytic inverse possible below cusp
+        return 1.;
+    }
+}
+
+float solve_J_intersect(float J,
+                        float M,
+                        float focusJ,
+                        float maxJ,
+                        float slope_gain)
+{
+    float a = M / (focusJ * slope_gain);
+    float b = 0.0;
+    float c = 0.0;
+    float intersectJ = 0.0;
+
+    if (J < focusJ)
+    {
+        b = 1.0 - M / slope_gain;
+    }
+    else
+    {
+        b = -(1.0 + M / slope_gain + maxJ * M / (focusJ * slope_gain));
+    }
+
+    if (J < focusJ)
+    {
+        c = -J;
+    }
+    else
+    {
+        c = maxJ * M / slope_gain + J;
+    }
+
+    float root = sqrt(b * b - 4.0 * a * c);
+
+    if (J < focusJ)
+    {
+        intersectJ = 2.0 * c / (-b - root);
+    }
+    else
+    {
+        intersectJ = 2.0 * c / (-b + root);
+    }
+
+    return intersectJ;
+}
+
+float3 findGamutBoundaryIntersection(float3 JMh_s,
+                                       float2 JM_cusp_in,
+                                       float J_focus,
+                                       float J_max,
+                                       float slope_gain,
+                                       float gamma_top,
+                                       float gamma_bottom)
+{
+    float slope = 0.0;
+
+    float s = max(0.000001, smoothCusps);
+    float2 JM_cusp = JM_cusp_in;
+    JM_cusp[1] = JM_cusp_in[1] * (1.0 + smoothM * s); // M
+
+    float J_intersect_source = solve_J_intersect(JMh_s[0],
+                                                 JMh_s[1],
+                                                 J_focus,
+                                                 J_max,
+                                                 slope_gain);
+    float J_intersect_cusp = solve_J_intersect(JM_cusp[0],
+                                               JM_cusp[1],
+                                               J_focus,
+                                               J_max,
+                                               slope_gain);
+
+    if (J_intersect_source < J_focus)
+    {
+        slope = J_intersect_source * (J_intersect_source - J_focus) / (J_focus * slope_gain);
+    }
+    else
+    {
+        slope = (J_max - J_intersect_source) * (J_intersect_source - J_focus) / (J_focus * slope_gain);
+    }
+
+    float M_boundary_lower = J_intersect_cusp * pow(J_intersect_source / J_intersect_cusp, 1. / gamma_bottom) / (JM_cusp[0] / JM_cusp[1] - slope);
+
+    float M_boundary_upper = JM_cusp[1] * (J_max - J_intersect_cusp) * pow((J_max - J_intersect_source) / (J_max - J_intersect_cusp), 1. / gamma_top) / (slope * JM_cusp[1] + J_max - JM_cusp[0]);
+
+    float M_boundary = JM_cusp[1] * smin(M_boundary_lower / JM_cusp[1], M_boundary_upper / JM_cusp[1], s);
+
+    float J_boundary = J_intersect_source + slope * M_boundary;
+
+    float3 return_JMh = {J_boundary, M_boundary, J_intersect_source};
+
+    return return_JMh;
+}
+
+float3 getReachBoundary(float J,
+                          float M,
+                          float h,
+                          ODTParams PARAMS,
+                          float2 JMcusp,
+                          float focusJ)
+{
+    float limitJmax = PARAMS.limitJmax;
+    float midJ = PARAMS.midJ;
+    float model_gamma = PARAMS.model_gamma;
+    float focusDist = PARAMS.focusDist;
+
+    const float reachMaxM = reachMFromTable(h,  PARAMS.reach_table);
+
+    float slope_gain = limitJmax * focusDist * getFocusGain(J, JMcusp[0], limitJmax);
+
+    float intersectJ = solve_J_intersect(J, M, focusJ, limitJmax, slope_gain);
+    float slope;
+    if (intersectJ < focusJ)
+    {
+        slope = intersectJ * (intersectJ - focusJ) / (focusJ * slope_gain);
+    }
+    else
+    {
+        slope = (limitJmax - intersectJ) * (intersectJ - focusJ) / (focusJ * slope_gain);
+    }
+    float boundary = limitJmax * pow(intersectJ / limitJmax, model_gamma) * reachMaxM / (limitJmax - slope * reachMaxM);
+    float3 result = {J, boundary, h};
+    return result;
+}
+
+float3 compressGamut(float3 JMh,
+                       ODTParams PARAMS,
+                       bool _invert = false)
+{
+    float limitJmax = PARAMS.limitJmax;
+    float midJ = PARAMS.midJ;
+    float focusDist = PARAMS.focusDist;
+    float model_gamma = PARAMS.model_gamma;
+
+    float2 project_from = float2(JMh[0], JMh[1]);
+    float2 JMcusp = cuspFromTable(JMh[2], PARAMS.gamut_table);
+
+    if (JMh[1] < 0.0001 || JMh[0] > limitJmax)
+    {
+        float3 JMh_return = float3(JMh[0], 0.0, JMh[2]);
+        return JMh_return;
+    }
+
+    // Calculate where the out of gamut color is projected to
+    float focusJ = lerp(JMcusp[0], midJ, min(1., cuspMidBlend - (JMcusp[0] / limitJmax)));
+
+    float slope_gain = limitJmax * focusDist * getFocusGain(JMh[0], JMcusp[0], limitJmax);
+
+    // Find gamut intersection
+    float gamma_top = hueDependentUpperHullGamma(JMh[2], PARAMS.gamma_table);
+    float gamma_bottom = PARAMS.lowerHullGamma;
+
+    float3 boundaryReturn = findGamutBoundaryIntersection(JMh,
+                                                            JMcusp,
+                                                            focusJ,
+                                                            limitJmax,
+                                                            slope_gain,
+                                                            gamma_top,
+                                                            gamma_bottom);
+
+    float2 JMboundary = float2(boundaryReturn[0], boundaryReturn[1]);
+    float2 project_to = float2(boundaryReturn[2], 0.0);
+    float projectJ = boundaryReturn[2];
+
+    // Calculate AP1 reach boundary
+    float3 reachBoundary = getReachBoundary(JMboundary[0],
+                                              JMboundary[1],
+                                              JMh[2],
+                                              PARAMS,
+                                              JMcusp,
+                                              focusJ);
+
+    float difference = max(1.0001, reachBoundary[1] / JMboundary[1]);
+    float threshold = max(compressionThreshold, 1. / difference);
+
+    // Compress the out of gamut color along the projection line
+    float v = project_from[1] / JMboundary[1];
+
+    v = compressionFunction(v,
+                            threshold,
+                            difference,
+                            _invert);
+
+    float2 JMcompressed = project_to + v * (JMboundary - project_to);
+    float3 return_JMh = float3(JMcompressed[0], JMcompressed[1], JMh[2]);
+
+    return return_JMh;
+}
+
+float3 gamutMap_fwd(float3 JMh,
+                      ODTParams PARAMS
+                    )
+{
+    return compressGamut(JMh, PARAMS, false);
+}
+
+float3 outputTransform_fwd( float3 aces,
+                            float peakLuminance,
+                            ODTParams PARAMS)
+{
+    float3 JMh = aces_to_JMh(aces, peakLuminance);
+
+    float3 tonemappedJMh = tonemapAndCompress_fwd(JMh, PARAMS);
+
+    float3 compressedJMh = gamutMap_fwd(tonemappedJMh, PARAMS);
+
+    float3 XYZ = JMh_to_output_XYZ(compressedJMh,
+                                      PARAMS);
+
+    return XYZ;
+}
\ No newline at end of file
diff --git a/Engine/Shaders/Private/ACES/ACESTonescale.ush b/Engine/Shaders/Private/ACES/ACESTonescale.ush
new file mode 100644
--- /dev/null
+++ b/Engine/Shaders/Private/ACES/ACESTonescale.ush
@@ -0,0 +1,95 @@
+// Copyright Epic Games, Inc. All Rights Reserved.
+
+// @todo: This should be moved to C++
+
+// <ACEStransformID>urn:ampas:aces:transformId:v2.0:Lib.Academy.Tonescale.a2.v1</ACEStransformID>
+// <ACESuserName>Tonescale Functions</ACESuserName>
+
+// 
+// Library File with functions used for pre-calculating the tonescale parameters and 
+// applying the fwd/inv tonescale function
+//
+
+#pragma once
+
+struct TSParams 
+{
+	float n;
+	float n_r;
+	float g;
+	float t_1;
+	float c_t; 
+	float s_2;
+	float u_2;
+	float m_2;
+};
+
+// Tonescale pre-calculations
+TSParams init_TSParams(float peakLuminance) {
+
+    // Preset constants that set the desired behavior for the curve
+    const float n = peakLuminance;
+    
+    static const float n_r = 100.0;    // normalized white in nits (what 1.0 should be)
+    static const float g = 1.15;       // surround / contrast
+    static const float c = 0.18;	    // anchor for 18% grey
+    static const float c_d = 10.013;   // output luminance of 18% grey (in nits)
+    static const float w_g = 0.14;     // change in grey between different peak luminance
+    static const float t_1 = 0.04;     // shadow toe or flare/glare compensation
+    static const float r_hit_min = 128.;	// scene-referred value "hitting the roof" for SDR (e.g. when n = 100 nits)
+    static const float r_hit_max = 896.;   // scene-referred value "hitting the roof" for when n = 10000 nits
+
+    // Calculate output constants
+    const float r_hit = r_hit_min + (r_hit_max - r_hit_min) * (log(n/n_r)/log(10000./100.));
+    const float m_0 = (n / n_r);
+    const float m_1 = 0.5 * (m_0 + sqrt(m_0 * (m_0 + 4. * t_1)));
+    const float u = pow((r_hit/m_1)/((r_hit/m_1)+1),g);
+    const float m = m_1 / u;
+    const float w_i = log(n/100.)/log(2.);
+    const float c_t = c_d/n_r * (1. + w_i * w_g);
+    const float g_ip = 0.5 * (c_t + sqrt(c_t * (c_t + 4. * t_1)));
+    const float g_ipp2 = -(m_1 * pow((g_ip/m),(1./g))) / (pow( g_ip/m , 1./g)-1.);
+    const float w_2 = c / g_ipp2;
+    const float s_2 = w_2 * m_1;
+    const float u_2 = pow( (r_hit/m_1)/((r_hit/m_1) + w_2), g);
+    const float m_2 = m_1 / u_2;
+    
+    TSParams TonescaleParams = { 
+    	n, 
+    	n_r, 
+    	g, 
+    	t_1, 
+    	c_t, 
+    	s_2, 
+    	u_2, 
+    	m_2 
+    };
+
+    return TonescaleParams;
+}
+
+
+/* --- Tone scale math --- */
+float tonescale_fwd( 
+	float x, 				// scene-referred input (i.e. linear ACES2065-1)
+	TSParams params 		// struct of type TSParams
+) 	
+{
+	// forward MM tone scale
+	float f = params.m_2 * pow( max(0.0, x) / (x + params.s_2), params.g);  	
+	float h = max(0., f * f / (f + params.t_1));		 // forward flare 
+
+	return h * params.n_r;	// output is luminance in cd/m^2
+}
+
+float tonescale_inv( 
+	float Y,				// luminance in cd/m^2
+	TSParams params 		// struct of type TSParams
+)  
+{
+	float Z = max(0., min(params.n / (params.u_2 * params.n_r), Y));
+	float h = (Z + sqrt(Z * (4. * params.t_1 + Z))) / 2.;
+	float f = params.s_2 / (pow((params.m_2 / h), (1. / params.g)) - 1.);
+
+	return f;				// output is scene-referred input
+}
diff --git a/Engine/Shaders/Private/ACES/ACESUtilities.ush b/Engine/Shaders/Private/ACES/ACESUtilities.ush
new file mode 100644
--- /dev/null
+++ b/Engine/Shaders/Private/ACES/ACESUtilities.ush
@@ -0,0 +1,71 @@
+// Copyright Epic Games, Inc. All Rights Reserved.
+
+// <ACEStransformID>urn:ampas:aces:transformId:v2.0:Lib.Academy.Utilities.a2.v1</ACEStransformID>
+// <ACESuserName>Utilities</ACESuserName>
+
+//
+// Generic functions that may be useful for writing CTL programs
+//
+#pragma once
+
+#include "ACESCommon.ush"
+
+#define UE_ACES_SAFE_FUNCTIONS 1
+
+float radians_to_degrees( float radians )
+{
+	return radians * 180.0 / PI;
+}
+
+float degrees_to_radians( float degrees )
+{
+	return degrees / 180.0 * PI;
+}
+
+float copysign(float x, float y)
+{
+	return sign(y) * abs(x);
+}
+
+float3 copysign(float3 x, float3 y)
+{
+	return sign(y) * abs(x);
+}
+
+float smin(float a, float b, float s)
+{
+    float h = max(s - abs(a - b), 0.0) / s;
+    return min(a, b) - h * h * h * s * (1.0 / 6.0);
+}
+
+#if UE_ACES_SAFE_FUNCTIONS
+
+float spow(float base, float exponent)
+{ 
+	return (base < 0.0 && exponent != floor(exponent)) ? 0.0 : pow(base, exponent);
+}
+
+float3 spow(float3 base, float3 exponent)
+{
+	return float3(
+        spow(base[0], exponent[0]), 
+        spow(base[1], exponent[1]), 
+		spow(base[2], exponent[2])
+    );
+}
+
+float sdiv(float a, float b)
+{
+    return b == 0.0 ? 0.0 : a / b;
+}
+
+#else // UE_ACES_SAFE_FUNCTIONS
+
+#define spow pow
+
+float sdiv(float a, float b)
+{
+    return a / b;
+}
+
+#endif // UE_SAFE_FUNCTIONS_SLOW
diff --git a/Engine/Shaders/Private/ACES/ACES_v1.3.ush b/Engine/Shaders/Private/ACES/ACES_v1.3.ush
new file mode 100644
--- /dev/null
+++ b/Engine/Shaders/Private/ACES/ACES_v1.3.ush
@@ -0,0 +1,1481 @@
+// Copyright Epic Games, Inc. All Rights Reserved.
+
+/*
+=============================================================================
+	ACES: Academy Color Encoding System
+	[REDACTED_PROFILE_URL]/aces-dev/tree/v1.3
+
+	License Terms for Academy Color Encoding System Components
+
+	Academy Color Encoding System (ACES) software and tools are provided by the Academy under
+	the following terms and conditions: A worldwide, royalty-free, non-exclusive right to copy, modify, create
+	derivatives, and use, in source and binary forms, is hereby granted, subject to acceptance of this license.
+
+	Copyright © 2015 Academy of Motion Picture Arts and Sciences (A.M.P.A.S.). Portions contributed by
+	others as indicated. All rights reserved.
+
+	Performance of any of the aforementioned acts indicates acceptance to be bound by the following
+	terms and conditions:
+
+	 *	Copies of source code, in whole or in part, must retain the above copyright
+		notice, this list of conditions and the Disclaimer of Warranty.
+	 *	Use in binary form must retain the above copyright notice, this list of
+		conditions and the Disclaimer of Warranty in the documentation and/or other
+		materials provided with the distribution.
+	 *	Nothing in this license shall be deemed to grant any rights to trademarks,
+		copyrights, patents, trade secrets or any other intellectual property of
+		A.M.P.A.S. or any contributors, except as expressly stated herein.
+	 *	Neither the name "A.M.P.A.S." nor the name of any other contributors to this
+		software may be used to endorse or promote products derivative of or based on
+		this software without express prior written permission of A.M.P.A.S. or the
+		contributors, as appropriate.
+
+	This license shall be construed pursuant to the laws of the State of California,
+	and any disputes related thereto shall be subject to the jurisdiction of the courts therein.
+
+	Disclaimer of Warranty: THIS SOFTWARE IS PROVIDED BY A.M.P.A.S. AND CONTRIBUTORS "AS
+	IS" AND ANY EXPRESS OR IMPLIED WARRANTIES, INCLUDING, BUT NOT LIMITED TO, THE
+	IMPLIED WARRANTIES OF MERCHANTABILITY, FITNESS FOR A PARTICULAR PURPOSE, AND
+	NON-INFRINGEMENT ARE DISCLAIMED. IN NO EVENT SHALL A.M.P.A.S., OR ANY
+	CONTRIBUTORS OR DISTRIBUTORS, BE LIABLE FOR ANY DIRECT, INDIRECT, INCIDENTAL,
+	SPECIAL, EXEMPLARY, RESITUTIONARY, OR CONSEQUENTIAL DAMAGES (INCLUDING, BUT
+	NOT LIMITED TO, PROCUREMENT OF SUBSTITUTE GOODS OR SERVICES; LOSS OF USE,
+	DATA, OR PROFITS; OR BUSINESS INTERRUPTION) HOWEVER CAUSED AND ON ANY THEORY
+	OF LIABILITY, WHETHER IN CONTRACT, STRICT LIABILITY, OR TORT (INCLUDING
+	NEGLIGENCE OR OTHERWISE) ARISING IN ANY WAY OUT OF THE USE OF THIS SOFTWARE,
+	EVEN IF ADVISED OF THE POSSIBILITY OF SUCH DAMAGE.
+
+	WITHOUT LIMITING THE GENERALITY OF THE FOREGOING, THE ACADEMY SPECIFICALLY
+	DISCLAIMS ANY REPRESENTATIONS OR WARRANTIES WHATSOEVER RELATED TO PATENT OR
+	OTHER INTELLECTUAL PROPERTY RIGHTS IN THE ACADEMY COLOR ENCODING SYSTEM, OR
+	APPLICATIONS THEREOF, HELD BY PARTIES OTHER THAN A.M.P.A.S.,WHETHER DISCLOSED
+	OR UNDISCLOSED.
+=============================================================================
+*/
+
+/*
+	we do matrix/matrix and matrix/vectors multiplications are following textbook order, like in our current implementation in ACES1.0.
+	An example:
+	in ACES:
+		const float AP0_2_AP1_MAT[4][4] = mult_f44_f44( AP0_2_XYZ_MAT, XYZ_2_AP1_MAT);
+		const float RGB_out[3] = mult_f3_f44( RGB_in, AP0_2_AP1_MAT);
+	Equivalent math:
+		R_out =  1.4514393161 * R_in + -0.2365107469 * G_in + -0.2149285693 * B_in;
+		G_out = -0.0765537734 * R_in +  1.1762296998 * G_in + -0.0996759264 * B_in;
+		B_out =  0.0083161484 * R_in + -0.0060324498 * G_in +  0.9977163014 * B_in;
+
+	in HLSL:
+	static const float3x3 AP0_2_AP1_MAT = //mul( XYZ_2_AP1_MAT, AP0_2_XYZ_MAT );
+	{
+		1.4514393161, -0.2365107469, -0.2149285693,
+		-0.0765537734,  1.1762296998, -0.0996759264,
+		0.0083161484, -0.0060324498,  0.9977163014,
+	};
+	float3 RGB_out = mul(AP0_2_AP1_MAT, RGB_in.xyz);
+ */
+
+ /////////////////////////////////////////////////////////////////////////////////////////
+
+#pragma once
+
+#include "ACESCommon.ush"
+
+///////////////////////////////////////////////////////////////////////////////////////////
+
+static const float HALF_MAX = 65504.0;
+static const float HALF_POS_INF = 65535.0f;
+
+static const float MIN_STOP_SDR = -6.5;
+static const float MAX_STOP_SDR = 6.5;
+
+static const float MIN_STOP_RRT = -15.;
+static const float MAX_STOP_RRT = 18.;
+
+static const float MIN_LUM_SDR = 0.02;
+static const float MAX_LUM_SDR = 48.0;
+
+static const float MIN_LUM_RRT = 0.0001;
+static const float MAX_LUM_RRT = 10000.0;
+
+float pow10(float x)
+{
+	return pow(10, x);
+}
+
+// ------- Red modifier functions
+float cubic_basis_shaper
+(
+	float x,
+	float w   // full base width of the shaper function (in degrees)
+)
+{
+	float M[4][4] = { { -1. / 6,  3. / 6, -3. / 6,  1. / 6 },
+					  {  3. / 6, -6. / 6,  3. / 6,  0. / 6 },
+					  { -3. / 6,  0. / 6,  3. / 6,  0. / 6 },
+					  {  1. / 6,  4. / 6,  1. / 6,  0. / 6 } };
+
+	float knots[5] = { -0.5 * w, -0.25 * w, 0, 0.25 * w, 0.5 * w };
+
+	float y = 0;
+	if ((x > knots[0]) && (x < knots[4])) {
+		float knot_coord = (x - knots[0]) * 4.0 / w;
+		int j = knot_coord;
+		float t = knot_coord - j;
+
+		float monomials[4] = { t * t * t, t * t, t, 1.0 };
+
+		// (if/else structure required for compatibility with CTL < v1.5.)
+		if (j == 3) {
+			y = monomials[0] * M[0][0] + monomials[1] * M[1][0] +
+				monomials[2] * M[2][0] + monomials[3] * M[3][0];
+		}
+		else if (j == 2) {
+			y = monomials[0] * M[0][1] + monomials[1] * M[1][1] +
+				monomials[2] * M[2][1] + monomials[3] * M[3][1];
+		}
+		else if (j == 1) {
+			y = monomials[0] * M[0][2] + monomials[1] * M[1][2] +
+				monomials[2] * M[2][2] + monomials[3] * M[3][2];
+		}
+		else if (j == 0) {
+			y = monomials[0] * M[0][3] + monomials[1] * M[1][3] +
+				monomials[2] * M[2][3] + monomials[3] * M[3][3];
+		}
+		else {
+			y = 0.0;
+		}
+	}
+
+	return y * 1.5;
+}
+
+// note: this matrix is not transposed unlike other matrix operations because instead of doing v_out = mul(M, v_in), we do v_out = mul(v_in, M)
+// Textbook monomial to basis-function conversion matrix.
+static const float3x3 M1 =
+{
+	{  0.5, -1.0, 0.5 },
+	{ -1.0,  1.0, 0.5 },
+	{  0.5,  0.0, 0.0 }
+};
+
+
+/* ---- Functions to compress highlights ---- */
+// allow for simulated white points without clipping
+
+float roll_white_fwd(
+	float inValue,      // color value to adjust (white scaled to around 1.0)
+	float new_wht, // white adjustment (e.g. 0.9 for 10% darkening)
+	float width    // adjusted width (e.g. 0.25 for top quarter of the tone scale)
+)
+{
+	const float x0 = -1.0;
+	const float x1 = x0 + width;
+	const float y0 = -new_wht;
+	const float y1 = x1;
+	const float m1 = (x1 - x0);
+	const float a = y0 - y1 + m1;
+	const float b = 2 * (y1 - y0) - m1;
+	const float c = y0;
+	const float t = (-inValue - x0) / (x1 - x0);
+	float outValue = 0.0;
+	if (t < 0.0)
+		outValue = -(t * b + c);
+	else if (t > 1.0)
+		outValue = inValue;
+	else
+		outValue = -((t * a + b) * t + c);
+	return outValue;
+}
+
+float roll_white_rev(
+	float inValue,      // color value to adjust (white scaled to around 1.0)
+	float new_wht, // white adjustment (e.g. 0.9 for 10% darkening)
+	float width    // adjusted width (e.g. 0.25 for top quarter of the tone scale)
+)
+{
+	const float x0 = -1.0;
+	const float x1 = x0 + width;
+	const float y0 = -new_wht;
+	const float y1 = x1;
+	const float m1 = (x1 - x0);
+	const float a = y0 - y1 + m1;
+	const float b = 2. * (y1 - y0) - m1;
+	float c = y0;
+	float outValue = 0.0;
+	if (-inValue < y0)
+		outValue = -x0;
+	else if (-inValue > y1)
+		outValue = inValue;
+	else {
+		c = c + inValue;
+		const float discrim = sqrt(b * b - 4. * a * c);
+		const float t = (2. * c) / (-discrim - b);
+		outValue = -((t * (x1 - x0)) + x0);
+	}
+	return outValue;
+}
+
+float3 limit_to_primaries
+(
+	float3 XYZ,
+	FColorSpace LIMITING_PRI
+)
+{
+	float3x3 XYZ_2_LIMITING_PRI_MAT = LIMITING_PRI.XYZtoRGB;
+	float3x3 LIMITING_PRI_2_XYZ_MAT = LIMITING_PRI.RGBtoXYZ;
+
+	// XYZ to limiting primaries
+	float3 rgb = mul(XYZ_2_LIMITING_PRI_MAT, XYZ);
+
+	// Clip any values outside the limiting primaries
+	float3 limitedRgb = clamp(rgb, 0.0.xxx, 1.0.xxx);
+
+	// Convert limited RGB to XYZ
+	return mul(LIMITING_PRI_2_XYZ_MAT, limitedRgb);
+}
+
+///+ TODO: check if valid
+float interpolate1D(const float2 table[2], float value)
+{
+	float t = saturate((value - table[0].x) / (table[1].x - table[0].x));
+	return lerp(table[0].y, table[1].y, t);
+}
+///-
+
+float lookup_ACESmin(float minLum)
+{
+	const float2 minTable[2] = { float2(log10(MIN_LUM_RRT), MIN_STOP_RRT),
+								 float2(log10(MIN_LUM_SDR), MIN_STOP_SDR) };
+
+	return 0.18 * pow(2., interpolate1D(minTable, log10(minLum)));
+}
+
+float lookup_ACESmax(float maxLum)
+{
+	const float2 maxTable[2] = { float2(log10(MAX_LUM_SDR), MAX_STOP_SDR),
+								   float2(log10(MAX_LUM_RRT), MAX_STOP_RRT) };
+
+	return 0.18 * pow(2., interpolate1D(maxTable, log10(maxLum)));
+}
+
+struct TsPoint
+{
+	float x;        // ACES
+	float y;        // luminance
+	float slope;    // 
+};
+
+struct TsParams
+{
+	TsPoint Min;
+	TsPoint Mid;
+	TsPoint Max;
+	float coefsLow[6];
+	float coefsHigh[6];
+};
+
+void init_coefsLow(
+	TsPoint TsPointLow,
+	TsPoint TsPointMid,
+	out float coefsLow[5]
+)
+{
+
+	float knotIncLow = (log10(TsPointMid.x) - log10(TsPointLow.x)) / 3.;
+	// float halfKnotInc = (log10(TsPointMid.x) - log10(TsPointLow.x)) / 6.;
+
+	// Determine two lowest coefficients (straddling minPt)
+	coefsLow[0] = (TsPointLow.slope * (log10(TsPointLow.x) - 0.5 * knotIncLow)) + (log10(TsPointLow.y) - TsPointLow.slope * log10(TsPointLow.x));
+	coefsLow[1] = (TsPointLow.slope * (log10(TsPointLow.x) + 0.5 * knotIncLow)) + (log10(TsPointLow.y) - TsPointLow.slope * log10(TsPointLow.x));
+	// NOTE: if slope=0, then the above becomes just 
+		// coefsLow[0] = log10(TsPointLow.y);
+		// coefsLow[1] = log10(TsPointLow.y);
+	// leaving it as a variable for now in case we decide we need non-zero slope extensions
+
+	// Determine two highest coefficients (straddling midPt)
+	coefsLow[3] = (TsPointMid.slope * (log10(TsPointMid.x) - 0.5 * knotIncLow)) + (log10(TsPointMid.y) - TsPointMid.slope * log10(TsPointMid.x));
+	coefsLow[4] = (TsPointMid.slope * (log10(TsPointMid.x) + 0.5 * knotIncLow)) + (log10(TsPointMid.y) - TsPointMid.slope * log10(TsPointMid.x));
+
+	// Middle coefficient (which defines the "sharpness of the bend") is linearly interpolated
+	float2 bendsLow[2] = { float2(MIN_STOP_RRT, 0.18),
+						   float2(MIN_STOP_SDR, 0.35) };
+	float pctLow = interpolate1D(bendsLow, log2(TsPointLow.x / 0.18));
+	coefsLow[2] = log10(TsPointLow.y) + pctLow * (log10(TsPointMid.y) - log10(TsPointLow.y));
+}
+
+void init_coefsHigh(
+	TsPoint TsPointMid,
+	TsPoint TsPointMax,
+	out float coefsHigh[5]
+)
+{
+
+	float knotIncHigh = (log10(TsPointMax.x) - log10(TsPointMid.x)) / 3.;
+	// float halfKnotInc = (log10(TsPointMax.x) - log10(TsPointMid.x)) / 6.;
+
+	// Determine two lowest coefficients (straddling midPt)
+	coefsHigh[0] = (TsPointMid.slope * (log10(TsPointMid.x) - 0.5 * knotIncHigh)) + (log10(TsPointMid.y) - TsPointMid.slope * log10(TsPointMid.x));
+	coefsHigh[1] = (TsPointMid.slope * (log10(TsPointMid.x) + 0.5 * knotIncHigh)) + (log10(TsPointMid.y) - TsPointMid.slope * log10(TsPointMid.x));
+
+	// Determine two highest coefficients (straddling maxPt)
+	coefsHigh[3] = (TsPointMax.slope * (log10(TsPointMax.x) - 0.5 * knotIncHigh)) + (log10(TsPointMax.y) - TsPointMax.slope * log10(TsPointMax.x));
+	coefsHigh[4] = (TsPointMax.slope * (log10(TsPointMax.x) + 0.5 * knotIncHigh)) + (log10(TsPointMax.y) - TsPointMax.slope * log10(TsPointMax.x));
+	// NOTE: if slope=0, then the above becomes just
+		// coefsHigh[0] = log10(TsPointHigh.y);
+		// coefsHigh[1] = log10(TsPointHigh.y);
+	// leaving it as a variable for now in case we decide we need non-zero slope extensions
+
+	// Middle coefficient (which defines the "sharpness of the bend") is linearly interpolated
+	float2 bendsHigh[2] = { float2(MAX_STOP_SDR, 0.89) ,
+							float2(MAX_STOP_RRT, 0.90) };
+	float pctHigh = interpolate1D(bendsHigh, log2(TsPointMax.x / 0.18));
+	coefsHigh[2] = log10(TsPointMid.y) + pctHigh * (log10(TsPointMax.y) - log10(TsPointMid.y));
+}
+
+float shift(float inValue, float expShift)
+{
+	return pow(2., (log2(inValue) - expShift));
+}
+
+TsParams init_TsParams(
+	float minLum,
+	float maxLum,
+	float expShift = 0
+)
+{
+	TsPoint MIN_PT = { lookup_ACESmin(minLum), minLum, 0.0 };
+	TsPoint MID_PT = { 0.18, 4.8, 1.55 };
+	TsPoint MAX_PT = { lookup_ACESmax(maxLum), maxLum, 0.0 };
+	float cLow[5];
+	init_coefsLow(MIN_PT, MID_PT, cLow);
+	float cHigh[5];
+	init_coefsHigh(MID_PT, MAX_PT, cHigh);
+	MIN_PT.x = shift(lookup_ACESmin(minLum), expShift);
+	MID_PT.x = shift(0.18, expShift);
+	MAX_PT.x = shift(lookup_ACESmax(maxLum), expShift);
+
+	TsParams P = {
+		{MIN_PT.x, MIN_PT.y, MIN_PT.slope},
+		{MID_PT.x, MID_PT.y, MID_PT.slope},
+		{MAX_PT.x, MAX_PT.y, MAX_PT.slope},
+		{cLow[0], cLow[1], cLow[2], cLow[3], cLow[4], cLow[4]},
+		{cHigh[0], cHigh[1], cHigh[2], cHigh[3], cHigh[4], cHigh[4]}
+	};
+
+	return P;
+}
+
+float ssts
+(
+	const float x,
+	const TsParams C
+)
+{
+	const int N_KNOTS_LOW = 4;
+	const int N_KNOTS_HIGH = 4;
+
+	// Check for negatives or zero before taking the log. If negative or zero,
+	// set to HALF_MIN
+	float logx = log10(max(x, 1e-10));
+	float logy;
+
+	if (logx <= log10(C.Min.x)) {
+
+		logy = logx * C.Min.slope + (log10(C.Min.y) - C.Min.slope * log10(C.Min.x));
+
+	}
+	else if ((logx > log10(C.Min.x)) && (logx < log10(C.Mid.x))) {
+
+		float knot_coord = (N_KNOTS_LOW - 1) * (logx - log10(C.Min.x)) / (log10(C.Mid.x) - log10(C.Min.x));
+		int j = knot_coord;
+		float t = knot_coord - j;
+
+		float3 cf = { C.coefsLow[j], C.coefsLow[j + 1], C.coefsLow[j + 2] };
+
+		float3 monomials = { t * t, t, 1.0 };
+		logy = dot(monomials, mul(cf, M1));
+
+	}
+	else if ((logx >= log10(C.Mid.x)) && (logx < log10(C.Max.x))) {
+
+		float knot_coord = (N_KNOTS_HIGH - 1) * (logx - log10(C.Mid.x)) / (log10(C.Max.x) - log10(C.Mid.x));
+		int j = knot_coord;
+		float t = knot_coord - j;
+
+		float3 cf = { C.coefsHigh[j], C.coefsHigh[j + 1], C.coefsHigh[j + 2] };
+
+		float3 monomials = { t * t, t, 1.0 };
+		logy = dot(monomials, mul(cf, M1));
+
+	}
+	else { //if ( logIn >= log10(C.Max.x) ) { 
+
+		logy = logx * C.Max.slope + (log10(C.Max.y) - C.Max.slope * log10(C.Max.x));
+
+	}
+
+	return pow10(logy);
+
+}
+
+float inv_ssts
+(
+	const float y,
+	const TsParams C
+)
+{
+	const int N_KNOTS_LOW = 4;
+	const int N_KNOTS_HIGH = 4;
+
+	const float KNOT_INC_LOW = (log10(C.Mid.x) - log10(C.Min.x)) / (N_KNOTS_LOW - 1.);
+	const float KNOT_INC_HIGH = (log10(C.Max.x) - log10(C.Mid.x)) / (N_KNOTS_HIGH - 1.);
+
+	// KNOT_Y is luminance of the spline at each knot
+	float KNOT_Y_LOW[N_KNOTS_LOW];
+	///+warning: redefinition of 'i'
+	{
+		for (int i = 0; i < N_KNOTS_LOW; i = i + 1) {
+			KNOT_Y_LOW[i] = (C.coefsLow[i] + C.coefsLow[i + 1]) / 2.;
+		};
+	}
+	///-
+
+	float KNOT_Y_HIGH[N_KNOTS_HIGH];
+	///+ warning: redefinition of 'i'
+	{
+		for (int i = 0; i < N_KNOTS_HIGH; i = i + 1) {
+			KNOT_Y_HIGH[i] = (C.coefsHigh[i] + C.coefsHigh[i + 1]) / 2.;
+		};
+	}
+	///-
+	float logy = log10(max(y, 1e-10));
+
+	float logx;
+	if (logy <= log10(C.Min.y)) {
+
+		logx = log10(C.Min.x);
+
+	}
+	else if ((logy > log10(C.Min.y)) && (logy <= log10(C.Mid.y))) {
+
+		int j;
+		float3 cf;
+		if (logy > KNOT_Y_LOW[0] && logy <= KNOT_Y_LOW[1]) {
+			cf[0] = C.coefsLow[0];  cf[1] = C.coefsLow[1];  cf[2] = C.coefsLow[2];  j = 0;
+		}
+		else if (logy > KNOT_Y_LOW[1] && logy <= KNOT_Y_LOW[2]) {
+			cf[0] = C.coefsLow[1];  cf[1] = C.coefsLow[2];  cf[2] = C.coefsLow[3];  j = 1;
+		}
+		else if (logy > KNOT_Y_LOW[2] && logy <= KNOT_Y_LOW[3]) {
+			cf[0] = C.coefsLow[2];  cf[1] = C.coefsLow[3];  cf[2] = C.coefsLow[4];  j = 2;
+		}
+
+		const float3 tmp = mul(cf, M1);
+
+		float a = tmp[0];
+		float b = tmp[1];
+		float c = tmp[2];
+		c = c - logy;
+
+		const float d = sqrt(b * b - 4. * a * c);
+
+		const float t = (2. * c) / (-d - b);
+
+		logx = log10(C.Min.x) + (t + j) * KNOT_INC_LOW;
+
+	}
+	else if ((logy > log10(C.Mid.y)) && (logy < log10(C.Max.y))) {
+
+		int j;
+		float3 cf;
+		if (logy >= KNOT_Y_HIGH[0] && logy <= KNOT_Y_HIGH[1]) {
+			cf[0] = C.coefsHigh[0];  cf[1] = C.coefsHigh[1];  cf[2] = C.coefsHigh[2];  j = 0;
+		}
+		else if (logy > KNOT_Y_HIGH[1] && logy <= KNOT_Y_HIGH[2]) {
+			cf[0] = C.coefsHigh[1];  cf[1] = C.coefsHigh[2];  cf[2] = C.coefsHigh[3];  j = 1;
+		}
+		else if (logy > KNOT_Y_HIGH[2] && logy <= KNOT_Y_HIGH[3]) {
+			cf[0] = C.coefsHigh[2];  cf[1] = C.coefsHigh[3];  cf[2] = C.coefsHigh[4];  j = 2;
+		}
+
+		const float3 tmp = mul(cf, M1);
+
+		float a = tmp[0];
+		float b = tmp[1];
+		float c = tmp[2];
+		c = c - logy;
+
+		const float d = sqrt(b * b - 4. * a * c);
+
+		const float t = (2. * c) / (-d - b);
+
+		logx = log10(C.Mid.x) + (t + j) * KNOT_INC_HIGH;
+
+	}
+	else { //if ( logy >= log10(C.Max.y) ) {
+
+		logx = log10(C.Max.x);
+
+	}
+
+	return pow10(logx);
+
+}
+
+float3 ssts_f3
+(
+	const float3 x,
+	const TsParams C
+)
+{
+	float3 outValue;
+	outValue[0] = ssts(x[0], C);
+	outValue[1] = ssts(x[1], C);
+	outValue[2] = ssts(x[2], C);
+
+	return outValue;
+}
+
+float3 inv_ssts_f3
+(
+	const float3 x,
+	const TsParams C
+)
+{
+	float3 outValue;
+	outValue[0] = inv_ssts(x[0], C);
+	outValue[1] = inv_ssts(x[1], C);
+	outValue[2] = inv_ssts(x[2], C);
+
+	return outValue;
+}
+
+float moncurve_f(float x, float gamma, float offs)
+{
+	// Forward monitor curve
+	float y;
+	const float fs = ((gamma - 1.0) / offs) * pow(offs * gamma / ((gamma - 1.0) * (1.0 + offs)), gamma);
+	const float xb = offs / (gamma - 1.0);
+	if (x >= xb)
+		y = pow((x + offs) / (1.0 + offs), gamma);
+	else
+		y = x * fs;
+	return y;
+}
+
+float moncurve_r(float y, float gamma, float offs)
+{
+	// Reverse monitor curve
+	float x;
+	const float yb = pow(offs * gamma / ((gamma - 1.0) * (1.0 + offs)), gamma);
+	const float rs = pow((gamma - 1.0) / offs, gamma - 1.0) * pow((1.0 + offs) / gamma, gamma);
+	if (y >= yb)
+		x = (1.0 + offs) * pow(y, 1.0 / gamma) - offs;
+	else
+		x = y * rs;
+	return x;
+}
+
+float3 moncurve_f_f3(float3 x, float gamma, float offs)
+{
+	float3 y;
+	y[0] = moncurve_f(x[0], gamma, offs);
+	y[1] = moncurve_f(x[1], gamma, offs);
+	y[2] = moncurve_f(x[2], gamma, offs);
+	return y;
+}
+
+float3 moncurve_r_f3(float3 y, float gamma, float offs)
+{
+	float3 x;
+	x[0] = moncurve_r(y[0], gamma, offs);
+	x[1] = moncurve_r(y[1], gamma, offs);
+	x[2] = moncurve_r(y[2], gamma, offs);
+	return x;
+}
+
+float bt1886_f(float V, float gamma, float Lw, float Lb)
+{
+	// The reference EOTF specified in Rec. ITU-R BT.1886
+	// L = a(max[(V+b),0])^g
+	float a = pow(pow(Lw, 1. / gamma) - pow(Lb, 1. / gamma), gamma);
+	float b = pow(Lb, 1. / gamma) / (pow(Lw, 1. / gamma) - pow(Lb, 1. / gamma));
+	float L = a * pow(max(V + b, 0.), gamma);
+	return L;
+}
+
+float bt1886_r(float L, float gamma, float Lw, float Lb)
+{
+	// The reference EOTF specified in Rec. ITU-R BT.1886
+	// L = a(max[(V+b),0])^g
+	float a = pow(pow(Lw, 1. / gamma) - pow(Lb, 1. / gamma), gamma);
+	float b = pow(Lb, 1. / gamma) / (pow(Lw, 1. / gamma) - pow(Lb, 1. / gamma));
+	float V = pow(max(L / a, 0.), 1. / gamma) - b;
+	return V;
+}
+
+float3 bt1886_f_f3(float3 V, float gamma, float Lw, float Lb)
+{
+	float3 L;
+	L[0] = bt1886_f(V[0], gamma, Lw, Lb);
+	L[1] = bt1886_f(V[1], gamma, Lw, Lb);
+	L[2] = bt1886_f(V[2], gamma, Lw, Lb);
+	return L;
+}
+
+float3 bt1886_r_f3(float3 L, float gamma, float Lw, float Lb)
+{
+	float3 V;
+	V[0] = bt1886_r(L[0], gamma, Lw, Lb);
+	V[1] = bt1886_r(L[1], gamma, Lw, Lb);
+	V[2] = bt1886_r(L[2], gamma, Lw, Lb);
+	return V;
+}
+
+
+// SMPTE Range vs Full Range scaling formulas
+float smpteRange_to_fullRange(float inValue)
+{
+	const float REFBLACK = (64. / 1023.);
+	const float REFWHITE = (940. / 1023.);
+
+	return ((inValue - REFBLACK) / (REFWHITE - REFBLACK));
+}
+
+float fullRange_to_smpteRange(float inValue)
+{
+	const float REFBLACK = (64. / 1023.);
+	const float REFWHITE = (940. / 1023.);
+
+	return (inValue * (REFWHITE - REFBLACK) + REFBLACK);
+}
+
+float3 smpteRange_to_fullRange_f3(float3 rgbIn)
+{
+	float3 rgbOut;
+	rgbOut[0] = smpteRange_to_fullRange(rgbIn[0]);
+	rgbOut[1] = smpteRange_to_fullRange(rgbIn[1]);
+	rgbOut[2] = smpteRange_to_fullRange(rgbIn[2]);
+
+	return rgbOut;
+}
+
+float3 fullRange_to_smpteRange_f3(float3 rgbIn)
+{
+	float3 rgbOut;
+	rgbOut[0] = fullRange_to_smpteRange(rgbIn[0]);
+	rgbOut[1] = fullRange_to_smpteRange(rgbIn[1]);
+	rgbOut[2] = fullRange_to_smpteRange(rgbIn[2]);
+
+	return rgbOut;
+}
+
+// Base functions from SMPTE ST 2084-2014
+
+// Constants from SMPTE ST 2084-2014
+static const float pq_m1 = 0.1593017578125; // ( 2610.0 / 4096.0 ) / 4.0;
+static const float pq_m2 = 78.84375; // ( 2523.0 / 4096.0 ) * 128.0;
+static const float pq_c1 = 0.8359375; // 3424.0 / 4096.0 or pq_c3 - pq_c2 + 1.0;
+static const float pq_c2 = 18.8515625; // ( 2413.0 / 4096.0 ) * 32.0;
+static const float pq_c3 = 18.6875; // ( 2392.0 / 4096.0 ) * 32.0;
+
+static const float pq_C = 10000.0;
+
+// Converts from the non-linear perceptually quantized space to linear cd/m^2
+// Note that this is in float, and assumes normalization from 0 - 1
+// (0 - pq_C for linear) and does not handle the integer coding in the Annex 
+// sections of SMPTE ST 2084-2014
+float ST2084_2_Y(float N)
+{
+	// Note that this does NOT handle any of the signal range
+	// considerations from 2084 - this assumes full range (0 - 1)
+	float Np = pow(N, 1.0 / pq_m2);
+	float L = Np - pq_c1;
+	if (L < 0.0)
+		L = 0.0;
+	L = L / (pq_c2 - pq_c3 * Np);
+	L = pow(L, 1.0 / pq_m1);
+	return L * pq_C; // returns cd/m^2
+}
+
+// Converts from linear cd/m^2 to the non-linear perceptually quantized space
+// Note that this is in float, and assumes normalization from 0 - 1
+// (0 - pq_C for linear) and does not handle the integer coding in the Annex 
+// sections of SMPTE ST 2084-2014
+float Y_2_ST2084(float C)
+//pq_r
+{
+	// Note that this does NOT handle any of the signal range
+	// considerations from 2084 - this returns full range (0 - 1)
+	float L = C / pq_C;
+	float Lm = pow(L, pq_m1);
+	float N = (pq_c1 + pq_c2 * Lm) / (1.0 + pq_c3 * Lm);
+	N = pow(N, pq_m2);
+	return N;
+}
+
+float3 Y_2_ST2084_f3(float3 inValue)
+{
+	// converts from linear cd/m^2 to PQ code values
+
+	float3 outValue;
+	outValue[0] = Y_2_ST2084(inValue[0]);
+	outValue[1] = Y_2_ST2084(inValue[1]);
+	outValue[2] = Y_2_ST2084(inValue[2]);
+
+	return outValue;
+}
+
+float3 ST2084_2_Y_f3(float3 inValue)
+{
+	// converts from PQ code values to linear cd/m^2
+
+	float3 outValue;
+	outValue[0] = ST2084_2_Y(inValue[0]);
+	outValue[1] = ST2084_2_Y(inValue[1]);
+	outValue[2] = ST2084_2_Y(inValue[2]);
+
+	return outValue;
+}
+
+// Conversion of PQ signal to HLG, as detailed in Section 7 of ITU-R BT.2390-0
+float3 ST2084_2_HLG_1000nits_f3(float3 PQ)
+{
+	// ST.2084 EOTF (non-linear PQ to display light)
+	float3 displayLinear = ST2084_2_Y_f3(PQ);
+
+	// HLG Inverse EOTF (i.e. HLG inverse OOTF followed by the HLG OETF)
+	// HLG Inverse OOTF (display linear to scene linear)
+	float Y_d = 0.2627 * displayLinear[0] + 0.6780 * displayLinear[1] + 0.0593 * displayLinear[2];
+	const float L_w = 1000.;
+	const float L_b = 0.;
+	const float alpha = (L_w - L_b);
+	const float beta = L_b;
+	const float gamma = 1.2;
+
+	float3 sceneLinear;
+	if (Y_d == 0.) {
+		/* This case is to protect against pow(0,-N)=Inf error. The ITU document
+		does not offer a recommendation for this corner case. There may be a
+		better way to handle this, but for now, this works.
+		*/
+		sceneLinear[0] = 0.;
+		sceneLinear[1] = 0.;
+		sceneLinear[2] = 0.;
+	}
+	else {
+		sceneLinear[0] = pow((Y_d - beta) / alpha, (1. - gamma) / gamma) * ((displayLinear[0] - beta) / alpha);
+		sceneLinear[1] = pow((Y_d - beta) / alpha, (1. - gamma) / gamma) * ((displayLinear[1] - beta) / alpha);
+		sceneLinear[2] = pow((Y_d - beta) / alpha, (1. - gamma) / gamma) * ((displayLinear[2] - beta) / alpha);
+	}
+
+	// HLG OETF (scene linear to non-linear signal value)
+	const float a = 0.17883277;
+	const float b = 0.28466892; // 1.-4.*a;
+	const float c = 0.55991073; // 0.5-a*log(4.*a);
+
+	float3 HLG;
+	if (sceneLinear[0] <= 1. / 12) {
+		HLG[0] = sqrt(3. * sceneLinear[0]);
+	}
+	else {
+		HLG[0] = a * log(12. * sceneLinear[0] - b) + c;
+	}
+	if (sceneLinear[1] <= 1. / 12) {
+		HLG[1] = sqrt(3. * sceneLinear[1]);
+	}
+	else {
+		HLG[1] = a * log(12. * sceneLinear[1] - b) + c;
+	}
+	if (sceneLinear[2] <= 1. / 12) {
+		HLG[2] = sqrt(3. * sceneLinear[2]);
+	}
+	else {
+		HLG[2] = a * log(12. * sceneLinear[2] - b) + c;
+	}
+
+	return HLG;
+}
+
+
+// Conversion of HLG to PQ signal, as detailed in Section 7 of ITU-R BT.2390-0
+float3 HLG_2_ST2084_1000nits_f3(float3 HLG)
+{
+	const float a = 0.17883277;
+	const float b = 0.28466892; // 1.-4.*a;
+	const float c = 0.55991073; // 0.5-a*log(4.*a);
+
+	const float L_w = 1000.;
+	const float L_b = 0.;
+	const float alpha = (L_w - L_b);
+	const float beta = L_b;
+	const float gamma = 1.2;
+
+	// HLG EOTF (non-linear signal value to display linear)
+	// HLG to scene-linear
+	float3 sceneLinear;
+	if (HLG[0] >= 0. && HLG[0] <= 0.5) {
+		sceneLinear[0] = pow(HLG[0], 2.) / 3.;
+	}
+	else {
+		sceneLinear[0] = (exp((HLG[0] - c) / a) + b) / 12.;
+	}
+	if (HLG[1] >= 0. && HLG[1] <= 0.5) {
+		sceneLinear[1] = pow(HLG[1], 2.) / 3.;
+	}
+	else {
+		sceneLinear[1] = (exp((HLG[1] - c) / a) + b) / 12.;
+	}
+	if (HLG[2] >= 0. && HLG[2] <= 0.5) {
+		sceneLinear[2] = pow(HLG[2], 2.) / 3.;
+	}
+	else {
+		sceneLinear[2] = (exp((HLG[2] - c) / a) + b) / 12.;
+	}
+
+	float Y_s = 0.2627 * sceneLinear[0] + 0.6780 * sceneLinear[1] + 0.0593 * sceneLinear[2];
+
+	// Scene-linear to display-linear
+	float3 displayLinear;
+	displayLinear[0] = alpha * pow(Y_s, gamma - 1.) * sceneLinear[0] + beta;
+	displayLinear[1] = alpha * pow(Y_s, gamma - 1.) * sceneLinear[1] + beta;
+	displayLinear[2] = alpha * pow(Y_s, gamma - 1.) * sceneLinear[2] + beta;
+
+	// ST.2084 Inverse EOTF
+	float3 PQ = Y_2_ST2084_f3(displayLinear);
+
+	return PQ;
+}
+
+// Desaturation contants
+static const float RRT_SAT_FACTOR = 0.96;
+static const float ONE_MINUS_RRT_SAT_FACTOR = 0.04;
+static const float3x3 RRT_SAT_MAT =
+{
+	//    {ONE_MINUS_RRT_SAT_FACTOR * AP1_RGB2Y.x + RRT_SAT_FACTOR, ONE_MINUS_RRT_SAT_FACTOR * AP1_RGB2Y.y,                  ONE_MINUS_RRT_SAT_FACTOR * AP1_RGB2Y.z},
+	//    {ONE_MINUS_RRT_SAT_FACTOR * AP1_RGB2Y.x,                  ONE_MINUS_RRT_SAT_FACTOR * AP1_RGB2Y.y + RRT_SAT_FACTOR, ONE_MINUS_RRT_SAT_FACTOR * AP1_RGB2Y.z},
+	//    {ONE_MINUS_RRT_SAT_FACTOR * AP1_RGB2Y.x,                  ONE_MINUS_RRT_SAT_FACTOR * AP1_RGB2Y.y,                  ONE_MINUS_RRT_SAT_FACTOR * AP1_RGB2Y.z + RRT_SAT_FACTOR},
+
+		{0.970889148672, 0.02696327063, 0.0021475807},
+		{0.010889148672, 0.98696327063, 0.0021475807},
+		{0.010889148672, 0.02696327063, 0.9621475807}
+};
+
+static const float3x3 RRT_SAT_MAT_INV =
+{
+	{  1.03032386   , -0.0280867405 , -0.00223706313 },
+	{ -0.0113428626 , 1.01357996    , -0.00223706337 },
+	{ -0.0113428626 , -0.0280867405 , 1.03942955 }
+};
+
+
+// "Glow" module constants
+static const float RRT_GLOW_GAIN = 0.05;
+static const float RRT_GLOW_MID = 0.08;
+
+// Red modifier constants
+static const float RRT_RED_SCALE = 0.82;
+static const float RRT_RED_PIVOT = 0.03;
+static const float RRT_RED_HUE = 0.;
+static const float RRT_RED_WIDTH = 135.;
+
+
+float3 rrt_sweeteners(float3 inValue)
+{
+	float3 aces = inValue;
+
+	// --- Glow module --- //
+	float saturation = rgb_2_saturation(aces);
+	float ycIn = rgb_2_yc(aces);
+	float s = sigmoid_shaper((saturation - 0.4) / 0.2);
+	float addedGlow = 1 + glow_fwd(ycIn, RRT_GLOW_GAIN * s, RRT_GLOW_MID);
+	aces *= addedGlow;
+
+	// --- Red modifier --- //
+	float hue = rgb_2_hue(aces);
+	float centeredHue = center_hue(hue, RRT_RED_HUE);
+	float hueWeight = cubic_basis_shaper(centeredHue, RRT_RED_WIDTH);
+
+	aces.r += hueWeight * saturation * (RRT_RED_PIVOT - aces.r) * (1. - RRT_RED_SCALE);
+
+	// --- ACES to RGB rendering space --- //
+	aces = clamp(aces, 0, HALF_POS_INF);
+	float3 rgbPre = mul(AP0_2_AP1_MAT, aces);
+	rgbPre = clamp(rgbPre, 0, HALF_MAX);
+
+	// --- Global desaturation --- //
+	//rgbPre = mul(RRT_SAT_MAT, rgbPre);
+	rgbPre = lerp(dot(rgbPre, AP1_RGB2Y).xxx, rgbPre, RRT_SAT_FACTOR);
+	return rgbPre;
+}
+
+float3 inv_rrt_sweeteners(float3 inValue)
+{
+	float3 rgbPost = inValue;
+
+	// --- Global desaturation --- //
+	rgbPost = mul(RRT_SAT_MAT_INV, rgbPost);
+
+	rgbPost = clamp(rgbPost, 0.0.xxx, HALF_MAX.xxx);
+
+	// --- RGB rendering space to ACES --- //
+	float3 aces = mul(AP1_2_AP0_MAT, rgbPost);
+
+	aces = clamp(aces, 0.0.xxx, HALF_MAX.xxx);
+
+	// --- Red modifier --- //
+	float hue = rgb_2_hue(aces);
+	float centeredHue = center_hue(hue, RRT_RED_HUE);
+	float hueWeight = cubic_basis_shaper(centeredHue, RRT_RED_WIDTH);
+
+	float minChan;
+	if (centeredHue < 0) { // min_f3(aces) = aces[1] (i.e. magenta-red)
+		minChan = aces[1];
+	}
+	else { // min_f3(aces) = aces[2] (i.e. yellow-red)
+		minChan = aces[2];
+	}
+
+	float a = hueWeight * (1. - RRT_RED_SCALE) - 1.;
+	float b = aces[0] - hueWeight * (RRT_RED_PIVOT + minChan) * (1. - RRT_RED_SCALE);
+	float c = hueWeight * RRT_RED_PIVOT * minChan * (1. - RRT_RED_SCALE);
+
+	aces[0] = (-b - sqrt(b * b - 4. * a * c)) / (2. * a);
+
+	// --- Glow module --- //
+	float saturation = rgb_2_saturation(aces);
+	float ycOut = rgb_2_yc(aces);
+	float s = sigmoid_shaper((saturation - 0.4) / 0.2);
+	float reducedGlow = 1. + glow_inv(ycOut, RRT_GLOW_GAIN * s, RRT_GLOW_MID);
+
+	aces *= reducedGlow;
+	return aces;
+}
+
+float Y_2_linCV(float Y, float Ymax, float Ymin)
+{
+	return (Y - Ymin) / (Ymax - Ymin);
+}
+
+float linCV_2_Y(float linCV, float Ymax, float Ymin)
+{
+	return linCV * (Ymax - Ymin) + Ymin;
+}
+
+float3 Y_2_linCV_f3(float3 Y, float Ymax, float Ymin)
+{
+	float3 linCV;
+	linCV[0] = Y_2_linCV(Y[0], Ymax, Ymin);
+	linCV[1] = Y_2_linCV(Y[1], Ymax, Ymin);
+	linCV[2] = Y_2_linCV(Y[2], Ymax, Ymin);
+	return linCV;
+}
+
+float3 linCV_2_Y_f3(float3 linCV, float Ymax, float Ymin)
+{
+	float3 Y;
+	Y[0] = linCV_2_Y(linCV[0], Ymax, Ymin);
+	Y[1] = linCV_2_Y(linCV[1], Ymax, Ymin);
+	Y[2] = linCV_2_Y(linCV[2], Ymax, Ymin);
+	return Y;
+}
+
+// Gamma compensation factor
+static const float DIM_SURROUND_GAMMA = 0.9811;
+
+float3 darkSurround_to_dimSurround(float3 linearCV)
+{
+	float3 XYZ = mul(AP1_2_XYZ_MAT, linearCV);
+
+	float3 xyY = XYZ_2_xyY(XYZ);
+	xyY[2] = clamp(xyY[2], 0., HALF_POS_INF);
+	xyY[2] = pow(xyY[2], DIM_SURROUND_GAMMA);
+	XYZ = xyY_2_XYZ(xyY);
+
+	return mul(XYZ_2_AP1_MAT, XYZ);
+}
+
+float3 dimSurround_to_darkSurround(float3 linearCV)
+{
+	float3 XYZ = mul(AP1_2_XYZ_MAT, linearCV);
+
+	float3 xyY = XYZ_2_xyY(XYZ);
+	xyY[2] = clamp(xyY[2], 0., HALF_POS_INF);
+	xyY[2] = pow(xyY[2], 1. / DIM_SURROUND_GAMMA);
+	XYZ = xyY_2_XYZ(xyY);
+
+	return mul(XYZ_2_AP1_MAT, XYZ);
+}
+
+float3 dark_to_dim(float3 XYZ)
+{
+	float3 xyY = XYZ_2_xyY(XYZ);
+	xyY[2] = clamp(xyY[2], 0., HALF_POS_INF);
+	xyY[2] = pow(xyY[2], DIM_SURROUND_GAMMA);
+	return xyY_2_XYZ(xyY);
+}
+
+float3 dim_to_dark(float3 XYZ)
+{
+	float3 xyY = XYZ_2_xyY(XYZ);
+	xyY[2] = clamp(xyY[2], 0., HALF_POS_INF);
+	xyY[2] = pow(xyY[2], 1. / DIM_SURROUND_GAMMA);
+	return xyY_2_XYZ(xyY);
+}
+
+float3 outputTransform
+(
+	float3 inValue,
+	TsParams PARAMS,
+	FColorSpace DISPLAY_PRI,
+	FColorSpace LIMITING_PRI,
+	int EOTF,
+	int SURROUND,
+	bool STRETCH_BLACK = true,
+	bool D60_SIM = false,
+	bool LEGAL_RANGE = false
+)
+{
+	float3x3 XYZ_2_DISPLAY_PRI_MAT = DISPLAY_PRI.XYZtoRGB;
+
+	/*
+		NOTE: This is a bit of a hack - probably a more direct way to do this.
+		TODO: Fix in future version
+	*/
+
+	float Y_MIN = PARAMS.Min.y;
+	float Y_MAX = PARAMS.Max.y;
+
+	// RRT sweeteners
+	float3 rgbPre = rrt_sweeteners(inValue);
+
+	// Apply the tonescale independently in rendering-space RGB
+	float3 rgbPost = ssts_f3(rgbPre, PARAMS);
+
+	// At this point data encoded AP1, scaled absolute luminance (cd/m^2)
+
+	/*  Scale absolute luminance to linear code value  */
+	float3 linearCV = Y_2_linCV_f3(rgbPost, Y_MAX, Y_MIN);
+
+	// Rendering primaries to XYZ
+	float3 XYZ = mul(AP1_2_XYZ_MAT, linearCV);
+
+	// Apply gamma adjustment to compensate for dim surround
+	/*
+		NOTE: This is more or less a placeholder block and is largely inactive
+		in its current form. This section currently only applies for SDR, and
+		even then, only in very specific cases.
+		In the future it is fully intended for this module to be updated to
+		support surround compensation regardless of luminance dynamic range. */
+		/*
+			TOD0: Come up with new surround compensation algorithm, applicable
+			across all dynamic ranges and supporting dark/dim/normal surround.
+		*/
+	if (SURROUND == 0) { // Dark surround
+		/*
+		Current tone scale is designed for dark surround environment so no
+		adjustment is necessary.
+		*/
+	}
+	else if (SURROUND == 1) { // Dim surround
+	 // INACTIVE for HDR and crudely implemented for SDR (see comment below)        
+		if ((EOTF == 1) || (EOTF == 2) || (EOTF == 3)) {
+			/*
+			This uses a crude logical assumption that if the EOTF is BT.1886,
+			sRGB, or gamma 2.6 that the data is SDR and so the SDR gamma
+			compensation factor from v1.0 will apply.
+			*/
+			XYZ = dark_to_dim(XYZ); /*
+			This uses a local dark_to_dim function that is designed to take in
+			XYZ and return XYZ rather than AP1 as is currently in the functions
+			in 'ACESlib.ODT_Common.ctl' */
+		}
+	}
+	else if (SURROUND == 2) { // Normal surround
+	 // INACTIVE - this does NOTHING
+	}
+
+	// Gamut limit to limiting primaries
+	// NOTE: Would be nice to just say
+	//    if (LIMITING_PRI != DISPLAY_PRI)
+	// but you can't because Chromaticities do not work with bool comparison operator
+	// For now, limit no matter what.
+	XYZ = limit_to_primaries(XYZ, LIMITING_PRI);
+
+	// Apply CAT from ACES white point to assumed observer adapted white point
+	// TODO: Needs to expand from just supporting D60 sim to allow for any
+	// observer adapted white point.
+	if (D60_SIM == false) {
+		if (DISPLAY_PRI.White != WhiteStandard_D60) {
+			XYZ = mul(D60_2_D65_CAT, XYZ);
+		}
+	}
+
+	// CIE XYZ to display encoding primaries
+	linearCV = mul(XYZ_2_DISPLAY_PRI_MAT, XYZ);
+
+	// Scale to avoid clipping when device calibration is different from D60. 
+	// To simulate D60, unequal code values are sent to the display.
+	// TODO: Needs to expand from just supporting D60 sim to allow for any
+	// observer adapted white point.
+	if (D60_SIM == true) {
+		/* TODO: The scale requires calling itself. Scale is same no matter the luminance.
+		   Currently precalculated for D65, DCI. If DCI, roll_white_fwd is used also.
+		   This needs a more complex algorithm to handle all cases.
+		*/
+		float SCALE = 1.0;
+		if (DISPLAY_PRI.White == WhiteStandard_D65) { // D65
+			SCALE = 0.96362;
+		}
+		else if (DISPLAY_PRI.White == WhiteStandard_DCI) { // DCI
+			linearCV[0] = roll_white_fwd(linearCV[0], 0.918, 0.5);
+			linearCV[1] = roll_white_fwd(linearCV[1], 0.918, 0.5);
+			linearCV[2] = roll_white_fwd(linearCV[2], 0.918, 0.5);
+			SCALE = 0.96;
+		}
+		linearCV *= SCALE;
+	}
+
+
+	// Clip values < 0 (i.e. projecting outside the display primaries)
+	// NOTE: P3 red and values close to it fall outside of Rec.2020 green-red 
+	// boundary
+	linearCV = clamp(linearCV, 0.0.xxx, HALF_POS_INF.xxx);
+
+	// EOTF
+	// 0: ST-2084 (PQ)
+	// 1: BT.1886 (Rec.709/2020 settings)
+	// 2: sRGB (mon_curve w/ presets)
+	//    moncurve_r with gamma of 2.4 and offset of 0.055 matches the EOTF found in IEC 61966-2-1:1999 (sRGB)
+	// 3: gamma 2.6
+	// 4: linear (no EOTF)
+	// 5: HLG
+	float3 outputCV;
+	if (EOTF == 0) {  // ST-2084 (PQ)
+		// NOTE: This is a kludgy way of ensuring a PQ code value of 0. Ideally,
+		// luminance would map directly to code value, but colorists don't like
+		// that. Might just need the tonescale to go darker so that darkest
+		// values through the tone scale quantize to code value of 0.
+		if (STRETCH_BLACK == true) {
+			outputCV = Y_2_ST2084_f3(clamp(linCV_2_Y_f3(linearCV, Y_MAX, 0.0), 0.0, HALF_POS_INF));
+		}
+		else {
+			outputCV = Y_2_ST2084_f3(linCV_2_Y_f3(linearCV, Y_MAX, Y_MIN));
+		}
+	}
+	else if (EOTF == 1) { // BT.1886 (Rec.709/2020 settings)
+		outputCV = bt1886_r_f3(linearCV, 2.4, 1.0, 0.0);
+	}
+	else if (EOTF == 2) { // sRGB (mon_curve w/ presets)
+		outputCV = moncurve_r_f3(linearCV, 2.4, 0.055);
+	}
+	else if (EOTF == 3) { // gamma 2.6
+		outputCV = pow(linearCV, (1. / 2.6).xxx);
+	}
+	else if (EOTF == 4) { // linear
+	 ///+ have linear behave the same way at ST-2084 regarding STRETCH_BLACK option
+		if (STRETCH_BLACK == true) {
+			outputCV = clamp(linCV_2_Y_f3(linearCV, Y_MAX, 0.0), 0.0, HALF_POS_INF);
+		}
+		else {
+			outputCV = linCV_2_Y_f3(linearCV, Y_MAX, Y_MIN);
+		}
+	}
+	else if (EOTF == 5) { // HLG
+	 // NOTE: HLG just maps ST.2084 output to HLG encoding. 
+	 // TODO: Restructure if/else tree to minimize this redundancy.
+		if (STRETCH_BLACK == true) {
+			outputCV = Y_2_ST2084_f3(clamp(linCV_2_Y_f3(linearCV, Y_MAX, 0.0), 0.0, HALF_POS_INF));
+		}
+		else {
+			outputCV = Y_2_ST2084_f3(linCV_2_Y_f3(linearCV, Y_MAX, Y_MIN));
+		}
+		outputCV = ST2084_2_HLG_1000nits_f3(outputCV);
+	}
+
+	if (LEGAL_RANGE == true) {
+		outputCV = fullRange_to_smpteRange_f3(outputCV);
+	}
+
+	return outputCV;
+}
+
+float3 outputTransform
+(
+	float3 inValue,
+	float Y_MIN,
+	float Y_MID,
+	float Y_MAX,
+	FColorSpace DISPLAY_PRI,
+	FColorSpace LIMITING_PRI,
+	int EOTF,
+	int SURROUND,
+	bool STRETCH_BLACK = true,
+	bool D60_SIM = false,
+	bool LEGAL_RANGE = false
+)
+{
+	/*
+		NOTE: This is a bit of a hack - probably a more direct way to do this.
+		TODO: Fix in future version
+	*/
+	TsParams PARAMS_DEFAULT = init_TsParams(Y_MIN, Y_MAX);
+	float expShift = log2(inv_ssts(Y_MID, PARAMS_DEFAULT)) - log2(0.18);
+	TsParams PARAMS = init_TsParams(Y_MIN, Y_MAX, expShift);
+
+	return outputTransform(inValue, PARAMS, DISPLAY_PRI, LIMITING_PRI, EOTF, SURROUND, STRETCH_BLACK, D60_SIM, LEGAL_RANGE);
+}
+
+float3 invOutputTransform
+(
+	float3 inValue,
+	float Y_MIN,
+	float Y_MID,
+	float Y_MAX,
+	FColorSpace DISPLAY_PRI,
+	FColorSpace LIMITING_PRI,
+	int EOTF,
+	int SURROUND,
+	bool STRETCH_BLACK = true,
+	bool D60_SIM = false,
+	bool LEGAL_RANGE = false
+)
+{
+	float3x3 DISPLAY_PRI_2_XYZ_MAT = DISPLAY_PRI.RGBtoXYZ;
+
+	/*
+		NOTE: This is a bit of a hack - probably a more direct way to do this.
+		TODO: Update in accordance with forward algorithm.
+	*/
+	TsParams PARAMS_DEFAULT = init_TsParams(Y_MIN, Y_MAX);
+	float expShift = log2(inv_ssts(Y_MID, PARAMS_DEFAULT)) - log2(0.18);
+	TsParams PARAMS = init_TsParams(Y_MIN, Y_MAX, expShift);
+
+	float3 outputCV = inValue;
+
+	if (LEGAL_RANGE == true) {
+		outputCV = smpteRange_to_fullRange_f3(outputCV);
+	}
+
+	// Inverse EOTF
+	// 0: ST-2084 (PQ)
+	// 1: BT.1886 (Rec.709/2020 settings)
+	// 2: sRGB (mon_curve w/ presets)
+	//    moncurve_r with gamma of 2.4 and offset of 0.055 matches the EOTF found in IEC 61966-2-1:1999 (sRGB)
+	// 3: gamma 2.6
+	// 4: linear (no EOTF)
+	// 5: HLG
+	float3 linearCV;
+	if (EOTF == 0) {  // ST-2084 (PQ)
+		if (STRETCH_BLACK == true) {
+			linearCV = Y_2_linCV_f3(ST2084_2_Y_f3(outputCV), Y_MAX, 0.);
+		}
+		else {
+			linearCV = Y_2_linCV_f3(ST2084_2_Y_f3(outputCV), Y_MAX, Y_MIN);
+		}
+	}
+	else if (EOTF == 1) { // BT.1886 (Rec.709/2020 settings)
+		linearCV = bt1886_f_f3(outputCV, 2.4, 1.0, 0.0);
+	}
+	else if (EOTF == 2) { // sRGB (mon_curve w/ presets)
+		linearCV = moncurve_f_f3(outputCV, 2.4, 0.055);
+	}
+	else if (EOTF == 3) { // gamma 2.6
+		linearCV = pow(outputCV, 2.6);
+	}
+	else if (EOTF == 4) { // linear
+		linearCV = Y_2_linCV_f3(outputCV, Y_MAX, Y_MIN);
+	}
+	else if (EOTF == 5) { // HLG
+		outputCV = HLG_2_ST2084_1000nits_f3(outputCV);
+		if (STRETCH_BLACK == true) {
+			linearCV = Y_2_linCV_f3(ST2084_2_Y_f3(outputCV), Y_MAX, 0.);
+		}
+		else {
+			linearCV = Y_2_linCV_f3(ST2084_2_Y_f3(outputCV), Y_MAX, Y_MIN);
+		}
+	}
+
+	// Un-scale
+	if (D60_SIM == true) {
+		/* TODO: The scale requires calling itself. Need an algorithm for this.
+			Scale is same no matter the luminance.
+			Currently using precalculated values for D65, DCI.
+			If DCI, roll_white_fwd is used also.
+		*/
+		float SCALE = 1.0;
+		if (DISPLAY_PRI.White == WhiteStandard_D65) { // D65
+			SCALE = 0.96362;
+			linearCV /= SCALE;
+		}
+		else if (DISPLAY_PRI.White == WhiteStandard_DCI) { // DCI
+			SCALE = 0.96;
+			linearCV[0] = roll_white_rev(linearCV[0] / SCALE, 0.918, 0.5);
+			linearCV[1] = roll_white_rev(linearCV[1] / SCALE, 0.918, 0.5);
+			linearCV[2] = roll_white_rev(linearCV[2] / SCALE, 0.918, 0.5);
+		}
+
+	}
+
+	// Encoding primaries to CIE XYZ
+	float3 XYZ = mul(DISPLAY_PRI_2_XYZ_MAT, linearCV);
+
+	// Undo CAT from assumed observer adapted white point to ACES white point
+	if (D60_SIM == false) {
+		if (DISPLAY_PRI.White != WhiteStandard_D60) {
+			XYZ = mul(D65_2_D60_CAT, XYZ);
+		}
+	}
+
+	// Apply gamma adjustment to compensate for dim surround
+	/*
+		NOTE: This is more or less a placeholder block and is largely inactive
+		in its current form. This section currently only applies for SDR, and
+		even then, only in very specific cases.
+		In the future it is fully intended for this module to be updated to
+		support surround compensation regardless of luminance dynamic range. */
+		/*
+			TOD0: Come up with new surround compensation algorithm, applicable
+			across all dynamic ranges and supporting dark/dim/normal surround.
+		*/
+	if (SURROUND == 0) { // Dark surround
+		/*
+		Current tone scale is designed for dark surround environment so no
+		adjustment is necessary.
+		*/
+	}
+	else if (SURROUND == 1) { // Dim surround
+	 // INACTIVE for HDR and crudely implemented for SDR (see comment below)        
+		if ((EOTF == 1) || (EOTF == 2) || (EOTF == 3)) {
+			/*
+			This uses a crude logical assumption that if the EOTF is BT.1886,
+			sRGB, or gamma 2.6 that the data is SDR and so the SDR gamma
+			compensation factor from v1.0 will apply.
+			*/
+			XYZ = dim_to_dark(XYZ); /*
+			This uses a local dim_to_dark function that is designed to take in
+			XYZ and return XYZ rather than AP1 as is currently in the functions
+			in 'ACESlib.ODT_Common.ctl' */
+		}
+	}
+	else if (SURROUND == 2) { // Normal surround
+	 // INACTIVE - this does NOTHING
+	}
+
+	// XYZ to rendering primaries
+	linearCV = mul(XYZ_2_AP1_MAT, XYZ);
+
+	float3 rgbPost = linCV_2_Y_f3(linearCV, Y_MAX, Y_MIN);
+
+	// Apply the inverse tonescale independently in rendering-space RGB
+	float3 rgbPre = inv_ssts_f3(rgbPost, PARAMS);
+
+	// RRT sweeteners
+	float3 aces = inv_rrt_sweeteners(rgbPre);
+
+	return aces;
+}
+
+/* --- Gamut Compress Parameters --- */
+// Distance from achromatic which will be compressed to the gamut boundary
+// Values calculated to encompass the encoding gamuts of common digital cinema cameras
+static const float LIM_CYAN = 1.147;
+static const float LIM_MAGENTA = 1.264;
+static const float LIM_YELLOW = 1.312;
+
+// Percentage of the core gamut to protect
+// Values calculated to protect all the colors of the ColorChecker Classic 24 as given by
+// ISO 17321-1 and Ohta (1997)
+static const float THR_CYAN = 0.815;
+static const float THR_MAGENTA = 0.803;
+static const float THR_YELLOW = 0.880;
+
+// Aggressiveness of the compression curve
+static const float PWR = 1.2;
+
+
+
+// Calculate compressed distance
+float compress(float dist, float lim, float thr, float pwr, bool invert)
+{
+	float comprDist;
+	float scl;
+	float nd;
+	float p;
+
+	if (dist < thr) {
+		comprDist = dist; // No compression below threshold
+	}
+	else {
+		// Calculate scale factor for y = 1 intersect
+		scl = (lim - thr) / pow(pow((1.0 - thr) / (lim - thr), -pwr) - 1.0, 1.0 / pwr);
+
+		// Normalize distance outside threshold by scale factor
+		nd = (dist - thr) / scl;
+		p = pow(nd, pwr);
+
+		if (!invert) {
+			comprDist = thr + scl * nd / (pow(1.0 + p, 1.0 / pwr)); // Compress
+		}
+		else {
+			if (dist > (thr + scl)) {
+				comprDist = dist; // Avoid singularity
+			}
+			else {
+				comprDist = thr + scl * pow(-(p / (p - 1.0)), 1.0 / pwr); // Uncompress
+			}
+		}
+	}
+
+	return comprDist;
+}
+
+float max_f3(float3 a)
+{
+	return max(a[0], max(a[1], a[2]));
+}
+
+float3 compressColor
+(
+	const float3 ACES,
+	const bool invert = false
+)
+{
+	// Convert to ACEScg
+	float3 linAP1 = mul(AP0_2_AP1_MAT, ACES);
+
+	// Achromatic axis
+	float ach = max_f3(linAP1);
+
+	// Distance from the achromatic axis for each color component aka inverse RGB ratios
+	float3 comprLinAP1;
+	///+ added any(ACES < 0): don't try to compress colors outside the AP0 color space: clamping will happen anyway during RRT/ODT, 
+	// and this allows to keep proper results when testing with a synthetic chart
+	if (ach < 1e-10f || any(ACES < 0)) {
+		///-
+		comprLinAP1 = linAP1;
+	}
+	else {
+		float3 dist;
+		dist[0] = (ach - linAP1[0]) / abs(ach);
+		dist[1] = (ach - linAP1[1]) / abs(ach);
+		dist[2] = (ach - linAP1[2]) / abs(ach);
+
+		float3 comprDist = {
+		compress(dist[0], LIM_CYAN, THR_CYAN, PWR, invert),
+		compress(dist[1], LIM_MAGENTA, THR_MAGENTA, PWR, invert),
+		compress(dist[2], LIM_YELLOW, THR_YELLOW, PWR, invert)
+		};
+
+		// Recalculate RGB from compressed distance and achromatic
+		comprLinAP1 = float3(
+			ach - comprDist[0] * abs(ach),
+			ach - comprDist[1] * abs(ach),
+			ach - comprDist[2] * abs(ach)
+		);
+	}
+
+	// Convert back to ACES2065-1
+	return mul(AP1_2_AP0_MAT, comprLinAP1);
+}
+
diff --git a/Engine/Shaders/Private/PostProcessCombineLUTs.usf b/Engine/Shaders/Private/PostProcessCombineLUTs.usf
--- a/Engine/Shaders/Private/PostProcessCombineLUTs.usf
+++ b/Engine/Shaders/Private/PostProcessCombineLUTs.usf
@@ -76,6 +76,11 @@ float ACESCoefsHigh_4;
 float ACESSceneColorMultiplier;
 float ACESGamutCompression;
 
+// ACES 2.0
+Texture2D<float> ACESReachTable;
+Texture2D<float3> ACESGamutTable;
+Texture2D<float> ACESGammaTable;
+
 float3 ColorCorrect( float3 WorkingColor,
 	float4 ColorSaturation,
 	float4 ColorContrast,
@@ -349,9 +354,12 @@ float4 CombineLUTsCommon(float2 InUV, uint InLayerIndex)
 	// ACES transform with PQ/2084 encoding, user specified gamut 
 	else if( GetOutputDevice() == TONEMAPPER_OUTPUT_ACES1000nitST2084 || GetOutputDevice() == TONEMAPPER_OUTPUT_ACES2000nitST2084)
 	{
+#if USE_ACES_2
+		float3 ODTColor = ACESOutputTransform(GradedColor * ACESSceneColorMultiplier, (float3x3) WorkingColorSpace.ToAP0, OutputMaxLuminance, ACESReachTable, ACESGamutTable, ACESGammaTable);
+#else
 		FACESTonemapParams AcesParams = ComputeACESTonemapParams(ACESMinMaxData, ACESMidData, ACESCoefsLow_0, ACESCoefsHigh_0, ACESCoefsLow_4, ACESCoefsHigh_4, ACESSceneColorMultiplier, ACESGamutCompression);
 		float3 ODTColor = ACESOutputTransform( GradedColor, (float3x3)WorkingColorSpace.ToAP0, AcesParams);
-
+#endif
 		ODTColor = mul(AP1_2_Output, ODTColor);
 
 		// Apply conversion to ST-2084 (Dolby PQ)
@@ -361,9 +369,12 @@ float4 CombineLUTsCommon(float2 InUV, uint InLayerIndex)
 	// ACES transform with PQ/2084 encoding, AP1 gamut
 	else if( GetOutputDevice() == TONEMAPPER_OUTPUT_ACES1000nitScRGB || GetOutputDevice() == TONEMAPPER_OUTPUT_ACES2000nitScRGB)
 	{
+#if USE_ACES_2
+		float3 ODTColor = ACESOutputTransform(GradedColor * ACESSceneColorMultiplier, (float3x3) WorkingColorSpace.ToAP0, OutputMaxLuminance, ACESReachTable, ACESGamutTable, ACESGammaTable);
+#else
 		FACESTonemapParams AcesParams = ComputeACESTonemapParams(ACESMinMaxData, ACESMidData, ACESCoefsLow_0, ACESCoefsHigh_0, ACESCoefsLow_4, ACESCoefsHigh_4, ACESSceneColorMultiplier, ACESGamutCompression);
 		float3 ODTColor = ACESOutputTransform( GradedColor, (float3x3)WorkingColorSpace.ToAP0, AcesParams);
-		
+#endif
 		OutDeviceColor = mul(AP1_2_sRGB, ODTColor / UE_SCRGB_WHITE_NITS);
 	}	
 	
diff --git a/Engine/Shaders/Private/PostProcessDeviceEncodingOnly.usf b/Engine/Shaders/Private/PostProcessDeviceEncodingOnly.usf
--- a/Engine/Shaders/Private/PostProcessDeviceEncodingOnly.usf
+++ b/Engine/Shaders/Private/PostProcessDeviceEncodingOnly.usf
@@ -36,6 +36,11 @@ float ACESCoefsHigh_4;
 float ACESSceneColorMultiplier;
 float ACESGamutCompression;
 
+// ACES 2.0
+Texture2D<float> ACESReachTable;
+Texture2D<float3> ACESGamutTable;
+Texture2D<float> ACESGammaTable;
+
 uint GetOutputDevice()
 {
 #if OUTPUT_DEVICE_SRGB
@@ -98,11 +103,14 @@ float4 DeviceEncodingOnlyCommonPS(
 
 	// ACES transform with PQ/2084 encoding, user specified gamut 
 	else if (GetOutputDevice() == TONEMAPPER_OUTPUT_ACES1000nitST2084 || GetOutputDevice() == TONEMAPPER_OUTPUT_ACES2000nitST2084)
-	{		 
+	{
+#if USE_ACES_2
+		float3 ODTColor = ACESOutputTransform(GradedColor * ACESSceneColorMultiplier, (float3x3) WorkingColorSpace.ToAP0, OutputMaxLuminance, ACESReachTable, ACESGamutTable, ACESGammaTable);
+#else
 		// ACES ODT
 		FACESTonemapParams AcesParams = ComputeACESTonemapParams(ACESMinMaxData, ACESMidData, ACESCoefsLow_0, ACESCoefsHigh_0, ACESCoefsLow_4, ACESCoefsHigh_4, ACESSceneColorMultiplier, ACESGamutCompression);
 		float3 ODTColor = ACESOutputTransform( FilmColor, (float3x3)WorkingColorSpace.ToAP0, AcesParams);
-
+#endif
 		// Convert from AP1 to specified output gamut
 		ODTColor = mul( AP1_2_Output, ODTColor );
 
@@ -113,11 +121,14 @@ float4 DeviceEncodingOnlyCommonPS(
 	// ACES transform to linear ScRGB
 	else if (GetOutputDevice() == TONEMAPPER_OUTPUT_ACES1000nitScRGB || GetOutputDevice() == TONEMAPPER_OUTPUT_ACES2000nitScRGB)
 	{
+#if USE_ACES_2
+		float3 ODTColor = ACESOutputTransform(GradedColor * ACESSceneColorMultiplier, (float3x3) WorkingColorSpace.ToAP0, OutputMaxLuminance, ACESReachTable, ACESGamutTable, ACESGammaTable);
+#else
 		FACESTonemapParams AcesParams = ComputeACESTonemapParams(ACESMinMaxData, ACESMidData, ACESCoefsLow_0, ACESCoefsHigh_0, ACESCoefsLow_4, ACESCoefsHigh_4, ACESSceneColorMultiplier, ACESGamutCompression);
 
 		// ACES ODT
 		float3 ODTColor = ACESOutputTransform(FilmColor, (float3x3) WorkingColorSpace.ToAP0, AcesParams);
-
+#endif
 		// Apply conversion to ScRGB
 		OutDeviceColor = mul( AP1_2_sRGB, ODTColor / UE_SCRGB_WHITE_NITS );
 	}
diff --git a/Engine/Shaders/Private/PostProcessSelectionOutline.usf b/Engine/Shaders/Private/PostProcessSelectionOutline.usf
--- a/Engine/Shaders/Private/PostProcessSelectionOutline.usf
+++ b/Engine/Shaders/Private/PostProcessSelectionOutline.usf
@@ -3,7 +3,7 @@
 #include "Common.ush"
 #include "ScreenPass.ush"
 #include "GammaCorrectionCommon.ush"
-#include "ACES.ush"
+#include "ACES/ACESCommon.ush"
 #include "LensDistortion.ush"
 
 
diff --git a/Engine/Shaders/Private/TonemapCommon.ush b/Engine/Shaders/Private/TonemapCommon.ush
--- a/Engine/Shaders/Private/TonemapCommon.ush
+++ b/Engine/Shaders/Private/TonemapCommon.ush
@@ -6,8 +6,18 @@
 
 #pragma once
 
-#include "ACES.ush"
+#ifndef USE_ACES_2
+	#define USE_ACES_2 0
+#endif
+
 #include "GammaCorrectionCommon.ush"
+#include "ACES/ACESCommon.ush"
+#if USE_ACES_2
+	#include "ACES/ACESOutputTransform.ush"
+	#include "ACES/ACESDisplayEncoding.ush"
+#else
+	#include "ACES/ACES_v1.3.ush"
+#endif
 
 /** Values of output device (matching C++'s EDisplayOutputFormat) */
 #define TONEMAPPER_OUTPUT_sRGB 0
@@ -309,6 +319,43 @@ static const bool UE_ACES_LEGAL_RANGE = false;
 
 static const float UE_SCRGB_WHITE_NITS = 80.0; // ScRGB D65 white luminance (cd/m^2)
 
+#if USE_ACES_2
+float3 ACESOutputTransform(
+	float3 SceneReferredLinearColor,
+	const float3x3 WCS_2_AP0,
+	float PeakLuminance,
+	Texture2D<float> ReachTable,
+    Texture2D<float3> GamutTable,
+    Texture2D<float> GammaTable
+)
+{
+	// TODO: The display limiting color space should be user-exposed as part of display characterization.
+	const FColorSpace LimitingColorSpace = GetColorSpace(TONEMAPPER_GAMUT_ACEScg_D60);
+	const FColorSpace EncodingColorSpace = GetColorSpace(TONEMAPPER_GAMUT_ACEScg_D60);
+
+	float3 Aces = mul( WCS_2_AP0, SceneReferredLinearColor );
+	
+	ODTParams PARAMS = init_ODTParams(
+		PeakLuminance,
+		LimitingColorSpace.RGBtoXYZ,
+		LimitingColorSpace.XYZtoRGB,
+		EncodingColorSpace.RGBtoXYZ,
+		EncodingColorSpace.XYZtoRGB,
+		ReachTable,
+		GamutTable,
+		GammaTable
+	);
+
+	float3 ColorXYZ = outputTransform_fwd(Aces, PeakLuminance, PARAMS);
+
+	// Scale white is usually only enabled when using a limiting white different from the encoding white
+	static const bool scale_white = false;
+	ColorXYZ = white_limiting(ColorXYZ, PARAMS, scale_white);
+
+	// Scale by reference luminance to return values in nits, matching our 1.3 transform
+	return mul(XYZ_2_AP1_MAT, ColorXYZ) * referenceLuminance;
+}
+#else // USE_ACES_2
 struct FACESTonemapParams
 {
 	TsParams AcesParams;
@@ -377,6 +424,7 @@ float3 ACESOutputTransform( float3 SceneReferredLinearsRGBColor, FACESTonemapPar
 
 	return ACESOutputTransform( SceneReferredLinearsRGBColor, sRGB_2_AP0, InACESTonemapParams);
 }
+#endif //USE_ACES_2
 
 static const float3x3 GamutMappingIdentityMatrix = { 1.0, 0.0, 0.0, 0.0, 1.0, 0.0, 0.0, 0.0, 1.0 };
 
diff --git a/Engine/Source/Runtime/RenderCore/Private/RenderCore.cpp b/Engine/Source/Runtime/RenderCore/Private/RenderCore.cpp
--- a/Engine/Source/Runtime/RenderCore/Private/RenderCore.cpp
+++ b/Engine/Source/Runtime/RenderCore/Private/RenderCore.cpp
@@ -422,6 +422,15 @@ static TAutoConsoleVariable<float> CVarHDRAcesGamutCompression(
 	ECVF_RenderThreadSafe
 );
 
+static TAutoConsoleVariable<int32> CVarHDRAcesVersion(
+	TEXT("r.HDR.Aces.Version"),
+	1,
+	TEXT("Version of the ACES transform:\n")
+	TEXT("1: ACES 1.3 (default)\n")
+	TEXT("2: ACES 2.0\n"),
+	ECVF_RenderThreadSafe
+);
+
 TAutoConsoleVariable<int32> CVarHDROutputEnabled(
 	TEXT("r.HDR.EnableHDROutput"),
 	0,
diff --git a/Engine/Source/Runtime/Renderer/Private/PostProcess/ACESUtils.cpp b/Engine/Source/Runtime/Renderer/Private/PostProcess/ACESUtils.cpp
new file mode 100644
--- /dev/null
+++ b/Engine/Source/Runtime/Renderer/Private/PostProcess/ACESUtils.cpp
@@ -0,0 +1,897 @@
+// Copyright Epic Games, Inc. All Rights Reserved.
+
+#include "ACESUtils.h"
+
+#include "ColorManagement/ColorSpace.h"
+#include "GlobalRenderResources.h"
+#include "RenderGraphBuilder.h"
+#include "RenderResource.h"
+#include "RHICommandList.h"
+#include "RHIStaticStates.h"
+
+namespace UE::Color::ACES
+{
+/**
+* Source code adapted from OpenColorIO/src/OpenColorIO/ops/fixedfunction/ACES2 (v2.4.1)
+*
+* Copyright Contributors to the OpenColorIO Project.
+
+Redistribution and use in source and binary forms, with or without
+modification, are permitted provided that the following conditions are
+met:
+
+* Redistributions of source code must retain the above copyright
+  notice, this list of conditions and the following disclaimer.
+* Redistributions in binary form must reproduce the above copyright
+  notice, this list of conditions and the following disclaimer in the
+  documentation and/or other materials provided with the distribution.
+* Neither the name of the copyright holder nor the names of its
+  contributors may be used to endorse or promote products derived from
+  this software without specific prior written permission.
+
+THIS SOFTWARE IS PROVIDED BY THE COPYRIGHT HOLDERS AND CONTRIBUTORS
+"AS IS" AND ANY EXPRESS OR IMPLIED WARRANTIES, INCLUDING, BUT NOT
+LIMITED TO, THE IMPLIED WARRANTIES OF MERCHANTABILITY AND FITNESS FOR
+A PARTICULAR PURPOSE ARE DISCLAIMED. IN NO EVENT SHALL THE COPYRIGHT
+HOLDER OR CONTRIBUTORS BE LIABLE FOR ANY DIRECT, INDIRECT, INCIDENTAL,
+SPECIAL, EXEMPLARY, OR CONSEQUENTIAL DAMAGES (INCLUDING, BUT NOT
+LIMITED TO, PROCUREMENT OF SUBSTITUTE GOODS OR SERVICES; LOSS OF USE,
+DATA, OR PROFITS; OR BUSINESS INTERRUPTION) HOWEVER CAUSED AND ON ANY
+THEORY OF LIABILITY, WHETHER IN CONTRACT, STRICT LIABILITY, OR TORT
+(INCLUDING NEGLIGENCE OR OTHERWISE) ARISING IN ANY WAY OUT OF THE USE
+OF THIS SOFTWARE, EVEN IF ADVISED OF THE POSSIBILITY OF SUCH DAMAGE.
+*/
+
+constexpr int32 TABLE_SIZE = 360;
+constexpr int32 TABLE_ADDITION_ENTRIES = 2;
+constexpr int32 TABLE_TotalSize = TABLE_SIZE + TABLE_ADDITION_ENTRIES;
+constexpr int32 GAMUT_TABLE_BASE_INDEX = 1;
+
+constexpr float ReferenceLuminance = 100.f;
+constexpr float L_A = 100.f;
+constexpr float Y_b = 20.f;
+constexpr float AcResp = 1.f;
+constexpr float Ra = 2.f * AcResp;
+constexpr float Ba = 0.05f + (2.f - Ra);
+constexpr float Surround[3] = { 0.9f, 0.59f, 0.9f }; // Dim surround
+
+// Gamut compression
+constexpr float SmoothCusps = 0.12f;
+constexpr float SmoothM = 0.27f;
+constexpr float CuspMidBlend = 1.3f;
+constexpr float FocusGainBlend = 0.3f;
+constexpr float FocusAdjustGain = 0.55f;
+constexpr float FocusDistance = 1.35f;
+constexpr float FocusDistanceScaling = 1.75f;
+constexpr float CompressionThreshold = 0.75f;
+
+// Table generation
+constexpr float GammaMinimum = 0.0f;
+constexpr float GammaMaximum = 5.0f;
+constexpr float GammaSearchStep = 0.4f;
+constexpr float GammaAccuracy = 1e-5f;
+
+struct JMhParams
+{
+	float F_L;
+	float z;
+	float A_w;
+	float A_w_J;
+	FVector3f XYZ_w;
+	FVector3f D_RGB;
+	FMatrix44f MATRIX_RGB_to_CAM16;
+	FMatrix44f MATRIX_CAM16_to_RGB;
+};
+
+struct ToneScaleParams
+{
+	float n;
+	float n_r;
+	float g;
+	float t_1;
+	float c_t;
+	float s_2;
+	float u_2;
+	float m_2;
+};
+
+struct Table3D
+{
+	static constexpr int32 BaseIndex = GAMUT_TABLE_BASE_INDEX;
+	static constexpr int32 Size = TABLE_SIZE;
+	static constexpr int32 TotalSize = TABLE_TotalSize;
+	FVector3f Data[TABLE_TotalSize];
+};
+
+struct Table1D
+{
+	static constexpr int32 BaseIndex = GAMUT_TABLE_BASE_INDEX;
+	static constexpr int32 Size = TABLE_SIZE;
+	static constexpr int32 TotalSize = TABLE_TotalSize;
+	float Data[TABLE_TotalSize];
+};
+
+// Post adaptation non linear response compression
+float PanlrcForward(float Value, float F_L)
+{
+	const float F_L_v = FMath::Pow(F_L * FMath::Abs(Value) / ReferenceLuminance, 0.42f);
+	return (400.f * FMath::CopySign(1.f, Value) * F_L_v) / (27.13f + F_L_v);
+}
+
+float PanlrcInverse(float Value, float F_L)
+{
+	return FMath::CopySign(1.f, Value) * ReferenceLuminance / F_L * FMath::Pow((27.13f * FMath::Abs(Value) / (400.f - FMath::Abs(Value))), 1.f / 0.42f);
+}
+
+inline bool AnyBelowZero(const FVector3f& RGB)
+{
+	return (RGB[0] < 0. || RGB[1] < 0. || RGB[2] < 0.);
+}
+
+bool OutsideHull(const FVector3f& RGB)
+{
+	// limit value, once we cross this value, we are outside of the top gamut shell
+	constexpr float MaxRGBtestVal = 1.0;
+	return RGB[0] > MaxRGBtestVal || RGB[1] > MaxRGBtestVal || RGB[2] > MaxRGBtestVal;
+}
+
+// Optimization used during initialization
+float Y_to_J(float Y, const JMhParams& Params)
+{
+	float F_L_Y = FMath::Pow(Params.F_L * FMath::Abs(Y) / ReferenceLuminance, 0.42f);
+	return FMath::CopySign(1.f, Y) * ReferenceLuminance * FMath::Pow(((400.f * F_L_Y) / (27.13f + F_L_Y)) / Params.A_w_J, Surround[1] * Params.z);
+}
+
+float WrapTo360(float Hue)
+{
+	float Y = FMath::Fmod(Hue, 360.f);
+	if (Y < 0.f)
+	{
+		Y = Y + 360.f;
+	}
+
+	return Y;
+}
+
+int32 HuePositionInUniformTable(float Hue, int32 TableSize)
+{
+	const float WrappedHue = WrapTo360(Hue);
+
+	return int32(WrappedHue / 360.f * (float)TableSize);
+}
+
+int32 ClampToTableBounds(int32 Entry, int32 TableSize)
+{
+	return FMath::Min(TableSize - 1, FMath::Max(0, Entry));
+}
+
+float Smin(float a, float b, float s)
+{
+	const float h = FMath::Max(s - FMath::Abs(a - b), 0.f) / s;
+	return FMath::Min(a, b) - h * h * h * s * (1.f / 6.f);
+}
+
+FVector3f JMh_to_RGB(const FVector3f& JMh, const JMhParams& Params)
+{
+	const float J = JMh[0];
+	const float M = JMh[1];
+	const float h = JMh[2];
+
+	const float HRad = h * UE_PI / 180.f;
+
+	const float Scale = M / (43.f * Surround[2]);
+	const float A = Params.A_w * FMath::Pow(J / 100.f, 1.f / (Surround[1] * Params.z));
+	const float a = Scale * FMath::Cos(HRad);
+	const float b = Scale * FMath::Sin(HRad);
+
+	const float RedA = (460.f * A + 451.f * a + 288.f * b) / 1403.f;
+	const float GrnA = (460.f * A - 891.f * a - 261.f * b) / 1403.f;
+	const float BluA = (460.f * A - 220.f * a - 6300.f * b) / 1403.f;
+
+	FVector3f CamM;
+	CamM.X = PanlrcInverse(RedA, Params.F_L) / Params.D_RGB[0];
+	CamM.Y = PanlrcInverse(GrnA, Params.F_L) / Params.D_RGB[1];
+	CamM.Z = PanlrcInverse(BluA, Params.F_L) / Params.D_RGB[2];
+
+	return Params.MATRIX_CAM16_to_RGB.TransformVector(CamM);
+}
+
+JMhParams InitJMhParams(const FColorSpace& InColorSpace)
+{
+	static FColorSpace ColorSpaceCAM16 = FColorSpace(EColorSpace::ACESCAM16);
+
+	const FVector XYZ_w = InColorSpace.GetRgbToXYZ().TransformVector(FVector::One() * ReferenceLuminance);
+	const float Y_W = XYZ_w[1];
+	const FVector RGB_w = ColorSpaceCAM16.GetXYZToRgb().TransformVector(XYZ_w);
+
+	// Viewing condition dependent parameters
+	const float K = 1.f / (5.f * L_A + 1.f);
+	const float K4 = FMath::Pow(K, 4.f);
+	const float N = Y_b / Y_W;
+	const float F_L = 0.2f * K4 * (5.f * L_A) + 0.1f * FMath::Pow((1.f - K4), 2.f) * FMath::Pow(5.f * L_A, 1.f / 3.f);
+	const float z = 1.48f + FMath::Sqrt(N);
+
+	const FVector3f D_RGB = {
+		Y_W / (float)RGB_w[0],
+		Y_W / (float)RGB_w[1],
+		Y_W / (float)RGB_w[2]
+	};
+
+	const FVector3f RGB_WC{
+		D_RGB[0] * (float)RGB_w[0],
+		D_RGB[1] * (float)RGB_w[1],
+		D_RGB[2] * (float)RGB_w[2]
+	};
+
+	const FVector3f RGB_AW = {
+		PanlrcForward(RGB_WC[0], F_L),
+		PanlrcForward(RGB_WC[1], F_L),
+		PanlrcForward(RGB_WC[2], F_L)
+	};
+
+	const float A_w = Ra * RGB_AW[0] + RGB_AW[1] + Ba * RGB_AW[2];
+	const float F_L_W = FMath::Pow(F_L, 0.42f);
+	const float A_w_J = (400.f * F_L_W) / (27.13f + F_L_W);
+
+	JMhParams Params;
+	Params.XYZ_w = FVector3f(FVector3d(XYZ_w));
+	Params.F_L = F_L;
+	Params.z = z;
+	Params.D_RGB = D_RGB;
+	Params.A_w = A_w;
+	Params.A_w_J = A_w_J;
+
+	FColorSpaceTransform ToCAM16T(InColorSpace, ColorSpaceCAM16, EChromaticAdaptationMethod::None);
+	FMatrix ToCAM16 = ToCAM16T.ApplyScale(100.0);
+
+	Params.MATRIX_RGB_to_CAM16 = FMatrix44f(ToCAM16);
+	Params.MATRIX_CAM16_to_RGB = FMatrix44f(ToCAM16.Inverse());
+
+	return Params;
+}
+
+ToneScaleParams InitToneScaleParams(float PeakLuminance)
+{
+	// Preset constants that set the desired behavior for the curve
+	const float n = PeakLuminance;
+
+	const float n_r = 100.0f;    // normalized white in nits (what 1.0 should be)
+	const float g = 1.15f;       // surround / contrast
+	const float c = 0.18f;       // anchor for 18% grey
+	const float c_d = 10.013f;   // output luminance of 18% grey (in nits)
+	const float w_g = 0.14f;     // change in grey between different peak luminance
+	const float t_1 = 0.04f;     // shadow toe or flare/glare compensation
+	const float r_hit_min = 128.f;   // scene-referred value "hitting the roof"
+	const float r_hit_max = 896.f;   // scene-referred value "hitting the roof"
+
+	// Calculate output constants
+	const float r_hit = r_hit_min + (r_hit_max - r_hit_min) * (FMath::Loge(n / n_r) / FMath::Loge(10000.f / 100.f));
+	const float m_0 = (n / n_r);
+	const float m_1 = 0.5f * (m_0 + FMath::Sqrt(m_0 * (m_0 + 4.f * t_1)));
+	const float u = FMath::Pow((r_hit / m_1) / ((r_hit / m_1) + 1.f), g);
+	const float m = m_1 / u;
+	const float w_i = FMath::Loge(n / 100.f) / FMath::Loge(2.f);
+	const float c_t = c_d / n_r * (1.f + w_i * w_g);
+	const float g_ip = 0.5f * (c_t + FMath::Sqrt(c_t * (c_t + 4.f * t_1)));
+	const float g_ipp2 = -(m_1 * FMath::Pow((g_ip / m), (1.f / g))) / (FMath::Pow(g_ip / m, 1.f / g) - 1.f);
+	const float w_2 = c / g_ipp2;
+	const float s_2 = w_2 * m_1;
+	const float u_2 = FMath::Pow((r_hit / m_1) / ((r_hit / m_1) + w_2), g);
+	const float m_2 = m_1 / u_2;
+
+	ToneScaleParams TonescaleParams = {
+		n,
+		n_r,
+		g,
+		t_1,
+		c_t,
+		s_2,
+		u_2,
+		m_2
+	};
+
+	return TonescaleParams;
+}
+
+Table1D MakeReachMTable(float PeakLuminance)
+{
+	static FColorSpace ColorSpaceAP1 = FColorSpace(EColorSpace::ACESAP1);
+
+	const JMhParams Params = InitJMhParams(ColorSpaceAP1);
+	const float LimitJMax = Y_to_J(PeakLuminance, Params);
+
+	Table1D GamutReachTable{};
+
+	for (int32 Index = 0; Index < GamutReachTable.Size; Index++)
+	{
+		const float Hue = (float)Index;
+		const float SearchRange = 50.f;
+
+		float Low = 0.;
+		float High = Low + SearchRange;
+		bool Outside = false;
+
+		while ((Outside != true) && (High < 1300.f))
+		{
+			const FVector3f SearchJMh = FVector3f{ LimitJMax, High, Hue };
+			const FVector3f NewLimitRGB = JMh_to_RGB(SearchJMh, Params);
+			Outside = AnyBelowZero(NewLimitRGB);
+			
+			if (Outside == false)
+			{
+				Low = High;
+				High = High + SearchRange;
+			}
+		}
+
+		while (High - Low > 1e-2)
+		{
+			const float SampleM = (High + Low) / 2.f;
+			const FVector3f SearchJMh = FVector3f{ LimitJMax, SampleM, Hue };
+			const FVector3f NewLimitRGB = JMh_to_RGB(SearchJMh, Params);
+			Outside = AnyBelowZero(NewLimitRGB);
+			
+			if (Outside)
+			{
+				High = SampleM;
+			}
+			else
+			{
+				Low = SampleM;
+			}
+		}
+
+		GamutReachTable.Data[Index] = High;
+	}
+
+	return GamutReachTable;
+}
+
+FVector3f HSV_to_RGB(const FVector3f& HSV)
+{
+	const float C = HSV[2] * HSV[1];
+	const float X = C * (1.f - FMath::Abs(FMath::Fmod(HSV[0] * 6.f, 2.f) - 1.f));
+	const float m = HSV[2] - C;
+
+	FVector3f RGB{};
+	if (HSV[0] < 1.f / 6.f) {
+		RGB = { C, X, 0.f };
+	}
+	else if (HSV[0] < 2. / 6.) {
+		RGB = { X, C, 0.f };
+	}
+	else if (HSV[0] < 3. / 6.) {
+		RGB = { 0.f, C, X };
+	}
+	else if (HSV[0] < 4. / 6.) {
+		RGB = { 0.f, X, C };
+	}
+	else if (HSV[0] < 5. / 6.) {
+		RGB = { X, 0.f, C };
+	}
+	else {
+		RGB = { C, 0.f, X };
+	}
+	RGB = RGB + m;
+
+	return RGB;
+}
+
+FVector3f RGB_to_JMh(const FVector3f& RGB, const JMhParams& Params)
+{
+	const FVector3f RGB_m = Params.MATRIX_RGB_to_CAM16.TransformVector(RGB);
+
+	const float RedA = PanlrcForward(RGB_m[0] * Params.D_RGB[0], Params.F_L);
+	const float GrnA = PanlrcForward(RGB_m[1] * Params.D_RGB[1], Params.F_L);
+	const float BluA = PanlrcForward(RGB_m[2] * Params.D_RGB[2], Params.F_L);
+
+	const float A = 2.f * RedA + GrnA + 0.05f * BluA;
+	const float a = RedA - 12.f * GrnA / 11.f + BluA / 11.f;
+	const float b = (RedA + GrnA - 2.f * BluA) / 9.f;
+
+	const float J = 100.f * FMath::Pow(A / Params.A_w, Surround[1] * Params.z);
+
+	const float M = J == 0.f ? 0.f : 43.f * Surround[2] * FMath::Sqrt(a * a + b * b);
+
+	const float h_rad = FMath::Atan2(b, a);
+	float h = FMath::Fmod(h_rad * 180.f / UE_PI, 360.f);
+	if (h < 0.f)
+	{
+		h += 360.f;
+	}
+
+	return { J, M, h };
+}
+
+Table3D MakeGamutTable(const FColorSpace& InLimitingColorSpace, float PeakLuminance)
+{
+	const JMhParams Params = InitJMhParams(InLimitingColorSpace);
+
+	Table3D GamutCuspTableUnsorted{};
+	for (int32 i = 0; i < GamutCuspTableUnsorted.Size; i++)
+	{
+		const float hNorm = (float)i / GamutCuspTableUnsorted.Size;
+		const FVector3f HSV = { hNorm, 1., 1. };
+		const FVector3f RGB = HSV_to_RGB(HSV);
+		const FVector3f ScaledRGB = (PeakLuminance / ReferenceLuminance) * RGB;
+		const FVector3f JMh = RGB_to_JMh(ScaledRGB, Params);
+
+		GamutCuspTableUnsorted.Data[i] = JMh;
+	}
+
+	int32 MinhIndex = 0;
+	for (int32 i = 0; i < GamutCuspTableUnsorted.Size; i++)
+	{
+		if (GamutCuspTableUnsorted.Data[i][2] < GamutCuspTableUnsorted.Data[MinhIndex][2])
+			MinhIndex = i;
+	}
+
+	Table3D GamutCuspTable{};
+	for (int32 i = 0; i < GamutCuspTableUnsorted.Size; i++)
+	{
+		GamutCuspTable.Data[i + GamutCuspTable.BaseIndex] = GamutCuspTableUnsorted.Data[(MinhIndex + i) % GamutCuspTableUnsorted.Size];
+	}
+
+	// Copy last populated entry to first empty spot
+	GamutCuspTable.Data[0] = GamutCuspTable.Data[GamutCuspTable.BaseIndex + GamutCuspTable.Size - 1];
+
+	// Copy first populated entry to last empty spot
+	GamutCuspTable.Data[GamutCuspTable.BaseIndex + GamutCuspTable.Size] = GamutCuspTable.Data[GamutCuspTable.BaseIndex];
+
+	// Wrap the hues, to maintain monotonicity. These entries will fall outside [0.0, 360.0]
+	GamutCuspTable.Data[0][2] = GamutCuspTable.Data[0][2] - 360.f;
+	GamutCuspTable.Data[GamutCuspTable.Size + 1][2] = GamutCuspTable.Data[GamutCuspTable.Size + 1][2] + 360.f;
+
+	return GamutCuspTable;
+}
+
+FVector2f CuspToTable(float h, const Table3D& Gt)
+{
+	int32 Idx_lo = 0;
+	int32 Idx_hi = Gt.BaseIndex + Gt.Size; // allowed as we have an extra entry in the table
+	int32 Idx = ClampToTableBounds(HuePositionInUniformTable(h, Gt.Size) + Gt.BaseIndex, Gt.TotalSize);
+
+	while (Idx_lo + 1 < Idx_hi)
+	{
+		if (h > Gt.Data[Idx][2])
+		{
+			Idx_lo = Idx;
+		}
+		else
+		{
+			Idx_hi = Idx;
+		}
+		
+		Idx = ClampToTableBounds((Idx_lo + Idx_hi) / 2, Gt.TotalSize);
+	}
+
+	Idx_hi = FMath::Max(1, Idx_hi);
+
+	const FVector3f Lo{
+		Gt.Data[Idx_hi - 1][0],
+		Gt.Data[Idx_hi - 1][1],
+		Gt.Data[Idx_hi - 1][2]
+	};
+
+	const FVector3f Hi{
+		Gt.Data[Idx_hi][0],
+		Gt.Data[Idx_hi][1],
+		Gt.Data[Idx_hi][2]
+	};
+
+	const float t = (h - Lo[2]) / (Hi[2] - Lo[2]);
+	const float CuspJ = FMath::Lerp(Lo[0], Hi[0], t);
+	const float CuspM = FMath::Lerp(Lo[1], Hi[1], t);
+
+	return FVector2f{ CuspJ, CuspM };
+}
+
+float GetFocusGain(float J, float CuspJ, float LimitJMax)
+{
+	const float Thr = FMath::Lerp(CuspJ, LimitJMax, FocusGainBlend);
+
+	if (J > Thr)
+	{
+		// Approximate inverse required above threshold
+		float Gain = (LimitJMax - Thr) / FMath::Max(0.0001f, (LimitJMax - FMath::Min(LimitJMax, J)));
+		return FMath::Pow(log10(Gain), 1.f / FocusAdjustGain) + 1.f;
+	}
+	else
+	{
+		// Analytic inverse possible below cusp
+		return 1.f;
+	}
+}
+
+float SolveJIntersect(float J, float M, float FocusJ, float MaxJ, float SlopeGain)
+{
+	const float a = M / (FocusJ * SlopeGain);
+	float b = 0.f;
+	float c = 0.f;
+	float IntersectJ = 0.f;
+
+	if (J < FocusJ)
+	{
+		b = 1.f - M / SlopeGain;
+	}
+	else
+	{
+		b = -(1.f + M / SlopeGain + MaxJ * M / (FocusJ * SlopeGain));
+	}
+
+	if (J < FocusJ)
+	{
+		c = -J;
+	}
+	else
+	{
+		c = MaxJ * M / SlopeGain + J;
+	}
+
+	const float Root = FMath::Sqrt(b * b - 4.f * a * c);
+
+	if (J < FocusJ)
+	{
+		IntersectJ = 2.f * c / (-b - Root);
+	}
+	else
+	{
+		IntersectJ = 2.f * c / (-b + Root);
+	}
+
+	return IntersectJ;
+}
+
+FVector3f FindGamutBoundaryIntersection(const FVector3f& JMh_s, const FVector2f& JM_cusp_in, float J_focus, float J_max, float SlopeGain, float gamma_top, float gamma_bottom)
+{
+	constexpr float s = FMath::Max(0.000001f, SmoothCusps);
+	const FVector2f JM_cusp = {
+		JM_cusp_in[0],
+		JM_cusp_in[1] * (1.f + SmoothM * s)
+	};
+
+	const float J_intersect_source = SolveJIntersect(JMh_s[0], JMh_s[1], J_focus, J_max, SlopeGain);
+	const float J_intersect_cusp = SolveJIntersect(JM_cusp[0], JM_cusp[1], J_focus, J_max, SlopeGain);
+
+	float Slope = 0.f;
+	if (J_intersect_source < J_focus)
+	{
+		Slope = J_intersect_source * (J_intersect_source - J_focus) / (J_focus * SlopeGain);
+	}
+	else
+	{
+		Slope = (J_max - J_intersect_source) * (J_intersect_source - J_focus) / (J_focus * SlopeGain);
+	}
+
+	const float M_boundary_lower = J_intersect_cusp * FMath::Pow(J_intersect_source / J_intersect_cusp, 1.f / gamma_bottom) / (JM_cusp[0] / JM_cusp[1] - Slope);
+	const float M_boundary_upper = JM_cusp[1] * (J_max - J_intersect_cusp) * FMath::Pow((J_max - J_intersect_source) / (J_max - J_intersect_cusp), 1.f / gamma_top) / (Slope * JM_cusp[1] + J_max - JM_cusp[0]);
+	const float M_boundary = JM_cusp[1] * Smin(M_boundary_lower / JM_cusp[1], M_boundary_upper / JM_cusp[1], s);
+	const float J_boundary = J_intersect_source + Slope * M_boundary;
+
+	return { J_boundary, M_boundary, J_intersect_source };
+}
+
+bool EvaluateGammaFit(
+	const FVector2f& JMcusp,
+	const FVector3f TestJMh[3],
+	float TopGamma,
+	float PeakLuminance,
+	float LimitJMax,
+	float Mid_J,
+	float FocusDist,
+	float LowerHullGamma,
+	const JMhParams& LimitJMhParams)
+{
+	const float FocusJ = FMath::Lerp(JMcusp[0], Mid_J, FMath::Min(1.f, CuspMidBlend - (JMcusp[0] / LimitJMax)));
+
+	for (size_t TestIndex = 0; TestIndex < 3; TestIndex++)
+	{
+		const float SlopeGain = LimitJMax * FocusDist * GetFocusGain(TestJMh[TestIndex][0], JMcusp[0], LimitJMax);
+		const FVector3f ApproxLimit = FindGamutBoundaryIntersection(TestJMh[TestIndex], JMcusp, FocusJ, LimitJMax, SlopeGain, TopGamma, LowerHullGamma);
+		const FVector3f Approximate_JMh = { ApproxLimit[0], ApproxLimit[1], TestJMh[TestIndex][2] };
+		const FVector3f NewLimitRGB = JMh_to_RGB(Approximate_JMh, LimitJMhParams);
+		const FVector3f NewLimitRGBScaled = (ReferenceLuminance / PeakLuminance) * NewLimitRGB;
+
+		if (!OutsideHull(NewLimitRGBScaled))
+		{
+			return false;
+		}
+	}
+
+	return true;
+}
+
+Table1D MakeUpperHullGamma(
+	const Table3D& GamutCuspTable,
+	float PeakLuminance,
+	float LimitJMax,
+	float Mid_J,
+	float FocusDist,
+	float LowerHullGamma,
+	const JMhParams& LimitJMhParams)
+{
+	const int32 TestCount = 3;
+	const float TestPositions[TestCount] = { 0.01f, 0.5f, 0.99f };
+
+	Table1D GammaTable{};
+	Table1D GamutTopGamma{};
+
+	for (int32 Index = 0; Index < GammaTable.Size; Index++)
+	{
+		GammaTable.Data[Index] = -1.f;
+
+		const float Hue = (float)Index;
+		const FVector2f JMcusp = CuspToTable(Hue, GamutCuspTable);
+
+		FVector3f TestJMh[TestCount]{};
+		for (int32 TestIndex = 0; TestIndex < TestCount; TestIndex++)
+		{
+			const float TestJ = JMcusp[0] + ((LimitJMax - JMcusp[0]) * TestPositions[TestIndex]);
+			TestJMh[TestIndex] = {
+				TestJ,
+				JMcusp[1],
+				Hue
+			};
+		}
+
+		const float SearchRange = GammaSearchStep;
+		float Low = GammaMinimum;
+		float High = Low + SearchRange;
+		bool Outside = false;
+
+		while (!(Outside) && (High < 5.f))
+		{
+			const bool bGammaFound = EvaluateGammaFit(JMcusp, TestJMh, High, PeakLuminance, LimitJMax, Mid_J, FocusDist, LowerHullGamma, LimitJMhParams);
+			if (!bGammaFound)
+			{
+				Low = High;
+				High = High + SearchRange;
+			}
+			else
+			{
+				Outside = true;
+			}
+		}
+
+		float TestGamma = -1.f;
+		while ((High - Low) > GammaAccuracy)
+		{
+			TestGamma = (High + Low) / 2.f;
+			const bool bGammaFound = EvaluateGammaFit(JMcusp, TestJMh, TestGamma, PeakLuminance, LimitJMax, Mid_J, FocusDist, LowerHullGamma, LimitJMhParams);
+			if (bGammaFound)
+			{
+				High = TestGamma;
+				GammaTable.Data[Index] = High;
+			}
+			else
+			{
+				Low = TestGamma;
+			}
+		}
+
+		// Duplicate gamma value to array, leaving empty entries at first and last position
+		GamutTopGamma.Data[Index + GamutTopGamma.BaseIndex] = GammaTable.Data[Index];
+	}
+
+	// Copy last populated entry to first empty spot
+	GamutTopGamma.Data[0] = GammaTable.Data[GammaTable.Size - 1];
+
+	// Copy first populated entry to last empty spot
+	GamutTopGamma.Data[GamutTopGamma.TotalSize - 1] = GammaTable.Data[0];
+
+	return GamutTopGamma;
+}
+
+/* Base class for ACES 2.0 table texture resources. */
+class FTextureLookupBase : public FTextureWithSRV
+{
+public:
+	/* Constructor */
+	FTextureLookupBase(const FString& InDebugName)
+		: FTextureWithSRV()
+		, DebugName(InDebugName)
+	{
+	}
+	
+	/* Destructor */
+	virtual ~FTextureLookupBase() = default;
+
+	virtual uint32 GetSizeY() const override
+	{
+		return 1;
+	}
+
+	virtual void InitRHI(FRHICommandListBase& RHICmdList) override
+	{
+		const FRHITextureCreateDesc Desc = FRHITextureCreateDesc::Create2D(*DebugName, GetSizeX(), GetSizeY(), GetPixelFormat())
+			.SetFlags(ETextureCreateFlags::ShaderResource);
+
+		// Create the RHI texture
+		TextureRHI = RHICreateTexture(Desc);
+
+		// Create the sampler state RHI resource.
+		FSamplerStateInitializerRHI SamplerStateInitializer(SF_Point, AM_Clamp, AM_Clamp, AM_Clamp);
+		SamplerStateRHI = GetOrCreateSamplerState(SamplerStateInitializer);
+
+		// Create a view of the texture
+		ShaderResourceViewRHI = RHICmdList.CreateShaderResourceView(TextureRHI, FRHIViewDesc::CreateTextureSRV().SetDimensionFromTexture(TextureRHI));
+	}
+
+	/* Update the texture table resource given a display-limiting color space & a peak luminance. */
+	virtual void Update(FRHICommandListImmediate& RHICmdList, const FColorSpace& InLimitingColorSpace, float InPeakLuminance) = 0;
+
+	/* Get the texture resource pixel format. */
+	virtual EPixelFormat GetPixelFormat() const = 0;
+
+private:
+	/* Resource debug name */
+	FString DebugName;
+};
+
+/* ACES 2.0 reach M value table. */
+class FReachMTable : public FTextureLookupBase
+{
+public:
+	FReachMTable()
+		: FTextureLookupBase(TEXT("ACES_ReachMTable"))
+	{
+	}
+
+	virtual ~FReachMTable() = default;
+
+	virtual void Update(FRHICommandListImmediate& RHICmdList, const FColorSpace& InLimitingColorSpace, float InPeakLuminance) override
+	{
+		Table1D ReachTableData = MakeReachMTable(InPeakLuminance);
+		check(GetSizeX() == ReachTableData.Size);
+		const uint32 DataSize = GetSizeX() * sizeof(float);
+		const FUpdateTextureRegion2D Region(0, 0, 0, 0, GetSizeX(), GetSizeY());
+		
+		RHICmdList.UpdateTexture2D(TextureRHI, 0, Region, DataSize, (const uint8*)ReachTableData.Data);
+	}
+
+	virtual uint32 GetSizeX() const override
+	{
+		return TABLE_SIZE;
+	}
+
+	virtual EPixelFormat GetPixelFormat() const override
+	{
+		return PF_R32_FLOAT;
+	}
+};
+
+/* ACES 2.0 gamut cusp table. */
+class FGamutCuspTable : public FTextureLookupBase
+{
+public:
+	FGamutCuspTable()
+		: FTextureLookupBase(TEXT("ACES_GamutCuspTable"))
+	{
+	}
+
+	virtual ~FGamutCuspTable() = default;
+
+	void Update(FRHICommandListImmediate& RHICmdList, const FColorSpace& InLimitingColorSpace, float InPeakLuminance) override
+	{
+		Table3D GamutCuspTable = MakeGamutTable(InLimitingColorSpace, InPeakLuminance);
+		check(GetSizeX() == GamutCuspTable.TotalSize);
+		const uint32 DataSize = GetSizeX() * sizeof(FVector3f);
+		const FUpdateTextureRegion2D Region(0, 0, 0, 0, GetSizeX(), GetSizeY());
+
+		RHICmdList.UpdateTexture2D(TextureRHI, 0, Region, DataSize, (const uint8*)GamutCuspTable.Data);
+	}
+
+	uint32 GetSizeX() const override
+	{
+		return TABLE_TotalSize;
+	}
+
+	EPixelFormat GetPixelFormat() const override
+	{
+		return PF_R32G32B32F;
+	}
+};
+
+/* ACES 2.0 upper hull gamma. */
+class FUpperHullGammaTable : public FTextureLookupBase
+{
+public:
+	FUpperHullGammaTable()
+		: FTextureLookupBase(TEXT("ACES_UpperHullGammaTable"))
+	{
+	}
+
+	virtual ~FUpperHullGammaTable() = default;
+
+	virtual void Update(FRHICommandListImmediate& RHICmdList, const FColorSpace& InLimitingColorSpace, float InPeakLuminance) override
+	{
+		static const FColorSpace ColorSpaceAP0 = FColorSpace(EColorSpace::ACESAP0);
+
+		Table3D GamutCuspTable = MakeGamutTable(InLimitingColorSpace, InPeakLuminance);
+
+		const ToneScaleParams TsParams = InitToneScaleParams(InPeakLuminance);
+		const JMhParams InputJMhParams = InitJMhParams(ColorSpaceAP0);
+
+		float LimitJMax = Y_to_J(InPeakLuminance, InputJMhParams);
+		float Mid_J = Y_to_J(TsParams.c_t * 100.f, InputJMhParams);
+
+		// Calculated chroma compress variables
+		const float LogPeak = log10(TsParams.n / TsParams.n_r);
+		const float ModelGamma = 1.f / (Surround[1] * (1.48f + sqrt(Y_b / L_A)));
+		const float FocusDist = FocusDistance + FocusDistance * FocusDistanceScaling * LogPeak;
+		const float LowerHullGamma = 1.14f + 0.07f * LogPeak;
+
+		const JMhParams LimitJMhParams = InitJMhParams(InLimitingColorSpace);
+
+		Table1D UpperHullGammaTable = MakeUpperHullGamma(
+			GamutCuspTable,
+			InPeakLuminance,
+			LimitJMax,
+			Mid_J,
+			FocusDist,
+			LowerHullGamma,
+			LimitJMhParams);
+
+		check(GetSizeX() == UpperHullGammaTable.TotalSize);
+		const uint32 DataSize = GetSizeX() * sizeof(float);
+		const FUpdateTextureRegion2D Region(0, 0, 0, 0, GetSizeX(), GetSizeY());
+
+		RHICmdList.UpdateTexture2D(TextureRHI, 0, Region, DataSize, (const uint8*)UpperHullGammaTable.Data);
+	}
+
+	virtual uint32 GetSizeX() const override
+	{
+		return TABLE_TotalSize;
+	}
+	
+	virtual EPixelFormat GetPixelFormat() const override
+	{
+		return PF_R32_FLOAT;
+	}
+};
+
+void GetTransformResources(FRDGBuilder& GraphBuilder, float InPeakLuminance, FRHIShaderResourceView*& OutReachMTable, FRHIShaderResourceView*& OutGamutCuspTable, FRHIShaderResourceView*& OutUpperHullGammaTable)
+{
+	static const auto CVarAcesVersion = IConsoleManager::Get().FindTConsoleVariableDataInt(TEXT("r.HDR.Aces.Version"));
+	
+	if (CVarAcesVersion->GetValueOnRenderThread() > 1)
+	{
+		// ACES 2.0 transform resources
+
+		static FTextureLookupBase* ReachMTable = new TGlobalResource<FReachMTable>;
+		static FTextureLookupBase* GamutCuspTable = new TGlobalResource<FGamutCuspTable>;
+		static FTextureLookupBase* UpperHullGammaTable = new TGlobalResource<FUpperHullGammaTable>;
+
+		static float CachedPeakLuminance = 0.0f;
+
+		if (!FMath::IsNearlyEqual(CachedPeakLuminance, InPeakLuminance))
+		{
+			// TODO: The display limiting color space should be user-exposed as part of display characterization.
+			static UE::Color::FColorSpace LimitingColorSpace = UE::Color::FColorSpace(UE::Color::EColorSpace::ACESAP1);
+
+			FRHICommandListImmediate& RHICmdListImmediate = GraphBuilder.RHICmdList.GetAsImmediate();
+
+			ReachMTable->Update(RHICmdListImmediate, LimitingColorSpace, InPeakLuminance);
+			GamutCuspTable->Update(RHICmdListImmediate, LimitingColorSpace, InPeakLuminance);
+			UpperHullGammaTable->Update(RHICmdListImmediate, LimitingColorSpace, InPeakLuminance);
+
+			// Update static cached peak luminance.
+			CachedPeakLuminance = InPeakLuminance;
+		}
+
+		OutReachMTable = ReachMTable->ShaderResourceViewRHI;
+		OutGamutCuspTable = GamutCuspTable->ShaderResourceViewRHI;
+		OutUpperHullGammaTable = UpperHullGammaTable->ShaderResourceViewRHI;
+	}
+	else
+	{
+		// ACES 1.3
+
+		OutReachMTable = GBlackTextureWithSRV->ShaderResourceViewRHI;
+		OutGamutCuspTable = GBlackTextureWithSRV->ShaderResourceViewRHI;
+		OutUpperHullGammaTable = GBlackTextureWithSRV->ShaderResourceViewRHI;
+	}
+}
+
+}  // namespace UE::Color::ACES
+
diff --git a/Engine/Source/Runtime/Renderer/Private/PostProcess/ACESUtils.h b/Engine/Source/Runtime/Renderer/Private/PostProcess/ACESUtils.h
new file mode 100644
--- /dev/null
+++ b/Engine/Source/Runtime/Renderer/Private/PostProcess/ACESUtils.h
@@ -0,0 +1,27 @@
+// Copyright Epic Games, Inc. All Rights Reserved.
+
+#pragma once
+
+#include "RHIFwd.h"
+
+class FRDGBuilder;
+
+namespace UE::Color::ACES
+{
+	/*
+	* Get the ACES 2.0 transform table resources.
+	* 
+	* @param GraphBuilder			Frame render dependency graph.
+	* @param InPeakLuminance		Peak luminance value in nits (cd/m^2).
+	* @param OutReachMTable			Reach M value table resource output.
+	* @param OutGamutCuspTable		Gamut cusp table resource output.
+	* @param OutUpperHullGammaTable	Upper hull gamma table resource output.
+	*/
+	void GetTransformResources(
+		FRDGBuilder& GraphBuilder,
+		float InPeakLuminance,
+		FRHIShaderResourceView*& OutReachMTable,
+		FRHIShaderResourceView*& OutGamutCuspTable,
+		FRHIShaderResourceView*& OutUpperHullGammaTable
+	);
+} // namespace UE::Color::ACES
diff --git a/Engine/Source/Runtime/Renderer/Private/PostProcess/PostProcessCombineLUTs.cpp b/Engine/Source/Runtime/Renderer/Private/PostProcess/PostProcessCombineLUTs.cpp
--- a/Engine/Source/Runtime/Renderer/Private/PostProcess/PostProcessCombineLUTs.cpp
+++ b/Engine/Source/Runtime/Renderer/Private/PostProcess/PostProcessCombineLUTs.cpp
@@ -9,6 +9,7 @@
 #include "VolumeRendering.h"
 #include "HDRHelper.h"
 #include "TextureResource.h"
+#include "ACESUtils.h"
 
 namespace
 {
@@ -114,6 +115,9 @@ BEGIN_SHADER_PARAMETER_STRUCT(FACESTonemapShaderParameters, )
 	SHADER_PARAMETER(float, ACESCoefsHigh_4)
 	SHADER_PARAMETER(float, ACESSceneColorMultiplier)
 	SHADER_PARAMETER(float, ACESGamutCompression)
+	SHADER_PARAMETER_SRV(Texture2D<float>, ACESReachTable)
+	SHADER_PARAMETER_SRV(Texture2D<FVector3f>, ACESGamutTable)
+	SHADER_PARAMETER_SRV(Texture2D<float>, ACESGammaTable)
 END_SHADER_PARAMETER_STRUCT()
 
 
@@ -180,10 +184,10 @@ struct FCachedLUTSettings
 	FWorkingColorSpaceShaderParameters WorkingColorSpaceShaderParameters;
 	bool bUseCompute = false;
 
-	bool UpdateCachedValues(const FViewInfo& View, const FTexture* const* Textures, const float* Weights, uint32 BlendCount)
+	bool UpdateCachedValues(FRDGBuilder& GraphBuilder, const FViewInfo& View, const FTexture* const* Textures, const float* Weights, uint32 BlendCount)
 	{
 		bool bHasChanged = false;
-		GetCombineLUTParameters(View, Textures, Weights, BlendCount, bHasChanged);
+		GetCombineLUTParameters(GraphBuilder, View, Textures, Weights, BlendCount, bHasChanged);
 		UPDATE_CACHE_SETTINGS(UniqueID, View.State ? View.State->GetViewKey() : 0, bHasChanged);
 		UPDATE_CACHE_SETTINGS(ShaderPlatform, View.GetShaderPlatform(), bHasChanged);
 		UPDATE_CACHE_SETTINGS(bUseCompute, View.bUseComputePasses, bHasChanged);
@@ -203,6 +207,7 @@ struct FCachedLUTSettings
 	}
 	
 	void GetCombineLUTParameters(
+		FRDGBuilder& GraphBuilder,
 		const FViewInfo& View,
 		const FTexture* const* Textures,
 		const float* Weights,
@@ -245,6 +250,13 @@ struct FCachedLUTSettings
 		UPDATE_CACHE_SETTINGS(Parameters.ACESTonemapParameters.ACESCoefsHigh_4, TonemapperParams.ACESCoefsHigh_4, bHasChanged);
 		UPDATE_CACHE_SETTINGS(Parameters.ACESTonemapParameters.ACESSceneColorMultiplier, TonemapperParams.ACESSceneColorMultiplier, bHasChanged);
 		UPDATE_CACHE_SETTINGS(Parameters.ACESTonemapParameters.ACESGamutCompression, TonemapperParams.ACESGamutCompression, bHasChanged);
+		UE::Color::ACES::GetTransformResources(
+			GraphBuilder,
+			HDRGetDisplayMaximumLuminance(),
+			Parameters.ACESTonemapParameters.ACESReachTable,
+			Parameters.ACESTonemapParameters.ACESGamutTable,
+			Parameters.ACESTonemapParameters.ACESGammaTable
+		);
 
 		UPDATE_CACHE_SETTINGS(Parameters.ColorScale, FVector3f(View.ColorScale), bHasChanged);
 		UPDATE_CACHE_SETTINGS(Parameters.OverlayColor, FVector4f(View.OverlayColor), bHasChanged);
@@ -311,7 +323,8 @@ public:
 	class FBlendCount : SHADER_PERMUTATION_RANGE_INT("BLENDCOUNT", 1, 5);
 	class FSkipTemperature : SHADER_PERMUTATION_BOOL("SKIP_TEMPERATURE");
 	class FOutputDeviceSRGB : SHADER_PERMUTATION_BOOL("OUTPUT_DEVICE_SRGB");
-	using FPermutationDomain = TShaderPermutationDomain<FBlendCount, FSkipTemperature, FOutputDeviceSRGB>;
+	class FUseACES2 : SHADER_PERMUTATION_BOOL("USE_ACES_2");
+	using FPermutationDomain = TShaderPermutationDomain<FBlendCount, FSkipTemperature, FOutputDeviceSRGB, FUseACES2>;
 
 	static bool ShouldCompilePermutation(const FGlobalShaderPermutationParameters& Parameters)
 	{
@@ -507,7 +520,7 @@ FRDGTextureRef AddCombineLUTPass(FRDGBuilder& GraphBuilder, const FViewInfo& Vie
 
 	View.SetValidTonemappingLUT();
 
-	const bool bHasChanged = CachedLUTSettings.UpdateCachedValues(View, LocalTextures, LocalWeights, LocalCount);
+	const bool bHasChanged = CachedLUTSettings.UpdateCachedValues(GraphBuilder, View, LocalTextures, LocalWeights, LocalCount);
 	if (!bHasChanged && OutputTexture && CVarLUTUpdateEveryFrame.GetValueOnRenderThread() == 0)
 	{
 		return OutputTexture;
@@ -529,6 +542,9 @@ FRDGTextureRef AddCombineLUTPass(FRDGBuilder& GraphBuilder, const FViewInfo& Vie
 	FLUTBlenderShader::FPermutationDomain PermutationVector;
 	PermutationVector.Set<FLUTBlenderShader::FBlendCount>(LocalCount);
 
+	static const auto CVarAcesVersion = IConsoleManager::Get().FindTConsoleVariableDataInt(TEXT("r.HDR.Aces.Version"));
+	PermutationVector.Set<FLUTBlenderShader::FUseACES2>(CVarAcesVersion->GetValueOnRenderThread() > 1);
+
 	const float DefaultTemperature = 6500;
 	const float DefaultTint = 0;
 
@@ -627,4 +643,4 @@ FRDGTextureRef AddCombineLUTPass(FRDGBuilder& GraphBuilder, const FViewInfo& Vie
 	}
 
 	return OutputTexture;
-}
\ No newline at end of file
+}
diff --git a/Engine/Source/Runtime/Renderer/Private/PostProcess/PostProcessDeviceEncodingOnly.cpp b/Engine/Source/Runtime/Renderer/Private/PostProcess/PostProcessDeviceEncodingOnly.cpp
--- a/Engine/Source/Runtime/Renderer/Private/PostProcess/PostProcessDeviceEncodingOnly.cpp
+++ b/Engine/Source/Runtime/Renderer/Private/PostProcess/PostProcessDeviceEncodingOnly.cpp
@@ -17,6 +17,7 @@
 #include "ClearQuad.h"
 #include "PipelineStateCache.h"
 #include "HDRHelper.h"
+#include "ACESUtils.h"
 
 namespace
 {
@@ -89,6 +90,9 @@ BEGIN_SHADER_PARAMETER_STRUCT(FDeviceEncodingACESTonemapShaderParameters, )
 	SHADER_PARAMETER(float, ACESCoefsHigh_4)
 	SHADER_PARAMETER(float, ACESSceneColorMultiplier)
 	SHADER_PARAMETER(float, ACESGamutCompression)
+	SHADER_PARAMETER_SRV(Texture2D<float>, ACESReachTable)
+	SHADER_PARAMETER_SRV(Texture2D<FVector3f>, ACESGamutTable)
+	SHADER_PARAMETER_SRV(Texture2D<float>, ACESGammaTable)
 END_SHADER_PARAMETER_STRUCT()
 
 BEGIN_SHADER_PARAMETER_STRUCT(FDeviceEncodingOnlyParameters, )
@@ -219,6 +223,14 @@ FScreenPassTexture AddDeviceEncodingOnlyPass(FRDGBuilder& GraphBuilder, const FV
 	CommonParameters.Color = GetScreenPassTextureViewportParameters(SceneColorViewport);
 	CommonParameters.Output = GetScreenPassTextureViewportParameters(OutputViewport);
 	GetACESTonemapParameters(CommonParameters.ACESTonemapParameters);
+	UE::Color::ACES::GetTransformResources(
+		GraphBuilder,
+		HDRGetDisplayMaximumLuminance(),
+		CommonParameters.ACESTonemapParameters.ACESReachTable,
+		CommonParameters.ACESTonemapParameters.ACESGamutTable,
+		CommonParameters.ACESTonemapParameters.ACESGammaTable
+	);
+
 	CommonParameters.ColorTexture = Inputs.SceneColor.Texture;
 	CommonParameters.ColorSampler = BilinearClampSampler;
 	CommonParameters.EditorNITLevel = EditorNITLevel;
