Reenable shader print on Mac SM6.

diff --git a/Engine/Shaders/Private/ShaderPrintCommon.ush b/Engine/Shaders/Private/ShaderPrintCommon.ush
--- a/Engine/Shaders/Private/ShaderPrintCommon.ush
+++ b/Engine/Shaders/Private/ShaderPrintCommon.ush
@@ -171,8 +171,8 @@ void ClearCounters(RWBuffer<uint> InRWBuffer)
 
 // Currently a bug Metal shader translation cause compilation to fails due to atomic operation.
 // As a workaround, disabling interlock add on Metal
-#if COMPILER_METAL
-#define SHADER_PRINT_INTERLOCKEDADD(InBuffer, InIncrement, OutOldValue) 
+#if COMPILER_METAL && FEATURE_LEVEL < FEATURE_LEVEL_SM6
+#define SHADER_PRINT_INTERLOCKEDADD(InBuffer, InIncrement, OutOldValue)
 #else
 #define SHADER_PRINT_INTERLOCKEDADD(InBuffer, InIncrement, OutOldValue) InterlockedAdd(InBuffer, InIncrement, OutOldValue)
 #endif
