[DOF] Fix typo in DOF scatter vignetting due to merge error. Pixels contributions should be accumulated if either of the pixels in the pair intersects, rather than just testing the first one.

diff --git a/Engine/Shaders/Private/DiaphragmDOF/DOFHybridScatterPixelShader.usf b/Engine/Shaders/Private/DiaphragmDOF/DOFHybridScatterPixelShader.usf
--- a/Engine/Shaders/Private/DiaphragmDOF/DOFHybridScatterPixelShader.usf
+++ b/Engine/Shaders/Private/DiaphragmDOF/DOFHybridScatterPixelShader.usf
@@ -221,7 +221,7 @@ void ScatterMainPS(
 		#endif
 
 		BRANCH
-		if (any(Intersection[0] > 0))
+		if (any(Intersection > 0))
 		{
 			half2 Visibility = 1;
 			#if DIM_SCATTER_OCCLUSION
