Fix hair cards's normal recomputation when not bound to a skel. mesh.

diff --git a/Engine/Shaders/Private/HairStrands/HairCardsDeformation.usf b/Engine/Shaders/Private/HairStrands/HairCardsDeformation.usf
--- a/Engine/Shaders/Private/HairStrands/HairCardsDeformation.usf
+++ b/Engine/Shaders/Private/HairStrands/HairCardsDeformation.usf
@@ -56,6 +56,7 @@ float3 ComputeDynamicGeometryOffset(
 
 	return (LocalDeformedGuidePoint - LocalRestGuidePoint) * GuideVertexWeight;
 }
+#endif
 
 void TransformTangent(
 	uint InVertexId, 
@@ -68,13 +69,31 @@ void TransformTangent(
 	float3 Deformed = VectorToWorld(Local, InDeformedBasis);
 	CardsDeformedTangentBuffer[InVertexId * 2 + InOffset] = float4(Deformed, Rest.w);
 }
-#endif
 
 void CopyTangent(uint InVertexId, uint InOffset)
 {
 	CardsDeformedTangentBuffer[InVertexId * 2 + InOffset] = CardsRestTangentBuffer[InVertexId * 2 + InOffset];
 }
 
+const FHairMeshBasis ConstructBasis(float3 GuideP0, float3 GuideP1, float3 P)
+{
+	FHairMeshBasis Out;
+
+	const float3 V0 = P;
+	const float3 V1 = GuideP0;
+	const float3 V2 = GuideP1;
+
+	const float3 E0 = (V1 - V0).xyz;
+	const float3 E1 = (V2 - V0).xyz;
+
+	// Build orthonormal frame
+	Out.BaseO = P;  //(V0 * B.x + V1 * B.y + V2 * B.z).xyz;
+	Out.BaseN = normalize(cross(E1, E0));
+	Out.BaseT = normalize(E0);
+	Out.BaseB = normalize(cross(Out.BaseT, Out.BaseN));
+	return Out;
+}
+
 [numthreads(GROUP_SIZE, 1, 1)]
 void MainCS(uint DispatchThreadId : SV_DispatchThreadID)
 {
@@ -130,11 +149,15 @@ void MainCS(uint DispatchThreadId : SV_DispatchThreadID)
 
 	const float3 DeformationOffset = DeformedGuidePosition - RestGuidePosition;
 	const float4 RestPosition = CardsRestPositionBuffer[VertexId];
-	const float3 OutPosition = RestPosition.xyz + DeformationOffset;
-	CardsDeformedPositionBuffer[VertexId] = float4(OutPosition, RestPosition.w);
+	const float4 DeformedPosition = float4(RestPosition.xyz + DeformationOffset, RestPosition.w);
+	CardsDeformedPositionBuffer[VertexId] = DeformedPosition;
 
-	CopyTangent(VertexId, 0);
-	CopyTangent(VertexId, 1);
+	// Recompute the vertex's normal
+	// Construct a local basis from the 2 guides points and the position to be deformed. 
+	const FHairMeshBasis RestBasis = ConstructBasis(RestGuideP0.Position, RestGuideP1.Position, RestPosition.xyz);
+	const FHairMeshBasis DeformedBasis = ConstructBasis(DeformedGuideP0.Position, DeformedGuideP1.Position, DeformedPosition.xyz);
+	TransformTangent(VertexId, 0, RestBasis, DeformedBasis);
+	TransformTangent(VertexId, 1, RestBasis, DeformedBasis);
 
 #elif PERMUTATION_DYNAMIC_GEOMETRY == 1
 	// Compute the deformation offset in triangle local space
@@ -163,9 +186,9 @@ void MainCS(uint DispatchThreadId : SV_DispatchThreadID)
 
 	#if DEBUG_ENABLE
 	{
-		float3 WP0 = mul(float4(DeformedBasis.P0, 1), LocalToWorld).xyz;
-		float3 WP1 = mul(float4(DeformedBasis.P1, 1), LocalToWorld).xyz;
-		float3 WP2 = mul(float4(DeformedBasis.P2, 1), LocalToWorld).xyz;
+		float3 WP0 = mul(float4(DeformedBasis.BaseO, 1), LocalToWorld).xyz;
+		float3 WP1 = mul(float4(DeformedBasis.BaseO + DeformedBasis.BaseT, 1), LocalToWorld).xyz;
+		float3 WP2 = mul(float4(DeformedBasis.BaseO + DeformedBasis.BaseB, 1), LocalToWorld).xyz;
 		AddLineTriangle(WP0, WP1, WP2, ColorOrange);
 	}
 	#endif
