Fix AlphaComposite with OIT.

diff --git a/Engine/Shaders/Private/BasePassPixelShader.usf b/Engine/Shaders/Private/BasePassPixelShader.usf
--- a/Engine/Shaders/Private/BasePassPixelShader.usf
+++ b/Engine/Shaders/Private/BasePassPixelShader.usf
@@ -2227,12 +2227,12 @@ void FPixelShaderInOut_MainPS(
 		#else
 		{
 			float3 WriteColor = Out.MRT[0].rgb;
-		#if MATERIALBLENDING_TRANSLUCENT
-			WriteColor *= Opacity; // WriteColor is premultiplied alpha color
-		#elif MATERIALBLENDING_ADDITIVE
+		#if MATERIALBLENDING_ADDITIVE
 			// NOP
 		#elif MATERIALBLENDING_ALPHACOMPOSITE || MATERIALBLENDING_ALPHAHOLDOUT
 			// Not clear what should be done in this case.
+		#elif MATERIALBLENDING_TRANSLUCENT
+			WriteColor *= Opacity; // WriteColor is premultiplied alpha color
 		#else
 			#error OIT is not support for this blend mode
 		#endif
