[SLW Lumen reflections] Fix shader a compilation error on metal platform

diff --git a/Engine/Shaders/Private/SingleLayerWaterComposite.usf b/Engine/Shaders/Private/SingleLayerWaterComposite.usf
--- a/Engine/Shaders/Private/SingleLayerWaterComposite.usf
+++ b/Engine/Shaders/Private/SingleLayerWaterComposite.usf
@@ -467,6 +467,8 @@ groupshared uint SharedTileData[THREADGROUP_SIZE * THREADGROUP_SIZE];
 groupshared uint SharedGlobalTileOffset;
 groupshared uint SharedGlobalClearTileOffset;
 
+#ifdef WaterTileClassificationBuildListsCS
+
 /**
  * Every group checks 64 tiles and builds a spatially coherent compacted list of water tiles
  */
@@ -604,6 +606,8 @@ void WaterTileClassificationBuildListsCS(
 	}
 }
 
+#endif //WaterTileClassificationBuildListsCS
+
 #elif defined(TILE_VERTEX_SHADER)
 
 
