GPUSkinCache: Fix minor compiler warning

half4 was being assigned to half3

#rnx

diff --git a/Engine/Shaders/Private/GpuSkinCacheComputeShader.usf b/Engine/Shaders/Private/GpuSkinCacheComputeShader.usf
--- a/Engine/Shaders/Private/GpuSkinCacheComputeShader.usf
+++ b/Engine/Shaders/Private/GpuSkinCacheComputeShader.usf
@@ -234,8 +234,8 @@ void SkinCacheUpdateBatchCS(uint3 GroupID : SV_GroupID,
 	Intermediates.UnpackedPosition.y = PositionInputBuffer[InputOffset * 3 + 1];
 	Intermediates.UnpackedPosition.z = PositionInputBuffer[InputOffset * 3 + 2];
 
-	half3 LocalTangentX = Unpacked.TangentX;
-	half3 LocalTangentZ = Unpacked.TangentZ;
+	half3 LocalTangentX = Unpacked.TangentX.xyz;
+	half3 LocalTangentZ = Unpacked.TangentZ.xyz;
 
 #if GPUSKIN_MORPH_BLEND
 	{
