Early out from adding hidden primitives to ray tracing.

While GetInstanceSceneDataInternal() is returning 0'd instance scene data if hidden, this is more explicit and short circuits some instructions.

#rnx

diff --git a/Engine/Shaders/Private/RayTracing/RayTracingInstanceBufferUtil.usf b/Engine/Shaders/Private/RayTracing/RayTracingInstanceBufferUtil.usf
--- a/Engine/Shaders/Private/RayTracing/RayTracingInstanceBufferUtil.usf
+++ b/Engine/Shaders/Private/RayTracing/RayTracingInstanceBufferUtil.usf
@@ -62,6 +62,11 @@ uint CullUsingGroups;
 
 bool CullInstance(FInstanceSceneData InstanceSceneData, FPrimitiveSceneData PrimitiveData, float4x4 LocalToTranslatedWorld4x4, bool bFarFieldInstance)
 {
+	if ((InstanceSceneData.Flags & INSTANCE_SCENE_DATA_FLAG_HIDDEN) != 0u)
+	{
+		return true;
+	}
+
 	bool bIsVisible = true;
 
 	const float3 LocalBoxCenter = InstanceSceneData.LocalBoundsCenter.xyz;
