VolumetricCloud - fixed new sample / reprojection debug code.

#rnx
#tests PC editor dx12

diff --git a/Engine/Shaders/Private/VolumetricRenderTarget.usf b/Engine/Shaders/Private/VolumetricRenderTarget.usf
--- a/Engine/Shaders/Private/VolumetricRenderTarget.usf
+++ b/Engine/Shaders/Private/VolumetricRenderTarget.usf
@@ -162,25 +162,6 @@ void ReconstructVolumetricRenderTargetPS(
 	const bool bUseNewSample = (XSub == CurrentTracingPixelOffset.x) && (YSub == CurrentTracingPixelOffset.y);
 	float4 RGBA = 0.0f;
 	FDepthData DepthData = (FDepthData)0;
-
-#if 0
-	// Debug: always use new sample else history result. Should be 
-	if (bUseNewSample)
-	{
-		RGBA  = SafeLoadTracingVolumetricTexture(int2(SVPos.xy) / DownSampleFactor);
-		//RGBA  = SafeLoadSecondaryTracingVolumetricTexture(int2(SVPos.xy) / DownSampleFactor);
-		DepthData = SafeLoadTracingVolumetricDepthTexture(int2(SVPos.xy) / DownSampleFactor);
-	}
-	else
-	{
-		RGBA = SafeSampleTracingVolumetricTexture(ScreenUV);
-		//RGBA = SafeSampleSecondaryTracingVolumetricTexture(ScreenUV);
-		DepthData = SafeSamplePreviousFrameVolumetricDepthTexture(ScreenUV);
-	}
-	OutputRt0 = RGBA;
-	OutputRt1 = float2(DepthData.CloudFrontDepthFromViewKm, DepthData.SceneDepthFromViewKm);
-	return;
-#endif
   
 	{
 		float2 ScreenPosition = ViewportUVToScreenPos(ScreenUV); // NDC in [-1,1] not using View.ScreenPositionScaleBias here
@@ -203,6 +184,27 @@ void ReconstructVolumetricRenderTargetPS(
 																							// especially visible when flying through the cloud layer. It is not perfect but will help in lots of cases. 
 																							// The problem when using this method: clouds will look noisier when closer to that distance. 
 
+		#if 0
+		{
+			// Debug: always use new sample else history result. Should be 
+			if (bUseNewSample)
+			{
+				RGBA = SafeLoadTracingVolumetricTexture(int2(SVPos.xy) / DownSampleFactor);
+				//RGBA  = SafeLoadSecondaryTracingVolumetricTexture(int2(SVPos.xy) / DownSampleFactor);
+				DepthData = SafeLoadTracingVolumetricDepthTexture(int2(SVPos.xy) / DownSampleFactor);
+			}
+			else
+			{
+				RGBA = SafeSamplePreviousFrameVolumetricTexture(PrevScreenUVs);
+				//RGBA = SafeSampleSecondaryTracingVolumetricTexture(PrevScreenUVs);
+				DepthData = SafeSamplePreviousFrameVolumetricDepthTexture(PrevScreenUVs);
+			}
+			OutputRt0 = RGBA;
+			OutputRt1 = float2(DepthData.CloudFrontDepthFromViewKm, DepthData.SceneDepthFromViewKm);
+			return;
+		}
+		#endif
+
 		if (VolumetricRenderTargetMode == 2)
 		{
 			const bool bUseNewSampleMode2 = ((IntPixelPos.x - IntPixelPosDownsample.x * DownSampleFactor) == CurrentTracingPixelOffset.x) && ((IntPixelPos.y - IntPixelPosDownsample.y * DownSampleFactor) == CurrentTracingPixelOffset.y);
