Wrong type used in Velocity pass shaders on Opengl Windows

diff --git a/Engine/Shaders/Private/Common.ush b/Engine/Shaders/Private/Common.ush
--- a/Engine/Shaders/Private/Common.ush
+++ b/Engine/Shaders/Private/Common.ush
@@ -241,7 +241,7 @@ float VectorSum(float4 V) { return V.x + V.y + V.z + V.w; }
 	#define VELOCITY_ENCODE_GAMMA 0
 #endif
 
-#if COMPILER_GLSL_ES3_1
+#if OPENGL_PROFILE
 	#define ENCODED_VELOCITY_TYPE uint4
 #else
 	#define ENCODED_VELOCITY_TYPE float4
