Improve parameter blending with vertical layering by changing transmittance computation to use refracted vector.

diff --git a/Engine/Shaders/Private/Substrate/Substrate.ush b/Engine/Shaders/Private/Substrate/Substrate.ush
--- a/Engine/Shaders/Private/Substrate/Substrate.ush
+++ b/Engine/Shaders/Private/Substrate/Substrate.ush
@@ -2920,6 +2920,35 @@ float VerticalLayeringParameterBlendingBSDFCoverageToNormalMix(FSubstrateData To
 #endif
 }
 
+float SubstrateComputeTransmittanceNoV(float3 InF0, float3 InV, float3 InN, float InNoV)
+{
+	#if 1
+	// For transmittance, we must account the fact that the view ray is going to be refracted at the interface of the slab.
+	// Otherwise, at tangent view direction, the path length traveled within the medium will tend towards infinity. That result in too dark edges and is unrealistic.
+	// To simplifiy things, we refract the view ray as if it was coming from a air interface into a medium of IOR deducted from F0.
+	// If we do this computation once for each slab, it is equivalent to have each layered slab having the same IOR with a single view vector refraction at the top interface.
+	// This is an aproximation we accept for the sake of performance today. 
+	// OutRefractedTransmittanceNoV is used to build ThroughputV that will be used when processing the substrate tree to get correct throughput to layers below.
+	// Ideas for later (requires an update of the substrate tree evaluation)
+	//  1- we could compute IOR for each layers (gather IOR for each slab horizontally integrated) and refract V according to the from/to IOR a the interface.
+	//  2- we could have V from a layer propagated to the next layer in order to have better specular match. 
+	//  3- we could also store extra data in the closure in order to be able to account for that effect for the Lighting vector too.
+	float3 RefractedV = refract(InV, -InN, 1.0 / DielectricF0ToIor(InF0.y));
+	const float OutRefractedTransmittanceNoV = dot(RefractedV, InN);
+	#else
+	const float OutRefractedTransmittanceNoV = InNoV;
+	#endif
+
+	return OutRefractedTransmittanceNoV;
+}
+
+float SubstrateComputeTransmittanceNoV(float3 InF0, float InNoV)
+{
+	const float3 LocalN = float3(0, 0, 1);
+	const float3 LocalV = float3(sqrtFast(1 - saturate(InNoV)), 0, InNoV);
+	return SubstrateComputeTransmittanceNoV(InF0, LocalV, LocalN, InNoV);
+}
+
 FSubstrateData SubstrateVerticalLayeringParameterBlending(in FSubstrateData Top, in FSubstrateData Base, uint NewNormalIndex, float TopNoV, float BaseNoV)
 {
 	FSubstrateData Result = (FSubstrateData)0;
@@ -2965,12 +2994,15 @@ FSubstrateData SubstrateVerticalLayeringParameterBlending(in FSubstrateData Top,
 	// EnableSlabBSDFSimpleVolumetric will be called at then end of the SubstrateTree as needed later.
 	const float3 OriginalBaseMFP = SLAB_SSSMFP(BaseBSDF);
 
+	// Compute refracted V for evaluating transmittance attenuation used by the bottom layer.
+	const float RefractedTransmittanceNoV = SubstrateComputeTransmittanceNoV(SLAB_F0(TopBSDF), TopNoV);
+
 	// Compute the top slab transmittance that will be used to blend the top and base material parameters.
 	// We do not check BSDF_GETISSIMPLEVOLUME(BSDFContext.BSDF) because that is set when we know the BSDF is not at the bottom later. And here we already know it is not the case.
 	TopBSDF.EnableSlabBSDFSimpleVolumetric();
 	BaseBSDF.EnableSlabBSDFSimpleVolumetric();
 	FParticipatingMedia PM				= SubstrateSlabCreateParticipatingMedia(SLAB_DIFFUSEALBEDO(TopBSDF), SLAB_SSSMFP(TopBSDF));
-	const float3 SlabTransmittanceV		= IsotropicMediumSlabTransmittance(PM, SUBSTRATE_SIMPLEVOLUME_THICKNESS_M, TopNoV);
+	const float3 SlabTransmittanceV		= IsotropicMediumSlabTransmittance(PM, SUBSTRATE_SIMPLEVOLUME_THICKNESS_M, RefractedTransmittanceNoV);
 	const float3 SlabTransmittanceN		= IsotropicMediumSlabTransmittance(PM, SUBSTRATE_SIMPLEVOLUME_THICKNESS_M, 1.0f /*Normal incidence i.e., NoL == 1.f*/);
 	const float3 SlabTransmittance		= SlabTransmittanceV * SlabTransmittanceN;
 
diff --git a/Engine/Shaders/Private/Substrate/SubstrateEvaluation.ush b/Engine/Shaders/Private/Substrate/SubstrateEvaluation.ush
--- a/Engine/Shaders/Private/Substrate/SubstrateEvaluation.ush
+++ b/Engine/Shaders/Private/Substrate/SubstrateEvaluation.ush
@@ -997,22 +997,7 @@ FSubstrateEvaluateResult SubstrateEvaluateBSDFCommon(FSubstrateBSDFContext BSDFC
 
 				const float3 SlabDirectionalAlbedo = PhaseFunction * IsotropicMediumSlabPunctualDirectionalAlbedo(PM, BSDFContext.Context.NoV, BSDFContext.Context.NoL);
 
-			#if 1
-				// For transmittance, we must account the fact that the view ray is going to be refracted at the interface of the slab.
-				// Otherwise, at tangent view direction, the path length traveled within the medium will tend towards infinity. That result in too dark edges and is unrealistic.
-				// To simplifiy things, we refract the view ray as if it was coming from a air interface into a medium of IOR deducted from F0.
-				// If we do this computation once for each slab, it is equivalent to have each layered slab having the same IOR with a single view vector refraction at the top interface.
-				// This is an aproximation we accept for the sake of performance today. 
-				// RefractedTransmittanceNoV is used to build ThroughputV that will be used when processing the substrate tree to get correct throughput to layers below.
-				// Ideas for later (requires an update of the substrate tree evaluation)
-				//  1- we could compute IOR for each layers (gather IOR for each slab horizontally integrated) and refract V according to the from/to IOR a the interface.
-				//  2- we could have V from a layer propagated to the next layer in order to have better specular match. 
-				//  3- we could also store extra data in the closure in order to be able to account for that effect for the Lighting vector too.
-				float3 RefractedV = refract(BSDFContext.V, -BSDFContext.N, 1.0 / DielectricF0ToIor(F0.y));
-				const float RefractedTransmittanceNoV = dot(RefractedV, BSDFContext.N);
-			#else
-				const float RefractedTransmittanceNoV = BSDFContext.Context.NoV;
-			#endif
+				const float RefractedTransmittanceNoV = SubstrateComputeTransmittanceNoV(F0, BSDFContext.V, BSDFContext.N, BSDFContext.Context.NoV);
 
 				// Shading point <-> View throughput
 				const float3 SlabTransmittanceV = IsotropicMediumSlabTransmittance(PM, SUBSTRATE_SIMPLEVOLUME_THICKNESS_M, RefractedTransmittanceNoV);
