Simplify logic for calculating pending virtual texture mips in shader and move to shared location instead of DebugViewModePixelShader.
#rnx

diff --git a/Engine/Shaders/Private/BasePassPixelShader.usf b/Engine/Shaders/Private/BasePassPixelShader.usf
--- a/Engine/Shaders/Private/BasePassPixelShader.usf
+++ b/Engine/Shaders/Private/BasePassPixelShader.usf
@@ -889,7 +889,7 @@ void FPixelShaderInOut_MainPS(
 #if LIGHTMAP_VT_ENABLED
 	// This must occur after CalcMaterialParameters(), which is required to initialize the VT feedback mechanism
 	// Lightmap request is always the first VT sample in the shader
-	StoreVirtualTextureFeedback(MaterialParameters.VirtualTextureFeedback, 0, LightmapVTPageTableResult.PackedRequest);
+	StoreVirtualTextureFeedback(MaterialParameters.VirtualTextureFeedback, 0, LightmapVTPageTableResult.PackedRequest, LightmapVTPageTableResult.PendingMips);
 #endif
 
 #if USE_EDITOR_COMPOSITING && (FEATURE_LEVEL >= FEATURE_LEVEL_SM4 || MOBILE_EMULATION)
diff --git a/Engine/Shaders/Private/DebugViewModePixelShader.usf b/Engine/Shaders/Private/DebugViewModePixelShader.usf
--- a/Engine/Shaders/Private/DebugViewModePixelShader.usf
+++ b/Engine/Shaders/Private/DebugViewModePixelShader.usf
@@ -135,32 +135,6 @@ MaterialFloat StoreTexSample(in out FTexCoordScalesParams Params, float4 C, int
 	return 1.f;
 }
 
-MaterialFloat StoreVTSampleInfo(in out FTexCoordScalesParams Params, VTPageTableResult PageTableResult, uint LayerIndex, int TextureReferenceIndex)
-{
-	if (VisualizeMode == DVSM_VirtualTexturePendingMips)
-	{
-		uint RequestLevel = GetRequestedLevel(PageTableResult);
-
-		// Find the highest resolution from all the page table entries sampled for this VT stack.
-		// We do that because a single VT stack can cover textures of different resolutions.
-		// Taking info from the lower resolution textures will always show undersampling compared to the requested level that we are not interested in.
-		uint SampledLevel = 0xffff;
-		for (uint Index = 0; Index < 8; ++Index)
-		{
-			SampledLevel = min(SampledLevel, GetSampledLevel(PageTableResult, Index, 0xffff));
-		}
-
-		if (SampledLevel < 0xffff)
-		{
-			// Accumulate the worst case that we find for the current pixel.
-			// Multiple VT stacks can exist, and we want to visualize the worst case.
-			int Diff = max((int)SampledLevel - (int)RequestLevel, 0);
-			Params.VTMissingMipCount = max(Diff, Params.VTMissingMipCount);
-		}
-	}
-	return 1.f;
-}
-
 #include "/Engine/Generated/Material.ush"
 #include "DebugViewModeCommon.ush"
 
@@ -436,12 +410,15 @@ float4 VisualizeRequiredTextureResolution(in FDebugPSIn DebugInputs, in bool bIs
 	}
 	else if (VisualizeMode == DVSM_VirtualTexturePendingMips)
 	{
-		if (MaterialParameters.TexCoordScalesParams.VTMissingMipCount >= 0)
+#if MATERIAL_VIRTUALTEXTURE_FEEDBACK || LIGHTMAP_VT_ENABLED
+		uint PendingMips = MaterialParameters.VirtualTextureFeedback.PendingMips;
+		if (PendingMips > 0)
 		{
-			float Accuracy = -MaterialParameters.TexCoordScalesParams.VTMissingMipCount;
-			Result = AccuracyColorLookup(Accuracy);
+			int ColorIndex = PendingMips > 1 ? 0 : 1;
+			Result = DebugViewModeStruct.AccuracyColors[ColorIndex].rgb;
 			Result *= 0.5f;
 		}
+#endif
 	}
 
 #if MATERIAL_VIRTUALTEXTURE_FEEDBACK || LIGHTMAP_VT_ENABLED
diff --git a/Engine/Shaders/Private/DebugViewModeVertexShader.usf b/Engine/Shaders/Private/DebugViewModeVertexShader.usf
--- a/Engine/Shaders/Private/DebugViewModeVertexShader.usf
+++ b/Engine/Shaders/Private/DebugViewModeVertexShader.usf
@@ -26,11 +26,6 @@ MaterialFloat StoreTexSample(in out FTexCoordScalesParams Params, float4 C, int
 	return 1.f;
 }
 
-MaterialFloat StoreVTSampleInfo(in out FTexCoordScalesParams Params, VTPageTableResult PageTableResult, uint LayerIndex, int TextureReferenceIndex)
-{
-	return 1.f;
-}
-
 #include "/Engine/Generated/Material.ush"
 #include "/Engine/Generated/VertexFactory.ush"
 
diff --git a/Engine/Shaders/Private/MaterialTemplate.ush b/Engine/Shaders/Private/MaterialTemplate.ush
--- a/Engine/Shaders/Private/MaterialTemplate.ush
+++ b/Engine/Shaders/Private/MaterialTemplate.ush
@@ -626,11 +626,9 @@ struct FMaterialPixelParameters
 #if TEX_COORD_SCALE_ANALYSIS
 #define MaterialStoreTexCoordScale(Parameters, UV, TextureReferenceIndex) StoreTexCoordScale(Parameters.TexCoordScalesParams, UV, TextureReferenceIndex)
 #define MaterialStoreTexSample(Parameters, UV, TextureReferenceIndex) StoreTexSample(Parameters.TexCoordScalesParams, UV, TextureReferenceIndex)
-#define MaterialStoreVTSampleInfo(Parameters, PageTableResult, LayerIndex, TextureReferenceIndex) StoreVTSampleInfo(Parameters, PageTableResult, LayerIndex, TextureReferenceIndex)
 #else
 #define MaterialStoreTexCoordScale(Parameters, UV, TextureReferenceIndex) 1.0f
 #define MaterialStoreTexSample(Parameters, UV, TextureReferenceIndex) 1.0f
-#define MaterialStoreVTSampleInfo(Parameters, PageTableResult, LayerIndex, TextureReferenceIndex) 1.0f
 #endif
 
 // @todo compat hack
diff --git a/Engine/Shaders/Private/MobileBasePassPixelShader.usf b/Engine/Shaders/Private/MobileBasePassPixelShader.usf
--- a/Engine/Shaders/Private/MobileBasePassPixelShader.usf
+++ b/Engine/Shaders/Private/MobileBasePassPixelShader.usf
@@ -609,7 +609,7 @@ void Main(
 #if LIGHTMAP_VT_ENABLED
 	// This must occur after CalcMaterialParameters(), which is required to initialize the VT feedback mechanism
 	// Lightmap request is always the first VT sample in the shader
-	StoreVirtualTextureFeedback(MaterialParameters.VirtualTextureFeedback, 0, LightmapVTPageTableResult.PackedRequest);
+	StoreVirtualTextureFeedback(MaterialParameters.VirtualTextureFeedback, 0, LightmapVTPageTableResult.PackedRequest, LightmapVTPageTableResult.PendingMips);
 #endif
 
 #if SUBSTRATE_ENABLED && !SUBSTRATE_INLINE_SINGLELAYERWATER
diff --git a/Engine/Shaders/Private/VirtualTextureCommon.ush b/Engine/Shaders/Private/VirtualTextureCommon.ush
--- a/Engine/Shaders/Private/VirtualTextureCommon.ush
+++ b/Engine/Shaders/Private/VirtualTextureCommon.ush
@@ -21,13 +21,16 @@
 #define LIGHTMAP_VT_ENABLED 0
 #endif
 
+#ifndef WRITE_PENDING_VIRTUALTEXTURE_MIPS
+#define WRITE_PENDING_VIRTUALTEXTURE_MIPS 0
+#endif
+
 /** Struct used to store feedback information in. */
 struct FVirtualTextureFeedbackParams
 {
-	uint Request;	// The tile to request for this pixel
-#if (NUM_VIRTUALTEXTURE_FEEDBACK_REQUESTS + LIGHTMAP_VT_ENABLED) > 1	
-	uint RequestId;	// The id of the VT texture sample this pixel will generate feedback for
-#endif	
+	uint RequestId;		// The id of the VT texture sample this pixel will store a request for. Only used when we have more than one sample.
+	uint Request;		// The tile to request for this pixel.
+	uint PendingMips;	// The accumulated pending mip count for this pixel from all samples.
 };
 
 /** Initializes the FVirtualTextureFeedbackParams for the pixel shader. */
@@ -36,20 +39,26 @@ void InitializeVirtualTextureFeedback(in out FVirtualTextureFeedbackParams Param
 #if (NUM_VIRTUALTEXTURE_FEEDBACK_REQUESTS + LIGHTMAP_VT_ENABLED) > 1
 	const uint NumVTSamplesInShader = NUM_VIRTUALTEXTURE_FEEDBACK_REQUESTS + LIGHTMAP_VT_ENABLED;
 	Params.RequestId = View.VirtualTextureFeedbackSampleOffset % NumVTSamplesInShader;
-	Params.Request = 0xFFFFFFFF;
-#else
-	Params.Request = 0xFFFFFFFF;	
 #endif	
+
+	Params.Request = 0xFFFFFFFF;
+
+#if WRITE_PENDING_VIRTUALTEXTURE_MIPS
+	Params.PendingMips = 0;
+#endif
 }
 
 /** Store feedback info for a VT request in the passed in FVirtualTextureFeedbackParams. */
-void StoreVirtualTextureFeedback(in out FVirtualTextureFeedbackParams Params, uint RequestId, uint Request)
+void StoreVirtualTextureFeedback(in out FVirtualTextureFeedbackParams Params, uint RequestId, uint Request, uint PendingMips)
 {
 #if (NUM_VIRTUALTEXTURE_FEEDBACK_REQUESTS + LIGHTMAP_VT_ENABLED) > 1
 	Params.Request = (RequestId == Params.RequestId) ? Request : Params.Request;
 #else
 	Params.Request = Request;	
 #endif
+#if WRITE_PENDING_VIRTUALTEXTURE_MIPS
+	Params.PendingMips = max(Params.PendingMips, PendingMips);
+#endif
 }
 
 /** This should be called at the end of the pixel shader to write out the gathered VT feedback info to the OutputBuffer. */
@@ -197,31 +206,23 @@ struct VTPageTableResult
 	uint4 PageTableValue[4];
 	uint PackedRequest;
 	float MipLevelFrac;
+	uint PendingMips;
 };
 
-uint GetRequestedLevel(VTPageTableResult PageTableResult)
+uint GetVirtualTexturePendingMips(uint RequestedLevel, uint4 PackedPageTableValue, uint AccumulatedResult)
 {
-	uint RequestLevelPlusOne = (PageTableResult.PackedRequest >> 24) & 0xf;
-	uint RequestLevel = max(RequestLevelPlusOne, 1u) - 1u;
-	return RequestLevel;
-}
-
-uint GetSampledLevel(VTPageTableResult PageTableResult, uint LayerIndex, uint FallbackValue)
-{
-	uint PackedPageTableValue = PageTableResult.PageTableValue[LayerIndex / 4u][LayerIndex & 3u];
-	
+#if WRITE_PENDING_VIRTUALTEXTURE_MIPS
 	// 0 is an empty value in the PageTableResult.
-	// We must also consider 1 as an empty value since it will appear in the .w channel page table textures with less than 4 channels.
-	// todo: Inject the number of layers inside VTPageTableResult and avoid the need to special case values here.
-	// HLSLMaterialTranslator could pass the value to TextureLoadVirtualPageTable through a macro similar to VIRTUALTEXTURE_PAGETABLE_%d.
-	if (PackedPageTableValue != 0 && PackedPageTableValue != 1)
+	// We must also ignore a value of 1 for this test since it will be in the .w channel for page table textures with less than 4 channels.
+	// So only consider packed page table values > 1.
+	uint4 SampledLevels = select(PackedPageTableValue.xyzw > 1,  PackedPageTableValue.xyzw & 0xf, ~0u);
+	uint MinSampledLevel = min(min(min(SampledLevels.x, SampledLevels.y), SampledLevels.z), SampledLevels.w);
+	if (MinSampledLevel != ~0u && MinSampledLevel > RequestedLevel)
 	{
-		uint SampleLevel = PackedPageTableValue & 0xf;
-		return SampleLevel;
+		AccumulatedResult = max(MinSampledLevel - RequestedLevel, AccumulatedResult);
 	}
-
-	// Return whatever the caller can interpret as empty.
-	return FallbackValue;
+#endif
+	return AccumulatedResult;
 }
 
 float GetStochasticMipNoise(float2 SvPositionXY)
@@ -323,6 +324,8 @@ void TextureLoadVirtualPageTableInternal(
 	OutResult.PageTableValue[0] = PageTable0.Load(int3(vPageX, vPageY, vLevelClamped));
 	OutResult.PageTableValue[1] = uint4(0u, 0u, 0u, 0u);
 
+	OutResult.PendingMips = GetVirtualTexturePendingMips(vLevelClamped, OutResult.PageTableValue[0], 0);
+
 #if VIRTUAL_TEXTURE_MANUAL_TRILINEAR_FILTERING
 	// Second page table for trilinear.
 	const uint vLevelClamped2 = clamp(vLevel + 1, 0, int(PageTableUniform.MaxLevel));
@@ -332,6 +335,8 @@ void TextureLoadVirtualPageTableInternal(
 	OutResult.PageTableValue[2] = PageTable0.Load(int3(vPageX2, vPageY2, vLevelClamped2));
 	OutResult.PageTableValue[3] = uint4(0u, 0u, 0u, 0u);
 	
+	OutResult.PendingMips = GetVirtualTexturePendingMips(vLevelClamped2, OutResult.PageTableValue[2], OutResult.PendingMips);
+
 	// Alternate requests to both mip levels
 	if ((View.FrameNumber & 1u) == 0u)
 	{
@@ -368,6 +373,9 @@ void TextureLoadVirtualPageTableInternal(
 	OutResult.PageTableValue[0] = PageTable0.Load(int3(vPageX, vPageY, vLevelClamped));
 	OutResult.PageTableValue[1] = PageTable1.Load(int3(vPageX, vPageY, vLevelClamped));
 
+	OutResult.PendingMips = GetVirtualTexturePendingMips(vLevelClamped, OutResult.PageTableValue[0], 0);
+	OutResult.PendingMips = GetVirtualTexturePendingMips(vLevelClamped, OutResult.PageTableValue[1], OutResult.PendingMips);
+
 #if VIRTUAL_TEXTURE_MANUAL_TRILINEAR_FILTERING
 	// Second page table for trilinear.
 	const uint vLevelClamped2 = clamp(vLevel + 1, 0, int(PageTableUniform.MaxLevel));
@@ -426,7 +434,7 @@ VTPageTableResult TextureLoadVirtualPageTable(
 #endif // PIXELSHADER
 		UV = ApplyAddressMode(UV, AddressU, AddressV);
 		TextureLoadVirtualPageTableInternal(Result, PageTable0, PageTableUniform, UV, vLevel);
-		StoreVirtualTextureFeedback(Feedback, SampleIndex, Result.PackedRequest);
+		StoreVirtualTextureFeedback(Feedback, SampleIndex, Result.PackedRequest, Result.PendingMips);
 	}
 	return Result;
 }
@@ -471,7 +479,7 @@ VTPageTableResult TextureLoadVirtualPageTable(
 #endif // PIXELSHADER
 		UV = ApplyAddressMode(UV, AddressU, AddressV);
 		TextureLoadVirtualPageTableInternal(Result, PageTable0, PageTable1, PageTableUniform, UV, vLevel);
-		StoreVirtualTextureFeedback(Feedback, SampleIndex, Result.PackedRequest);
+		StoreVirtualTextureFeedback(Feedback, SampleIndex, Result.PackedRequest, Result.PendingMips);
 	}
 	return Result;
 }
@@ -513,7 +521,7 @@ VTPageTableResult TextureLoadVirtualPageTableGrad(
 		UV = UV * PageTableUniform.UVScale;
 		UV = ApplyAddressMode(UV, AddressU, AddressV);
 		TextureLoadVirtualPageTableInternal(Result, PageTable0, PageTableUniform, UV, vLevel);
-		StoreVirtualTextureFeedback(Feedback, SampleIndex, Result.PackedRequest);
+		StoreVirtualTextureFeedback(Feedback, SampleIndex, Result.PackedRequest, Result.PendingMips);
 	}
 	return Result;
 }
@@ -552,7 +560,7 @@ VTPageTableResult TextureLoadVirtualPageTableGrad(
 		UV = UV * PageTableUniform.UVScale;
 		UV = ApplyAddressMode(UV, AddressU, AddressV);
 		TextureLoadVirtualPageTableInternal(Result, PageTable0, PageTable1, PageTableUniform, UV, vLevel);
-		StoreVirtualTextureFeedback(Feedback, SampleIndex, Result.PackedRequest);
+		StoreVirtualTextureFeedback(Feedback, SampleIndex, Result.PackedRequest, Result.PendingMips);
 	}
 	return Result;
 }
@@ -591,7 +599,7 @@ VTPageTableResult TextureLoadVirtualPageTableLevel(
 		int vLevel = (int)floor(Level + GetGlobalVirtualTextureMipBias());
 		UV = ApplyAddressMode(UV, AddressU, AddressV);
 		TextureLoadVirtualPageTableInternal(Result, PageTable0, PageTableUniform, UV, vLevel);
-		StoreVirtualTextureFeedback(Feedback, SampleIndex, Result.PackedRequest);
+		StoreVirtualTextureFeedback(Feedback, SampleIndex, Result.PackedRequest, Result.PendingMips);
 	}
 	return Result;
 }
@@ -630,7 +638,7 @@ VTPageTableResult TextureLoadVirtualPageTableLevel(
 		int vLevel = (int)floor(Level + GetGlobalVirtualTextureMipBias());
 		UV = ApplyAddressMode(UV, AddressU, AddressV);
 		TextureLoadVirtualPageTableInternal(Result, PageTable0, PageTable1, PageTableUniform, UV, vLevel);
-		StoreVirtualTextureFeedback(Feedback, SampleIndex, Result.PackedRequest);
+		StoreVirtualTextureFeedback(Feedback, SampleIndex, Result.PackedRequest, Result.PendingMips);
 	}
 	return Result;
 }
@@ -669,7 +677,7 @@ VTPageTableResult TextureLoadVirtualPageTableAdaptive(
 	UV = ApplyAddressMode(UV, AddressU, AddressV);
 	ApplyAdaptivePageTableUniform(PageTableIndirection, Result, PageTableUniform, UV, vLevel);
 	TextureLoadVirtualPageTableInternal(Result, PageTable0, PageTableUniform, UV, vLevel);
-	StoreVirtualTextureFeedback(Feedback, SampleIndex, Result.PackedRequest);
+	StoreVirtualTextureFeedback(Feedback, SampleIndex, Result.PackedRequest, Result.PendingMips);
 	return Result;
 }
 
@@ -689,7 +697,7 @@ VTPageTableResult TextureLoadVirtualPageTableAdaptiveGrad(
 	UV = ApplyAddressMode(UV, AddressU, AddressV);
 	ApplyAdaptivePageTableUniform(PageTableIndirection, Result, PageTableUniform, UV, vLevel);
 	TextureLoadVirtualPageTableInternal(Result, PageTable0, PageTableUniform, UV, vLevel);
-	StoreVirtualTextureFeedback(Feedback, SampleIndex, Result.PackedRequest);
+	StoreVirtualTextureFeedback(Feedback, SampleIndex, Result.PackedRequest, Result.PendingMips);
 	return Result;
 }
 
@@ -710,7 +718,7 @@ VTPageTableResult TextureLoadVirtualPageTableAdaptiveLevel(
 	UV = ApplyAddressMode(UV, AddressU, AddressV);
 	ApplyAdaptivePageTableUniform(PageTableIndirection, Result, PageTableUniform, UV, vLevel);
 	TextureLoadVirtualPageTableInternal(Result, PageTable0, PageTableUniform, UV, vLevel);
-	StoreVirtualTextureFeedback(Feedback, SampleIndex, Result.PackedRequest);
+	StoreVirtualTextureFeedback(Feedback, SampleIndex, Result.PackedRequest, Result.PendingMips);
 	return Result;
 }
 
diff --git a/Engine/Source/Runtime/Core/Private/UObject/FortniteMainBranchObjectVersion.cpp b/Engine/Source/Runtime/Core/Private/UObject/FortniteMainBranchObjectVersion.cpp
--- a/Engine/Source/Runtime/Core/Private/UObject/FortniteMainBranchObjectVersion.cpp
+++ b/Engine/Source/Runtime/Core/Private/UObject/FortniteMainBranchObjectVersion.cpp
@@ -14,6 +14,6 @@ TMap<FGuid, FGuid> FFortniteMainBranchObjectVersion::GetSystemGuids()
 	SystemGuids.Add(DevGuids.Niagara_LatestScriptCompileVersion, FGuid("670A2EBDEE1A4794BB77BBFF331EA7EF"));
 	SystemGuids.Add(DevGuids.SkeletalMeshDerivedDataVersion, FGuid("716424BCEC314E969ACFE36D2C2ECE49"));
 	SystemGuids.Add(DevGuids.STATICMESH_DERIVEDDATA_VER, FGuid("3ABF32149F9D4D838750368AB95573BE"));
-	SystemGuids.Add(DevGuids.MaterialTranslationDDCVersion, FGuid("56EE7432297C42BF90244A36A8D8D841"));
+	SystemGuids.Add(DevGuids.MaterialTranslationDDCVersion, FGuid("84CF0B79D69D4A4BBE8CAFD4B3D48AB7"));
 	return SystemGuids;
 }
diff --git a/Engine/Source/Runtime/Engine/Private/Materials/HLSLMaterialTranslator.cpp b/Engine/Source/Runtime/Engine/Private/Materials/HLSLMaterialTranslator.cpp
--- a/Engine/Source/Runtime/Engine/Private/Materials/HLSLMaterialTranslator.cpp
+++ b/Engine/Source/Runtime/Engine/Private/Materials/HLSLMaterialTranslator.cpp
@@ -8100,12 +8100,6 @@ int32 FHLSLMaterialTranslator::TextureSample(
 		{
 			SamplingCodeIndex = AddCodeChunk(MCT_Float4, *SampleCodeFinite, *TextureName, *VTPageTableResult_Finite, VTPageTableIndex, VirtualTextureIndex);
 		}
-	
-		const bool bStoreVTSampleInfo = ShaderFrequency == SF_Pixel && TextureReferenceIndex != INDEX_NONE;
-		if (bStoreVTSampleInfo)
-		{
-			AddCodeChunk(MCT_Float, TEXT("MaterialStoreVTSampleInfo(Parameters.TexCoordScalesParams, %s, %d, %d)"), *VTPageTableResult_Finite, VTPageTableIndex, TextureReferenceIndex);
-		}
 	}
 	else
 	{
diff --git a/Engine/Source/Runtime/Renderer/Private/DebugViewModeRendering.cpp b/Engine/Source/Runtime/Renderer/Private/DebugViewModeRendering.cpp
--- a/Engine/Source/Runtime/Renderer/Private/DebugViewModeRendering.cpp
+++ b/Engine/Source/Runtime/Renderer/Private/DebugViewModeRendering.cpp
@@ -284,7 +284,7 @@ void FDebugViewModePS::ModifyCompilationEnvironment(const FMaterialShaderPermuta
 	BufferRegister[1] += GetQuadOverdrawUAVIndex(Parameters.Platform);
 	OutEnvironment.SetDefine(TEXT("QUAD_BUFFER_REGISTER"), BufferRegister);
 	OutEnvironment.SetDefine(TEXT("OUTPUT_QUAD_OVERDRAW"), !FDataDrivenShaderPlatformInfo::GetIsLanguageMetal(Parameters.Platform));
-
+	OutEnvironment.SetDefine(TEXT("WRITE_PENDING_VIRTUALTEXTURE_MIPS"), 1u);
 
 	for (int i = 0; i < DVSM_MAX; ++i)
 	{
