Fix warnings due to redefinition when compiling OneColorShader.usf in platform preview.

diff --git a/Engine/Shaders/Private/OneColorShader.usf b/Engine/Shaders/Private/OneColorShader.usf
--- a/Engine/Shaders/Private/OneColorShader.usf
+++ b/Engine/Shaders/Private/OneColorShader.usf
@@ -89,15 +89,15 @@ void MainPixelShaderMRT(
 )
 {
 #if NUM_UINT_OUTPUTS < NUM_OUTPUTS
-	for (int i = 0; i < NUM_FLOAT_OUTPUTS; i++)
+	for (int i0 = 0; i0 < NUM_FLOAT_OUTPUTS; i0++)
 	{
-		OutColor[i] = DrawColorMRT[i];
+		OutColor[i0] = DrawColorMRT[i0];
 	}
 #endif
 #if NUM_UINT_OUTPUTS > 0
-	for (int i = 0; i < NUM_UINT_OUTPUTS; i++)
+	for (int i1 = 0; i1 < NUM_UINT_OUTPUTS; i1++)
 	{
-		OutUints[i] = DrawColorMRT[NUM_FLOAT_OUTPUTS + i].x;
+		OutUints[i1] = DrawColorMRT[NUM_FLOAT_OUTPUTS + i1].x;
 	}
 #endif
 }
