Path Tracer: Fix incorrect AOV decomposition of volumes

A recent refactor meant that direct lighting would toggle indirect lighting as well. We need to fetch samples if either direct or indirect sampling is enabled. Since this is rare in practice, always get the sample except on the last bounce when we only need to account for emissive samples.

#rnx

diff --git a/Engine/Shaders/Private/PathTracing/PathTracingCore.ush b/Engine/Shaders/Private/PathTracing/PathTracingCore.ush
--- a/Engine/Shaders/Private/PathTracing/PathTracingCore.ush
+++ b/Engine/Shaders/Private/PathTracing/PathTracingCore.ush
@@ -741,7 +741,7 @@ FPackedPathTracingPayload TraceTransparentRay(inout FPathState PathState, int Bo
 				{
 					// let callable shader handle cloud sampling
 					FPackedPathTracingPayload CloudPayload = (FPackedPathTracingPayload) 0;
-					const bool bNeedSample = !bLastBounce && PathState.ShouldAccumulateVolume();
+					const bool bNeedSample = !bLastBounce;
 					CloudPayload.SetVolumetricCallableShaderInput(
 						PathState.Ray.Origin,
 						PathState.Ray.Direction,
@@ -831,7 +831,7 @@ FPackedPathTracingPayload TraceTransparentRay(inout FPathState PathState, int Bo
 							}
 
 							// if this is not the last bounce, consider the possibility of scattering to give us medium hits
-							if (!bLastBounce && PathState.ShouldAccumulateVolume())
+							if (!bLastBounce)
 							{
 								float3 SigmaS = min(Result.SigmaSRayleigh + Result.SigmaSHG + Result.SigmaSDualHG, SigmaT);
 								// accumulate a signal for the denoiser
