Substrate - made more visualization shader complient with blendableGBuffer option.

#rnx
#tests PC editor dx12

diff --git a/Engine/Shaders/Private/Substrate/SubstrateVisualize.usf b/Engine/Shaders/Private/Substrate/SubstrateVisualize.usf
--- a/Engine/Shaders/Private/Substrate/SubstrateVisualize.usf
+++ b/Engine/Shaders/Private/Substrate/SubstrateVisualize.usf
@@ -55,6 +55,15 @@ void MaterialPrintInfoCS(uint3 DispatchThreadId : SV_DispatchThreadID)
 
 	FShaderPrintContext Context = InitShaderPrintContext(true, uint2(50, 50));
 
+#if SUBTRATE_GBUFFER_FORMAT==0
+
+	SerializeNullPixelDebugDataEntry();
+
+	Context = InitShaderPrintContext(true, uint2(50, 25));
+	Print(Context, TEXT("Blendable GBuffer does not support pixel inspection yet."), FontOrange);
+
+#else
+
 	// Print pixel material properties. Simple print for simple compilers such as dx11/fxc which cannot handle loop...
 #if SUBSTRATE_MAX_CLOSURE_COUNT!= 8
 #error Unsupported SUBSTRATE_MAX_CLOSURE_COUNT in visualisation shader.
@@ -74,6 +83,8 @@ void MaterialPrintInfoCS(uint3 DispatchThreadId : SV_DispatchThreadID)
 	SerializeSubstratePixelDebugDataEntry(PixelPos, PrimaryView.FrameNumber);
 	#endif
 
+#endif
+
 	// Print pixel footprint
 	const float3 WorldPositionX = GetWorldPositionFromPixelPos(PixelPos + uint2(1, 0));
 	const float3 WorldPositionY = GetWorldPositionFromPixelPos(PixelPos + uint2(0, 1));
@@ -135,10 +146,11 @@ void VisualizeMaterialPS(
 	const float3 WorldPosition = mul(float4(ScreenMaterialPosition * SceneDepth, SceneDepth, 1), PrimaryView.ScreenToTranslatedWorld).xyz;
 	const float3 V = -GetCameraVectorFromTranslatedWorldPosition(WorldPosition);
 
+
+#if SUBTRATE_GBUFFER_FORMAT==1
 	FSubstrateAddressing SubstrateAddressing = GetSubstratePixelDataByteOffset(PixelPosDynRes, uint2(View.BufferSizeAndInvSize.xy), Substrate.MaxBytesPerPixel);
-	const uint FootPrint_Start = SubstrateAddressing.ReadBytes;
 	FSubstratePixelHeader Header = UnpackSubstrateHeaderIn(Substrate.MaterialTextureArray, SubstrateAddressing, Substrate.TopLayerTexture);
-	const bool bSubstrateMaterial = Header.ClosureCount > 0;
+#endif
 
 	OutColor = 0;
 	OutColor.w = 0.1f;
@@ -146,11 +158,16 @@ void VisualizeMaterialPS(
 	{
 		// BSDF count
 		const float MaxClosureCount = SUBSTRATE_MAX_CLOSURE_COUNT;
+#if SUBTRATE_GBUFFER_FORMAT==1
 		if (Header.ClosureCount > 0)
 		{
 			const float BSDFCost = saturate(float(Header.ClosureCount) / MaxClosureCount);
 			OutColor.xyz = ColorMapTurbo(BSDFCost);
 		}
+#else
+		const float BSDFCost = saturate(1.0f / MaxClosureCount);
+		OutColor.xyz = ColorMapTurbo(BSDFCost);
+#endif
 	
 		// Print Legend
 		if (all(uint2(SVPos.xy) == uint2(0,0)))
@@ -164,10 +181,8 @@ void VisualizeMaterialPS(
 	}
 	else if (ViewMode == 3)
 	{
-		FSubstrateAddressing SubstrateAddressing = GetSubstratePixelDataByteOffset(PixelPosDynRes, uint2(View.BufferSizeAndInvSize.xy), Substrate.MaxBytesPerPixel);
-		FSubstratePixelHeader Header = UnpackSubstrateHeaderIn(Substrate.MaterialTextureArray, SubstrateAddressing, Substrate.TopLayerTexture);
-
 		uint MaterialBytePerPixel = 0;
+#if SUBTRATE_GBUFFER_FORMAT==1
 		Substrate_for (uint ClosureIndex = 0, ClosureIndex < Header.ClosureCount, ++ClosureIndex)
 		{
 			// Unpack BSDF data
@@ -179,6 +194,10 @@ void VisualizeMaterialPS(
 		{
 			OutColor.xyz = BytePerPixelToColor(MaterialBytePerPixel);
 		}
+#else
+		MaterialBytePerPixel = 4 * 4; // 4 32bits gbuffer
+		OutColor.xyz = BytePerPixelToColor(MaterialBytePerPixel);
+#endif
 	
 		// Print Legend
 		if (all(uint2(SVPos.xy) == uint2(0,0)))
diff --git a/Engine/Shaders/Shared/SubstrateVisualizeDefinitions.h b/Engine/Shaders/Shared/SubstrateVisualizeDefinitions.h
--- a/Engine/Shaders/Shared/SubstrateVisualizeDefinitions.h
+++ b/Engine/Shaders/Shared/SubstrateVisualizeDefinitions.h
@@ -611,6 +611,14 @@ void SerializeSubstratePixelDebugDataEntry(uint2 InCoord, FSubstrateRaytracingPa
 	FSubstrateDebugDataSerializer S = (FSubstrateDebugDataSerializer)0;
 	SerializeSubstratePixelDebugData(S, Data);
 }
+
+void SerializeNullPixelDebugDataEntry()
+{
+	FSubstratePixelDebugData Data = (FSubstratePixelDebugData)0;
+	FSubstrateDebugDataSerializer S = (FSubstrateDebugDataSerializer)0;
+	SerializeSubstratePixelDebugData(S, Data);
+}
+
 #endif
 
 
diff --git a/Engine/Source/Runtime/Renderer/Private/Substrate/SubstrateVisualize.cpp b/Engine/Source/Runtime/Renderer/Private/Substrate/SubstrateVisualize.cpp
--- a/Engine/Source/Runtime/Renderer/Private/Substrate/SubstrateVisualize.cpp
+++ b/Engine/Source/Runtime/Renderer/Private/Substrate/SubstrateVisualize.cpp
@@ -1146,9 +1146,11 @@ static void AddVisualizeSystemInfoPasses(FRDGBuilder& GraphBuilder, const FViewI
 				{
 					PrintS(DrawPosX, DrawPosY, TEXT("  Using blendable GBuffer with fixed layout."), DarkYellow);
 					NewLine();
-					PrintS(DrawPosX, DrawPosY, TEXT("    - 4 x 32 bits textures."), White);
+					PrintS(DrawPosX, DrawPosY, TEXT("  Bytes Per Pixel"), White);
+					PrintI(DrawPosX + 150, DrawPosY, 4*4, Orange);
 					NewLine();
-					PrintS(DrawPosX, DrawPosY, TEXT("    - 1 closures per pixel."), White);
+					PrintS(DrawPosX, DrawPosY, TEXT("  Closures Per Pixel"), White);
+					PrintI(DrawPosX + 150, DrawPosY, 1, Orange);
 				}
 				else
 				{
