[DOF] Optimize setup pass. Group size is now 8x8 regardless of output format. When available, lane swizzling is used to share texels, rather than LDS. A morton-like curve is used during sampling for better cache coherency.

diff --git a/Engine/Shaders/Private/DiaphragmDOF/DOFSetup.usf b/Engine/Shaders/Private/DiaphragmDOF/DOFSetup.usf
--- a/Engine/Shaders/Private/DiaphragmDOF/DOFSetup.usf
+++ b/Engine/Shaders/Private/DiaphragmDOF/DOFSetup.usf
@@ -11,17 +11,13 @@
 
 //------------------------------------------------------- COMPILE TIME CONSTANTS
 
-#if DIM_OUTPUT_RES_DIVISOR == 0
-	#define GROUP_BORDER_SIZE (2 * DEFAULT_GROUP_BORDER_SIZE)
-#else
-	#define GROUP_BORDER_SIZE (DEFAULT_GROUP_BORDER_SIZE)
-#endif
+#define THREADGROUP_SIZE_X (8)
+#define THREADGROUP_SIZE_Y (8) 
+#define THREADGROUP_SIZE_TOTAL (THREADGROUP_SIZE_X * THREADGROUP_SIZE_Y)
 
 
 //------------------------------------------------------- PARAMETERS
 
-uint4 ViewportRect;
-
 Texture2D SceneColorTexture;
 Texture2D SceneDepthTexture;
 
@@ -39,255 +35,469 @@ RWTexture2D<float4>	Output1;
 
 #endif
 
+//------------------------------------------------------- FUNCTIONS
 
-//------------------------------------------------------- LDS
 
-#define SUBTILE_SIZE (DEFAULT_GROUP_BORDER_SIZE)
+struct FDOFSetupParameters
+{
+	// Compile-time variables
+	bool bOutputFullRes;
+	bool bOutputHalfRes;
 
-#if CONFIG_DOF_ALPHA
-	groupshared float4 SharedColor[SUBTILE_SIZE * SUBTILE_SIZE * 3];
-	groupshared float SharedCocRadius[SUBTILE_SIZE * SUBTILE_SIZE * 3];
-#else
-	groupshared float4 SharedColorAndCocRadius[SUBTILE_SIZE * SUBTILE_SIZE * 3];
-#endif
+	Texture2D SceneColorTexture;
+	Texture2D SceneDepthTexture;
+};
 
-//------------------------------------------------------- ENTRY POINT
+struct FDOFSetupOutput
+{
+	// Requires bOutputFullRes
+	bool bFullResIsValid;
+	uint2 FullResOutputPixelPosition;
+	float4 FullResColor;
+	float FullResCocRadius;
+
+	// Requires bOutputHalfRes
+	bool bHalfResIsValid;
+	uint2 HalfResOutputPixelPosition;
+	float4 HalfResColor;
+	float HalfResCocRadius;
+};
+
+// Calculate a 2D space-filling curve, similar to Morton, from a 1D coordinate.
+// From the AMD GPUOpen "Optimizing for the Radeon RDNA architecture"
+//   00 01 08 09 16 17 24 25
+//   02 03 10 11 18 19 26 27
+//   04 05 12 13 20 21 28 29
+//   06 07 14 15 22 23 30 31
+//   32 33 40 41 48 49 56 57
+//   34 35 42 43 50 51 58 59
+//   36 37 44 45 52 53 60 61
+//   38 39 46 47 54 55 62 63
+uint2 MortonLikeDecode(uint Index)
+{
+	uint2 Result;
+	Result.x = (((Index >> 2) & 0x0007) & 0xFFFE) | (Index & 0x0001);
+	Result.y = ((Index >> 1) & 0x0003) | (((Index >> 3) & 0x0007) & 0xFFFC);
+	return Result;
+}
 
-#if DIM_OUTPUT_RES_DIVISOR == 0 // Output full & half resolution.
+#if WAVE_OPS
+groupshared uint GroupNumWaves;
 
-[numthreads(GROUP_BORDER_SIZE, GROUP_BORDER_SIZE, 1)]
-void SetupCS(
-	uint2 GroupId : SV_GroupID,
-	uint GroupThreadIndex : SV_GroupIndex)
+struct FSoftwareQuad
 {
-	uint2 SubGroupId = uint2((GroupThreadIndex & 64) != 0, (GroupThreadIndex & 128) != 0);
-	uint2 SubGroupPos = uint2(GroupThreadIndex & 0x7, (GroupThreadIndex / 8) & 0x7);
+	uint WaveQuadLeaderLane;
+	bool bIsLeaderLane;
+	uint2 LocalPixelPosition;
 
-	bool bIsDownSamplingSubGroup = GroupThreadIndex < SUBTILE_SIZE * SUBTILE_SIZE;
-	
-	uint2 InputPixelPosition = (uint2(View.ViewRectMin.xy) + GroupId * GROUP_BORDER_SIZE) + SubGroupPos * 2 + SubGroupId;
+	float4 LocalColor;
+	float LocalCocRadius;
 
-	// Fetch scene color/depth, and compute CocRadius.
-	float2 SceneBufferUV = (InputPixelPosition + 0.5) * View.BufferSizeAndInvSize.zw;
-	
-	if (true)
+	void Setup(uint2 GroupId, uint GroupThreadIndex)
 	{
-		SceneBufferUV = clamp(SceneBufferUV, View.BufferBilinearUVMinMax.xy, View.BufferBilinearUVMinMax.zw);
+		uint GroupWaveIndex = 0;
+		if (WaveGetLaneCount() < THREADGROUP_SIZE_TOTAL)
+		{
+			if(GroupThreadIndex == 0)
+			{
+				GroupNumWaves = 0;
+			}
+			GroupMemoryBarrierWithGroupSync();
+
+			if(WaveIsFirstLane())
+			{
+				InterlockedAdd(GroupNumWaves, 1, GroupWaveIndex);
+			}
+			GroupWaveIndex = WaveReadLaneFirst(GroupWaveIndex);
+		}
+
+		uint2 GroupSize = uint2(THREADGROUP_SIZE_X, THREADGROUP_SIZE_Y);
+
+		// Index of lane in group
+		uint GroupLane = (GroupWaveIndex * WaveGetLaneCount()) + WaveGetLaneIndex();
+
+		// Map threads to pixels using morton-like curve for better texture cache hit rate
+		LocalPixelPosition = (GroupSize * GroupId) + MortonLikeDecode(GroupLane);
+		WaveQuadLeaderLane = WaveGetLaneIndex() & (~(4 - 1));
+		bIsLeaderLane = (WaveGetLaneIndex() & (4 - 1)) == 0;
 	}
 
-	float4 SceneColor = SceneColorTexture.SampleLevel(GlobalPointClampedSampler, SceneBufferUV, 0);
-	float DeviceZ = SceneDepthTexture.SampleLevel(GlobalPointClampedSampler, SceneBufferUV, 0).r;
-	float CocRadius = SceneDepthToCocRadius(ConvertFromDeviceZ(DeviceZ));
-	
-	// Output full resolution.
-	#if CONFIG_DOF_ALPHA
+	bool IsLeader()
 	{
-		Output0[uint2(InputPixelPosition)] = CocRadius;
+		return bIsLeaderLane;
 	}
-	#else
+
+	uint2 GetLocalPixelPosition()
 	{
-		Output0[uint2(InputPixelPosition)] = float4(SceneColor.rgb, CocRadius);
+		return LocalPixelPosition;
 	}
-	#endif
 
-	// Have the auxiliary thread writing to LDS for main group to downscale.
-	BRANCH
-	if (!bIsDownSamplingSubGroup)
+	void SetLocalPixel(float4 Color, float CocRadius)
 	{
-		#if CONFIG_DOF_ALPHA
-		{
-			SharedColor[GroupThreadIndex - SUBTILE_SIZE * SUBTILE_SIZE] = SceneColor;
-			SharedCocRadius[GroupThreadIndex - SUBTILE_SIZE * SUBTILE_SIZE] = CocRadius;
-		}
+		LocalColor = Color;
+		LocalCocRadius = CocRadius;
+	}
+
+	void Sync()
+	{ }
+
+	void GetPixel(uint PixelIndex, out float4 Color, out float CocRadius)
+	{
+		// Could maybe also use QuadReadLaneAt where it is supported
+
+		#if COMPILER_SUPPORTS_WAVE_SWIZZLE_GCN
+			const uint AndMask = ~(4 - 1);
+			const uint OrMask = PixelIndex;
+			Color.r = asfloat(WaveLaneSwizzleGCN(asuint(LocalColor.r), AndMask, OrMask, 0));
+			Color.g = asfloat(WaveLaneSwizzleGCN(asuint(LocalColor.g), AndMask, OrMask, 0));
+			Color.b = asfloat(WaveLaneSwizzleGCN(asuint(LocalColor.b), AndMask, OrMask, 0));
+			Color.a = asfloat(WaveLaneSwizzleGCN(asuint(LocalColor.a), AndMask, OrMask, 0));
+			CocRadius = asfloat(WaveLaneSwizzleGCN(asuint(LocalCocRadius), AndMask, OrMask, 0));
+		#elif COMPILER_SUPPORTS_WAVE_SWIZZLE_RDNA
+			Color.r = asfloat(WaveLaneRotateSwizzleRDNA(asuint(LocalColor.r), PixelIndex, 0));
+			Color.g = asfloat(WaveLaneRotateSwizzleRDNA(asuint(LocalColor.g), PixelIndex, 0));
+			Color.b = asfloat(WaveLaneRotateSwizzleRDNA(asuint(LocalColor.b), PixelIndex, 0));
+			Color.a = asfloat(WaveLaneRotateSwizzleRDNA(asuint(LocalColor.a), PixelIndex, 0));
+			CocRadius = asfloat(WaveLaneRotateSwizzleRDNA(asuint(LocalCocRadius), PixelIndex, 0));
 		#else
+			Color = WaveReadLaneAt(LocalColor, WaveQuadLeaderLane + PixelIndex);
+			CocRadius = WaveReadLaneAt(LocalCocRadius, WaveQuadLeaderLane + PixelIndex);
+		#endif // COMPILER_SUPPORTS_WAVE_SWIZZLE_RDNA
+	}
+};
+
+#else // !WAVE_OPS
+
+groupshared float4 SoftQuadColors[THREADGROUP_SIZE_TOTAL];
+groupshared float SoftQuadCocRadius[THREADGROUP_SIZE_TOTAL];
+struct FSoftwareQuad
+{
+	uint GroupLaneIndex;
+	uint QuadLeaderLane;
+	bool bIsLeaderLane;
+	uint2 LocalPixelPosition;
+
+	void Setup(uint2 GroupId, uint GroupThreadIndex)
+	{
+		uint2 GroupSize = uint2(THREADGROUP_SIZE_X, THREADGROUP_SIZE_Y);
+
+		// Map threads to pixels using morton-like curve for better texture cache hit rate
+		GroupLaneIndex = GroupThreadIndex;
+		LocalPixelPosition = (GroupSize * GroupId) + MortonLikeDecode(GroupThreadIndex);
+		QuadLeaderLane = WaveGetLaneIndex() & (~(4 - 1));
+		bIsLeaderLane = (WaveGetLaneIndex() & (4 - 1)) == 0;
+	}
+
+	bool IsLeader()
+	{
+		return bIsLeaderLane;
+	}
+
+	uint2 GetLocalPixelPosition()
+	{
+		return LocalPixelPosition;
+	}
+
+	void SetLocalPixel(float4 Color, float CocRadius)
+	{
+		SoftQuadColors[GroupLaneIndex] = Color;
+		SoftQuadCocRadius[GroupLaneIndex] = CocRadius;
+	}
+
+	void Sync()
+	{
+		GroupMemoryBarrierWithGroupSync();
+	}
+
+	void GetPixel(uint PixelIndex, out float4 Color, out float CocRadius)
+	{
+		Color = 0;
+		CocRadius = 0;
+		
+		if (IsLeader())
 		{
-			SharedColorAndCocRadius[GroupThreadIndex - SUBTILE_SIZE * SUBTILE_SIZE] = float4(SceneColor.rgb, CocRadius);
+			Color = SoftQuadColors[QuadLeaderLane+PixelIndex];
+			CocRadius = SoftQuadCocRadius[QuadLeaderLane+PixelIndex];
 		}
-		#endif
 	}
+};
+
+#endif // WAVE_OPS
 
-	GroupMemoryBarrierWithGroupSync();
+// Launch with 1 thread per fullres pixel, and 8x8 groupsize
+FDOFSetupOutput SetupDOF(
+	FDOFSetupParameters Parameters,
+	uint2 GroupId,
+	uint GroupThreadIndex,
+	uint2 DispatchThreadId
+)
+{
+	FDOFSetupOutput Output = (FDOFSetupOutput)0;
 
-	// Have main threads to downscale.
-	BRANCH
-	if (bIsDownSamplingSubGroup)
+	if(Parameters.bOutputFullRes && Parameters.bOutputHalfRes)
 	{
-		// fetch from LDS.
+		uint2 GroupSize = uint2(THREADGROUP_SIZE_X, THREADGROUP_SIZE_Y);
+		if(any((GroupId * GroupSize) >= View.ViewRectMinAndSize.zw))
+		{
+			return Output;
+		}
+		
+		// group is twice as large in x and y as in halfres/fullres case
+		// group is subdivided in 4 subgroups of 8x8 (2x2 subgroups)
+		// 1 thread is 1 fullres pixel
+		// 1 thread from first subgroup is 1 halfres pixel
+
+		FSoftwareQuad Quad;
+		Quad.Setup(GroupId, GroupThreadIndex);
+
+		uint2 InputPixelPosition = uint2(View.ViewRectMin.xy) + min(Quad.GetLocalPixelPosition(), View.ViewRectMinAndSize.zw);
+		Output.FullResOutputPixelPosition = uint2(View.ViewRectMin.xy) + Quad.GetLocalPixelPosition();
+		
+		Output.bFullResIsValid = all(Quad.GetLocalPixelPosition() < View.ViewRectMinAndSize.zw);
+
+		// Fetch scene color/depth, and compute CocRadius.
+		float2 SceneBufferUV = (InputPixelPosition + 0.5) * View.BufferSizeAndInvSize.zw;
+	
+		if (true)
+		{
+			SceneBufferUV = clamp(SceneBufferUV, View.BufferBilinearUVMinMax.xy, View.BufferBilinearUVMinMax.zw);
+		}
+
+		float4 SceneColor = Parameters.SceneColorTexture.SampleLevel(GlobalPointClampedSampler, SceneBufferUV, 0);
+		float DeviceZ = Parameters.SceneDepthTexture.SampleLevel(GlobalPointClampedSampler, SceneBufferUV, 0).r;
+		float CocRadius = SceneDepthToCocRadius(ConvertFromDeviceZ(DeviceZ));
+	
+		// Output full resolution.
+		Output.FullResColor = SceneColor;
+		Output.FullResCocRadius = CocRadius;
+
+		Quad.SetLocalPixel(SceneColor, CocRadius);
+		Quad.Sync();
+
 		float4 SceneColors[4];
 		float CocRadii[4];
 
-		SceneColors[0] = SceneColor;
-		CocRadii[0] = CocRadius;
-
 		UNROLL
-		for (uint i = 0; i < 3; i++)
+		for (uint Index = 0; Index < 4; Index++)
 		{
-			uint LDSIndex = GroupThreadIndex + SUBTILE_SIZE * SUBTILE_SIZE * i;
-			#if CONFIG_DOF_ALPHA
-			{
-				SceneColors[i + 1] = SharedColor[LDSIndex];
-				CocRadii[i + 1] = SharedCocRadius[LDSIndex];
-			}
-			#else
-			{
-				float4 RawLDS = SharedColorAndCocRadius[LDSIndex];
+			Quad.GetPixel(Index, /* out */ SceneColors[Index], /* out */ CocRadii[Index]);
+		}
 
-				SceneColors[i + 1] = float4(RawLDS.rgb, 0);
-				CocRadii[i + 1] = RawLDS.a;
-			}
-			#endif
+		BRANCH
+		if (Quad.IsLeader())
+		{
+			Output.bHalfResIsValid = true;
+			Output.HalfResOutputPixelPosition = InputPixelPosition / 2;
+
+			FCocDownsampleParams DownsampleParams;
+			DownsampleParams.CocRadiusMultiplier = 1.0;
+			DownsampleParams.FrameExposureScale = 1.0;
+			DownsampleParams.bDoColorBasedWeighting = false;
+
+			DownsampleSceneColorWithCoc(DownsampleParams, SceneColors, CocRadii, Output.HalfResColor, Output.HalfResCocRadius);
 		}
+	}
+	else if(Parameters.bOutputFullRes && !Parameters.bOutputHalfRes)
+	{
+		if(any(DispatchThreadId >= View.ViewRectMinAndSize.zw))
+		{
+			return Output;
+		}
+
+		Output.bFullResIsValid = true;
+
+		uint2 InputPixelPosition = View.ViewRectMin.xy + DispatchThreadId;
+
+		Output.FullResOutputPixelPosition = InputPixelPosition;
+		float2 SceneBufferUV = (InputPixelPosition + 0.5) * View.BufferSizeAndInvSize.zw;
+	
+		Output.FullResColor = Parameters.SceneColorTexture.SampleLevel(GlobalPointClampedSampler, SceneBufferUV, 0);
+	
+		float DeviceZ = Parameters.SceneDepthTexture.SampleLevel(GlobalPointClampedSampler, SceneBufferUV, 0).r;
+		Output.FullResCocRadius = SceneDepthToCocRadius(ConvertFromDeviceZ(DeviceZ));
+	}
+	else if(!Parameters.bOutputFullRes && Parameters.bOutputHalfRes)
+	{
+		uint2 HalfResViewSize = uint2(
+			DivideAndRoundUp(View.ViewRectMinAndSize.z, 2), 
+			DivideAndRoundUp(View.ViewRectMinAndSize.w, 2)
+		);
+		if(any(DispatchThreadId >= HalfResViewSize))
+		{
+			return Output;
+		}
+
+		Output.bHalfResIsValid = true;
 		
+		Output.HalfResOutputPixelPosition = (View.ViewRectMin.xy / 2) + DispatchThreadId;
+		float2 SceneBufferUV = (float2(View.ViewRectMin.xy) + 2.0 * (DispatchThreadId + 0.5)) * View.BufferSizeAndInvSize.zw;
+		float2 SceneBufferUVOffset = View.BufferSizeAndInvSize.zw * 0.5;
+			
 		FCocDownsampleParams DownsampleParams;
 		DownsampleParams.CocRadiusMultiplier = 1.0;
 		DownsampleParams.FrameExposureScale = 1.0;
 		DownsampleParams.bDoColorBasedWeighting = false;
 
-		// Downsample.
-		float4 OutColor;
-		float OutCocRadius;
-		DownsampleSceneColorWithCoc(DownsampleParams, SceneColors, CocRadii, OutColor, OutCocRadius);
-	
-		uint2 HalfResOutputPixelPosition = (View.ViewRectMin.xy/2) + GroupId * (GROUP_BORDER_SIZE / 2) + SubGroupPos;
+		uint2 GroupSize = uint2(THREADGROUP_SIZE_X, THREADGROUP_SIZE_Y);
+		uint2 ViewportSize = uint2(
+			DivideAndRoundUp(View.ViewRectMinAndSize.z, 2),
+			DivideAndRoundUp(View.ViewRectMinAndSize.w, 2));
+		bool bIsBorderPixel = any(GroupId * GroupSize + GroupSize >= ViewportSize); 
 
-		//if (all(HalfResOutputPixelPosition < ViewportRect.zw))
-		#if CONFIG_DOF_ALPHA
+		float CocRadii[4];
+		float4 SceneColors[4];
+
+		BRANCH
+		if (!bIsBorderPixel)
 		{
-			Output1[HalfResOutputPixelPosition] = OutColor;
-			Output2[HalfResOutputPixelPosition] = OutCocRadius;
+			// Fetch Coc
+			float4 SceneDepthGather = Parameters.SceneDepthTexture.Gather(GlobalPointClampedSampler, SceneBufferUV);
+			CocRadii[0] = SceneDepthToCocRadius(ConvertFromDeviceZ(SceneDepthGather.w));
+			CocRadii[1] = SceneDepthToCocRadius(ConvertFromDeviceZ(SceneDepthGather.z));
+			CocRadii[2] = SceneDepthToCocRadius(ConvertFromDeviceZ(SceneDepthGather.x));
+			CocRadii[3] = SceneDepthToCocRadius(ConvertFromDeviceZ(SceneDepthGather.y));
+
+			// Fetch color.
+			for (uint i = 0; i < 4; i++)
+			{
+				float2 SampleUV = SceneBufferUV + SceneBufferUVOffset * float2(kOffsetsCross3x3[i]);
+				SceneColors[i] = Parameters.SceneColorTexture.SampleLevel(GlobalPointClampedSampler, SampleUV, 0);
+			}
 		}
-		#else
+		else
 		{
-			Output1[HalfResOutputPixelPosition] = float4(OutColor.rgb, OutCocRadius);
+			// Sample scene color and depth.
+			
+			for (uint i = 0; i < 4; i++)
+			{
+				float2 SampleUV = SceneBufferUV + SceneBufferUVOffset * float2(kOffsetsCross3x3[i]);
+
+				if (true)
+				{
+					SampleUV = clamp(SampleUV, View.BufferBilinearUVMinMax.xy, View.BufferBilinearUVMinMax.zw);
+				}
+
+				SceneColors[i] = Parameters.SceneColorTexture.SampleLevel(GlobalPointClampedSampler, SampleUV, 0);
+			
+				float SampleDeviceZ = Parameters.SceneDepthTexture.SampleLevel(GlobalPointClampedSampler, SampleUV, 0).r;
+				CocRadii[i] = SceneDepthToCocRadius(ConvertFromDeviceZ(SampleDeviceZ));
+			}
 		}
-		#endif
+
+		DownsampleSceneColorWithCoc(DownsampleParams, SceneColors, CocRadii, Output.HalfResColor, Output.HalfResCocRadius);
 	}
+
+	return Output;
 }
 
 
-#elif DIM_OUTPUT_RES_DIVISOR == 1 // Output full resolution only.
+#if DIM_OUTPUT_RES_DIVISOR == 0 // Output full & half resolution.
 
-[numthreads(GROUP_BORDER_SIZE, GROUP_BORDER_SIZE, 1)]
-void SetupCS(uint2 DispatchThreadId : SV_DispatchThreadID) 
+[numthreads(THREADGROUP_SIZE_X, THREADGROUP_SIZE_Y, 1)]
+void SetupCS(
+	uint2 GroupId : SV_GroupID,
+	uint GroupThreadIndex : SV_GroupIndex,
+	uint2 DispatchThreadId : SV_DispatchThreadID)
 {
-	float2 SceneBufferUV = (float2(View.ViewRectMin.xy) + (DispatchThreadId + 0.5)) * View.BufferSizeAndInvSize.zw;
-	
-	float4 OutColor = SceneColorTexture.SampleLevel(GlobalPointClampedSampler, SceneBufferUV, 0);
+	FDOFSetupParameters Parameters;
+	Parameters.bOutputFullRes = true;
+	Parameters.bOutputHalfRes = true;
 	
-	float DeviceZ = SceneDepthTexture.SampleLevel(GlobalPointClampedSampler, SceneBufferUV, 0).r;
-	float OutCocRadius = SceneDepthToCocRadius(ConvertFromDeviceZ(DeviceZ));
+	Parameters.SceneColorTexture = SceneColorTexture;
+	Parameters.SceneDepthTexture = SceneDepthTexture;
+
+	FDOFSetupOutput Output = SetupDOF(Parameters, GroupId, GroupThreadIndex, DispatchThreadId);
+
+	if (Output.bFullResIsValid)
+	{
+		#if CONFIG_DOF_ALPHA
+		{
+			Output0[Output.FullResOutputPixelPosition] = Output.FullResCocRadius;
+		}
+		#else
+		{
+			Output0[Output.FullResOutputPixelPosition] = float4(Output.FullResColor.rgb, Output.FullResCocRadius);
+		}
+		#endif
+	}
 
-	if (all(DispatchThreadId < ViewportRect.zw))
+	if (Output.bHalfResIsValid)
 	{
 		#if CONFIG_DOF_ALPHA
 		{
-			Output0[View.ViewRectMin.xy + DispatchThreadId] = OutCocRadius;
+			Output1[Output.HalfResOutputPixelPosition] = Output.HalfResColor;
+			Output2[Output.HalfResOutputPixelPosition] = Output.HalfResCocRadius;
 		}
 		#else
 		{
-			Output0[View.ViewRectMin.xy + DispatchThreadId] = float4(OutColor.xyz, OutCocRadius);
+			Output1[Output.HalfResOutputPixelPosition] = float4(Output.HalfResColor.rgb, Output.HalfResCocRadius);
 		}
 		#endif
 	}
 }
 
 
-#elif DIM_OUTPUT_RES_DIVISOR == 2 // Output half resolution only.
+#elif DIM_OUTPUT_RES_DIVISOR == 1 // Output full resolution only.
 
-[numthreads(GROUP_BORDER_SIZE, GROUP_BORDER_SIZE, 1)]
+[numthreads(THREADGROUP_SIZE_X, THREADGROUP_SIZE_Y, 1)]
 void SetupCS(
 	uint2 GroupId : SV_GroupID,
-	uint2 DispatchThreadId : SV_DispatchThreadID) 
+	uint GroupThreadIndex : SV_GroupIndex,
+	uint2 DispatchThreadId : SV_DispatchThreadID)
 {
-	// Compute scene buffer uv, but carefully to handle view size % 2 == 1,
-	// because this pass is running at half resolution.
-	float2 SceneBufferUV = (float2(View.ViewRectMin.xy) + 2.0 * (DispatchThreadId + 0.5)) * View.BufferSizeAndInvSize.zw;
-	float2 SceneBufferUVOffset = View.BufferSizeAndInvSize.zw * 0.5;
-			
-	FCocDownsampleParams DownsampleParams;
-	DownsampleParams.CocRadiusMultiplier = 1.0;
-	DownsampleParams.FrameExposureScale = 1.0;
-	DownsampleParams.bDoColorBasedWeighting = false;
+	FDOFSetupParameters Parameters;
+	Parameters.bOutputFullRes = true;
+	Parameters.bOutputHalfRes = false;
+	
+	Parameters.SceneColorTexture = SceneColorTexture;
+	Parameters.SceneDepthTexture = SceneDepthTexture;
 
-	BRANCH
-	if (all(GroupId * GROUP_BORDER_SIZE + (GROUP_BORDER_SIZE + 1) <= ViewportRect.zw))
-	{
-		// Fetch Coc
-		float4 SceneDepthGather = SceneDepthTexture.Gather(GlobalPointClampedSampler, SceneBufferUV);
-		float CocRadii[4];
-		CocRadii[0] = SceneDepthToCocRadius(ConvertFromDeviceZ(SceneDepthGather.w));
-		CocRadii[1] = SceneDepthToCocRadius(ConvertFromDeviceZ(SceneDepthGather.z));
-		CocRadii[2] = SceneDepthToCocRadius(ConvertFromDeviceZ(SceneDepthGather.x));
-		CocRadii[3] = SceneDepthToCocRadius(ConvertFromDeviceZ(SceneDepthGather.y));
+	FDOFSetupOutput Output = SetupDOF(Parameters, GroupId, GroupThreadIndex, DispatchThreadId);
 
-		// Fetch color.
-		float4 SceneColors[4];
-		for (uint i = 0; i < 4; i++)
-		{
-			float2 SampleUV = SceneBufferUV + SceneBufferUVOffset * float2(kOffsetsCross3x3[i]);
-			SceneColors[i] = SceneColorTexture.SampleLevel(GlobalPointClampedSampler, SampleUV, 0);
-		}
-	
-		// Downsample.
-		float4 OutColor;
-		float OutCocRadius;
-		DownsampleSceneColorWithCoc(DownsampleParams, SceneColors, CocRadii, OutColor, OutCocRadius);
-	
+	if (Output.bFullResIsValid)
+	{
 		#if CONFIG_DOF_ALPHA
 		{
-			Output1[DispatchThreadId] = OutColor;
-			Output2[DispatchThreadId] = OutCocRadius;
+			Output0[Output.FullResOutputPixelPosition] = Output.FullResCocRadius;
 		}
 		#else
 		{
-			Output1[DispatchThreadId] = float4(OutColor.rgb, OutCocRadius);
+			Output0[Output.FullResOutputPixelPosition] = float4(Output.FullResColor.rgb, Output.FullResCocRadius);
 		}
 		#endif
 	}
-	else
-	{
-		// Sample scene color and depth.
-		float CocRadii[4];
-		float4 SceneColors[4];
-		for (uint i = 0; i < 4; i++)
-		{
-			float2 SampleUV = SceneBufferUV + SceneBufferUVOffset * float2(kOffsetsCross3x3[i]);
+}
 
-			if (true)
-			{
-				SampleUV = clamp(SampleUV, View.BufferBilinearUVMinMax.xy, View.BufferBilinearUVMinMax.zw);
-			}
 
-			SceneColors[i] = SceneColorTexture.SampleLevel(GlobalPointClampedSampler, SampleUV, 0);
-			
-			float SampleDeviceZ = SceneDepthTexture.SampleLevel(GlobalPointClampedSampler, SampleUV, 0).r;
-			CocRadii[i] = SceneDepthToCocRadius(ConvertFromDeviceZ(SampleDeviceZ));
-		}
+#elif DIM_OUTPUT_RES_DIVISOR == 2 // Output half resolution only.
 
-		// Downsample.
-		float4 OutColor;
-		float OutCocRadius;
-		DownsampleSceneColorWithCoc(DownsampleParams, SceneColors, CocRadii, OutColor, OutCocRadius);
+[numthreads(THREADGROUP_SIZE_X, THREADGROUP_SIZE_Y, 1)]
+void SetupCS(
+	uint2 GroupId : SV_GroupID,
+	uint GroupThreadIndex : SV_GroupIndex,
+	uint2 DispatchThreadId : SV_DispatchThreadID)
+{
+	FDOFSetupParameters Parameters;
+	Parameters.bOutputFullRes = false;
+	Parameters.bOutputHalfRes = true;
 	
-		#if 0 // Lower VGPR footprint.
-			uint2 OutputPixelPosition = 0.5 * SceneBufferUV * View.BufferSizeAndInvSize.xy;
-		#else
-			uint2 OutputPixelPosition = (View.ViewRectMin.xy/2) + DispatchThreadId;
-		#endif
+	Parameters.SceneColorTexture = SceneColorTexture;
+	Parameters.SceneDepthTexture = SceneDepthTexture;
+
+	FDOFSetupOutput Output = SetupDOF(Parameters, GroupId, GroupThreadIndex, DispatchThreadId);
 
-		if (all(OutputPixelPosition < ViewportRect.zw))
+	if (Output.bHalfResIsValid)
+	{
+		#if CONFIG_DOF_ALPHA
 		{
-			#if CONFIG_DOF_ALPHA
-			{
-				Output1[OutputPixelPosition] = OutColor;
-				Output2[OutputPixelPosition] = OutCocRadius;
-			}
-			#else
-			{
-				Output1[OutputPixelPosition] = float4(OutColor.rgb, OutCocRadius);
-			}
-			#endif
+			Output1[Output.HalfResOutputPixelPosition] = Output.HalfResColor;
+			Output2[Output.HalfResOutputPixelPosition] = Output.HalfResCocRadius;
 		}
+		#else
+		{
+			Output1[Output.HalfResOutputPixelPosition] = float4(Output.HalfResColor.rgb, Output.HalfResCocRadius);
+		}
+		#endif
 	}
 }
 
diff --git a/Engine/Source/Runtime/Renderer/Private/PostProcess/DiaphragmDOF.cpp b/Engine/Source/Runtime/Renderer/Private/PostProcess/DiaphragmDOF.cpp
--- a/Engine/Source/Runtime/Renderer/Private/PostProcess/DiaphragmDOF.cpp
+++ b/Engine/Source/Runtime/Renderer/Private/PostProcess/DiaphragmDOF.cpp
@@ -470,6 +470,10 @@ public:
 		{
 			OutEnvironment.CompilerFlags.Add(CFLAG_AllowRealTypes);
 		}
+		if (FDataDrivenShaderPlatformInfo::GetSupportsWaveOperations(Parameters.Platform) != ERHIFeatureSupport::Unsupported)
+		{
+			OutEnvironment.SetDefine(TEXT("WAVE_OPS"), 1);
+		}
 
 		OutEnvironment.SetDefine(TEXT("COC_TILE_SIZE"), kCocTileSize);
 
@@ -790,7 +794,6 @@ class FDiaphragmDOFSetupCS : public FDiaphragmDOFShader
 		SHADER_PARAMETER_STRUCT_INCLUDE(FDOFCommonShaderParameters, CommonParameters)
 		SHADER_PARAMETER_STRUCT_INCLUDE(FDOFCocModelShaderParameters, CocModel)
 
-		SHADER_PARAMETER(FIntRect, ViewportRect)
 		SHADER_PARAMETER(FVector2f, CocRadiusBasis) // TODO: decompose
 
 		SHADER_PARAMETER_RDG_TEXTURE(Texture2D, SceneColorTexture)
@@ -1707,11 +1710,9 @@ bool DiaphragmDOF::AddPasses(
 		PermutationVector.Set<FDiaphragmDOFShader::FAlphaChannelDim>(bProcessSceneAlpha);
 	
 		FIntPoint PassViewSize = FullResViewSize;
-		FIntPoint GroupSize(kDefaultGroupSize, kDefaultGroupSize);
 		if (bOutputFullResolution && bOutputHalfResolution)
 		{
 			PermutationVector.Set<FDiaphragmDOFSetupCS::FOutputResDivisor>(0);
-			GroupSize *= 2;
 		}
 		else if (bOutputFullResolution)
 		{
@@ -1731,7 +1732,6 @@ bool DiaphragmDOF::AddPasses(
 		{
 			PassParameters->CommonParameters = CommonParameters;
 			SetCocModelParameters(GraphBuilder, &PassParameters->CocModel, CocModel, EncodedCocRadiusBasis);
-			PassParameters->ViewportRect = FIntRect(FIntPoint::ZeroValue, PassViewSize);
 			PassParameters->SceneColorTexture = InputSceneColor;
 			PassParameters->SceneDepthTexture = SceneTextures.SceneDepthTexture;
 		
@@ -1757,7 +1757,7 @@ bool DiaphragmDOF::AddPasses(
 			}
 		}
 
-		FIntVector GroupCount = FComputeShaderUtils::GetGroupCount(PassViewSize, GroupSize);
+		FIntVector GroupCount = FComputeShaderUtils::GetGroupCount(PassViewSize, kDefaultGroupSize);
 
 		TShaderMapRef<FDiaphragmDOFSetupCS> ComputeShader(View.ShaderMap, PermutationVector);
 		FComputeShaderUtils::AddPass(
