Fix shader compilation error under SM5 due to typo in DOFSetup.

diff --git a/Engine/Shaders/Private/DiaphragmDOF/DOFSetup.usf b/Engine/Shaders/Private/DiaphragmDOF/DOFSetup.usf
--- a/Engine/Shaders/Private/DiaphragmDOF/DOFSetup.usf
+++ b/Engine/Shaders/Private/DiaphragmDOF/DOFSetup.usf
@@ -184,8 +184,8 @@ struct FSoftwareQuad
 		// Map threads to pixels using morton-like curve for better texture cache hit rate
 		GroupLaneIndex = GroupThreadIndex;
 		LocalPixelPosition = (GroupSize * GroupId) + MortonLikeDecode(GroupThreadIndex);
-		QuadLeaderLane = WaveGetLaneIndex() & (~(4 - 1));
-		bIsLeaderLane = (WaveGetLaneIndex() & (4 - 1)) == 0;
+		QuadLeaderLane = GroupThreadIndex & (~(4 - 1));
+		bIsLeaderLane = (GroupThreadIndex & (4 - 1)) == 0;
 	}
 
 	bool IsLeader()
