Fixed: Incorrect 'Pre-Skinned Local Bounds' in a mobile specific GPUScene implementation

diff --git a/Engine/Shaders/Private/SceneDataMobileLoader.ush b/Engine/Shaders/Private/SceneDataMobileLoader.ush
--- a/Engine/Shaders/Private/SceneDataMobileLoader.ush
+++ b/Engine/Shaders/Private/SceneDataMobileLoader.ush
@@ -112,12 +112,6 @@ FPrimitiveSceneData LoadPrimitiveDataUBO(uint DrawInstanceId)
 		PrimitiveData.ObjectBoundsZ = Data.w;
 		DataIdx += 1u;
 	}
-	
-	// PreSkinnedLocalBounds, TODO: this works until we actually support GPUScene for skinned primitives
-	{
-		PrimitiveData.PreSkinnedLocalBoundsMin = PrimitiveData.LocalObjectBoundsMin;
-		PrimitiveData.PreSkinnedLocalBoundsMax = PrimitiveData.LocalObjectBoundsMax;
-	}
 
 	// InstanceLocalBounds
 	{
@@ -164,6 +158,19 @@ FPrimitiveSceneData LoadPrimitiveDataUBO(uint DrawInstanceId)
 	// Should not be used with regular primitives
 	PrimitiveData.WorldToPreviousWorld = float4x4(1,0,0,0, 0,1,0,0, 0,0,1,0, 0,0,0,1);
 
+	// PreSkinnedLocalBounds, only for skinned meshes
+	if ((PrimitiveData.Flags & PRIMITIVE_SCENE_DATA_FLAG_SKINNED_MESH) != 0u)
+	{
+		PrimitiveData.PreSkinnedLocalBoundsMin = LoadPrimitiveDataElementUBO(DrawInstanceId, DataIdx + 0);
+		PrimitiveData.PreSkinnedLocalBoundsMax = LoadPrimitiveDataElementUBO(DrawInstanceId, DataIdx + 1);
+		DataIdx += 2u;
+	}
+	else
+	{
+		PrimitiveData.PreSkinnedLocalBoundsMin = PrimitiveData.LocalObjectBoundsMin;
+		PrimitiveData.PreSkinnedLocalBoundsMax = PrimitiveData.LocalObjectBoundsMax;
+	}
+
 	// CustomData, TODO make optional
 	{
 		uint NumCustomData = min(NUM_CUSTOM_PRIMITIVE_DATA, BATCHED_PRIMITIVE_DATA_STRIDE_FLOAT4 - DataIdx); //23u
diff --git a/Engine/Shaders/Private/SceneDataMobileWriter.ush b/Engine/Shaders/Private/SceneDataMobileWriter.ush
--- a/Engine/Shaders/Private/SceneDataMobileWriter.ush
+++ b/Engine/Shaders/Private/SceneDataMobileWriter.ush
@@ -139,6 +139,14 @@ void WritePrimitiveDataUBO(RWStructuredBuffer<float4> Output, uint Offset, FPrim
 		DataIdx += 3u;
 	}
 
+	// PreSkinnedLocalBounds, only for skinned meshes
+	if ((PrimitiveData.Flags & PRIMITIVE_SCENE_DATA_FLAG_SKINNED_MESH) != 0u)
+	{
+		WriteDataElementUBO(Output, Offset, DataIdx + 0, float4(PrimitiveData.PreSkinnedLocalBoundsMin, 0.f));
+		WriteDataElementUBO(Output, Offset, DataIdx + 1, float4(PrimitiveData.PreSkinnedLocalBoundsMax, 0.f));
+		DataIdx += 2u;
+	}
+
 	// CustomData, TODO make optional
 	{
 		uint NumCustomData = min(NUM_CUSTOM_PRIMITIVE_DATA, BATCHED_PRIMITIVE_DATA_STRIDE_FLOAT4 - DataIdx); //23u
