ACES 2.0: Fix previously missed permutation support in the device-encoding pass.

diff --git a/Engine/Shaders/Private/PostProcessDeviceEncodingOnly.usf b/Engine/Shaders/Private/PostProcessDeviceEncodingOnly.usf
--- a/Engine/Shaders/Private/PostProcessDeviceEncodingOnly.usf
+++ b/Engine/Shaders/Private/PostProcessDeviceEncodingOnly.usf
@@ -105,7 +105,7 @@ float4 DeviceEncodingOnlyCommonPS(
 	else if (GetOutputDevice() == TONEMAPPER_OUTPUT_ACES1000nitST2084 || GetOutputDevice() == TONEMAPPER_OUTPUT_ACES2000nitST2084)
 	{
 #if USE_ACES_2
-		float3 ODTColor = ACESOutputTransform(GradedColor * ACESSceneColorMultiplier, (float3x3) WorkingColorSpace.ToAP0, OutputMaxLuminance, ACESReachTable, ACESGamutTable, ACESGammaTable);
+		float3 ODTColor = ACESOutputTransform( FilmColor * ACESSceneColorMultiplier, (float3x3) WorkingColorSpace.ToAP0, OutputMaxLuminance, ACESReachTable, ACESGamutTable, ACESGammaTable);
 #else
 		// ACES ODT
 		FACESTonemapParams AcesParams = ComputeACESTonemapParams(ACESMinMaxData, ACESMidData, ACESCoefsLow_0, ACESCoefsHigh_0, ACESCoefsLow_4, ACESCoefsHigh_4, ACESSceneColorMultiplier, ACESGamutCompression);
@@ -122,7 +122,7 @@ float4 DeviceEncodingOnlyCommonPS(
 	else if (GetOutputDevice() == TONEMAPPER_OUTPUT_ACES1000nitScRGB || GetOutputDevice() == TONEMAPPER_OUTPUT_ACES2000nitScRGB)
 	{
 #if USE_ACES_2
-		float3 ODTColor = ACESOutputTransform(GradedColor * ACESSceneColorMultiplier, (float3x3) WorkingColorSpace.ToAP0, OutputMaxLuminance, ACESReachTable, ACESGamutTable, ACESGammaTable);
+		float3 ODTColor = ACESOutputTransform(FilmColor * ACESSceneColorMultiplier, (float3x3) WorkingColorSpace.ToAP0, OutputMaxLuminance, ACESReachTable, ACESGamutTable, ACESGammaTable);
 #else
 		FACESTonemapParams AcesParams = ComputeACESTonemapParams(ACESMinMaxData, ACESMidData, ACESCoefsLow_0, ACESCoefsHigh_0, ACESCoefsLow_4, ACESCoefsHigh_4, ACESSceneColorMultiplier, ACESGamutCompression);
 
diff --git a/Engine/Source/Runtime/Renderer/Private/PostProcess/PostProcessDeviceEncodingOnly.cpp b/Engine/Source/Runtime/Renderer/Private/PostProcess/PostProcessDeviceEncodingOnly.cpp
--- a/Engine/Source/Runtime/Renderer/Private/PostProcess/PostProcessDeviceEncodingOnly.cpp
+++ b/Engine/Source/Runtime/Renderer/Private/PostProcess/PostProcessDeviceEncodingOnly.cpp
@@ -28,8 +28,9 @@ namespace
 	{
 		// Desktop renderer permutation dimensions.
 		class FDeviceEncodingOnlyOutputDeviceSRGB : SHADER_PERMUTATION_BOOL("OUTPUT_DEVICE_SRGB");
+		class FDeviceEncodingUseACES2 : SHADER_PERMUTATION_BOOL("USE_ACES_2");
 
-		using FDesktopDomain = TShaderPermutationDomain<FDeviceEncodingOnlyOutputDeviceSRGB>;
+		using FDesktopDomain = TShaderPermutationDomain<FDeviceEncodingOnlyOutputDeviceSRGB, FDeviceEncodingUseACES2>;
 	
 	} // namespace DeviceEncodingOnlyPermutation
 } // namespace
@@ -241,6 +242,9 @@ FScreenPassTexture AddDeviceEncodingOnlyPass(FRDGBuilder& GraphBuilder, const FV
 	const bool bOutputDeviceSRGB = (CommonParameters.OutputDevice.OutputDevice == (uint32)EDisplayOutputFormat::SDR_sRGB);
 	DesktopPermutationVector.Set<DeviceEncodingOnlyPermutation::FDeviceEncodingOnlyOutputDeviceSRGB>(bOutputDeviceSRGB);
 
+	static const auto CVarAcesVersion = IConsoleManager::Get().FindTConsoleVariableDataInt(TEXT("r.HDR.Aces.Version"));
+	DesktopPermutationVector.Set<DeviceEncodingOnlyPermutation::FDeviceEncodingUseACES2>(CVarAcesVersion->GetValueOnRenderThread() > 1);
+
 	// Override output might not support UAVs.
 	const bool bComputePass = (Output.Texture->Desc.Flags & TexCreate_UAV) == TexCreate_UAV ? View.bUseComputePasses : false;
 
