ACES 2.0: Expose transform inverse.

* We additionally make the white_limiting function optional: its clamp can hinder inversion of bright colors when going from scene color to output transform color back to scene color.
* Add working color space to AP0 matrix to scene uniform buffer for inversion callees.

diff --git a/Engine/Shaders/Private/ACES/ACESOutputTransform.ush b/Engine/Shaders/Private/ACES/ACESOutputTransform.ush
--- a/Engine/Shaders/Private/ACES/ACESOutputTransform.ush
+++ b/Engine/Shaders/Private/ACES/ACESOutputTransform.ush
@@ -322,7 +322,7 @@ float3 aces_to_JMh(float3 aces,
     float3 XYZ = clamp_AP1(AP1, peakLuminance);
 
     // XYZ to JMh
-    float3 RGB_w = float3(referenceLuminance, referenceLuminance, referenceLuminance);
+    const float3 RGB_w = referenceLuminance;
     float3 XYZ_w = mul(AP0_2_XYZ_MAT, RGB_w);
 
     float3 XYZluminance = referenceLuminance * XYZ;
@@ -331,6 +331,23 @@ float3 aces_to_JMh(float3 aces,
     return JMh;
 }
 
+float3 JMh_to_aces(float3 JMh,
+                     float peakLuminance)
+{
+    const float3 RGB_w = referenceLuminance;
+    float3 XYZ_w_aces = mul(AP0_2_XYZ_MAT, RGB_w);
+
+    // JMh to XYZ
+    float3 XYZluminance = Hellwig2022_JMh_to_XYZ(JMh, XYZ_w_aces);
+
+    float3 XYZ = (1. / referenceLuminance) * XYZluminance;
+
+    // XYZ to ACES
+    float3 ACES = mul(XYZ_2_AP0_MAT, XYZ);
+
+    return ACES;
+}
+
 ODTParams init_ODTParams(
     float peakLuminance,
     float3x3 LIMIT_RGB_TO_XYZ,
@@ -671,6 +688,29 @@ float3 tonemapAndCompress_fwd(float3 inputJMh,
     return outputJMh;
 }
 
+float3 tonemapAndCompress_inv(float3 JMh,
+                                ODTParams PARAMS)
+{
+    float3 tonemappedJMh = JMh;
+
+    float luminance = Hellwig_J_to_Y(JMh[0]);
+
+    TSParams TSPARAMS = init_TSParams(PARAMS.peakLuminance);
+    float linear_c = tonescale_inv(luminance / referenceLuminance,
+                                 TSPARAMS);
+
+    float untonemappedJ = Y_to_Hellwig_J(linear_c * referenceLuminance);
+    float3 untonemappedColorJMh = float3(untonemappedJ, tonemappedJMh[1], tonemappedJMh[2]);
+
+    // Chroma compression
+    untonemappedColorJMh[1] = chromaCompression(tonemappedJMh,
+                                                untonemappedColorJMh[0],
+                                                PARAMS,
+                                                true);
+
+    return untonemappedColorJMh;
+}
+
 float3 JMh_to_output_XYZ(float3 JMh,
                            ODTParams PARAMS)
 {
@@ -682,6 +722,17 @@ float3 JMh_to_output_XYZ(float3 JMh,
     return XYZ;
 }
 
+float3 XYZ_output_to_JMh(float3 XYZ,
+                           ODTParams PARAMS)
+{
+    float3 XYZluminance = referenceLuminance * XYZ;
+
+    float3 JMh = XYZ_to_Hellwig2022_JMh(XYZluminance,
+                                          PARAMS.XYZ_w_limit);
+
+    return JMh;
+}
+
 float hueDependentUpperHullGamma(float h,
                                  Texture2D<float> gamma_table)
 {
@@ -840,6 +891,7 @@ float3 getReachBoundary(float J,
 
 float3 compressGamut(float3 JMh,
                        ODTParams PARAMS,
+                       float Jx,
                        bool _invert = false)
 {
     float limitJmax = PARAMS.limitJmax;
@@ -859,7 +911,7 @@ float3 compressGamut(float3 JMh,
     // Calculate where the out of gamut color is projected to
     float focusJ = lerp(JMcusp[0], midJ, min(1., cuspMidBlend - (JMcusp[0] / limitJmax)));
 
-    float slope_gain = limitJmax * focusDist * getFocusGain(JMh[0], JMcusp[0], limitJmax);
+    float slope_gain = limitJmax * focusDist * getFocusGain(Jx, JMcusp[0], limitJmax);
 
     // Find gamut intersection
     float gamma_top = hueDependentUpperHullGamma(JMh[2], PARAMS.gamma_table);
@@ -906,7 +958,31 @@ float3 gamutMap_fwd(float3 JMh,
                       ODTParams PARAMS
                     )
 {
-    return compressGamut(JMh, PARAMS, false);
+    return compressGamut(JMh, PARAMS, JMh[0], false);
+}
+
+float3 gamutMap_inv(float3 JMh,
+                      ODTParams PARAMS)
+{
+    float2 JMcusp = cuspFromTable(JMh[2], PARAMS.gamut_table);
+    float Jx = JMh[0];
+
+    // Analytic inverse below threshold
+    if (Jx <= lerp(JMcusp[0], PARAMS.limitJmax, focusGainBlend))
+        return compressGamut(JMh,
+                             PARAMS,
+                             Jx,
+                             true);
+
+    // Approximation above threshold
+    Jx = compressGamut(JMh,
+                       PARAMS,
+                       Jx,
+                       true)[0];
+    return compressGamut(JMh,
+                         PARAMS,
+                         Jx,
+                         true);
 }
 
 float3 outputTransform_fwd( float3 aces,
@@ -923,4 +999,23 @@ float3 outputTransform_fwd( float3 aces,
                                       PARAMS);
 
     return XYZ;
+}
+
+float3 outputTransform_inv( float3 XYZ,
+                            float peakLuminance,
+                            ODTParams PARAMS)
+{
+    float3 compressedJMh = XYZ_output_to_JMh(XYZ,
+                                               PARAMS);
+
+    float3 tonemappedJMh = gamutMap_inv(compressedJMh,
+                                          PARAMS);
+
+    float3 JMh = tonemapAndCompress_inv(tonemappedJMh,
+                                          PARAMS);
+
+    float3 aces = JMh_to_aces(JMh,
+                                peakLuminance);
+
+    return aces;
 }
\ No newline at end of file
diff --git a/Engine/Shaders/Private/TonemapCommon.ush b/Engine/Shaders/Private/TonemapCommon.ush
--- a/Engine/Shaders/Private/TonemapCommon.ush
+++ b/Engine/Shaders/Private/TonemapCommon.ush
@@ -316,6 +316,8 @@ static const int UE_ACES_SURROUND = 0;      // 0: dark ( NOTE: this is the only
 static const bool UE_ACES_STRETCH_BLACK = true; // stretch black luminance to a PQ code value of 0
 static const bool UE_ACES_D60_SIM = false;
 static const bool UE_ACES_LEGAL_RANGE = false;
+// Scale white is usually only enabled when using a limiting white different from the encoding white
+static const bool UE_ACES_SCALE_WHITE = false;
 
 static const float UE_SCRGB_WHITE_NITS = 80.0; // ScRGB D65 white luminance (cd/m^2)
 
@@ -326,7 +328,8 @@ float3 ACESOutputTransform(
 	float PeakLuminance,
 	Texture2D<float> ReachTable,
     Texture2D<float3> GamutTable,
-    Texture2D<float> GammaTable
+    Texture2D<float> GammaTable,
+	const bool bApplyWhiteLimiting = true
 )
 {
 	// TODO: The display limiting color space should be user-exposed as part of display characterization.
@@ -348,13 +351,51 @@ float3 ACESOutputTransform(
 
 	float3 ColorXYZ = outputTransform_fwd(Aces, PeakLuminance, PARAMS);
 
-	// Scale white is usually only enabled when using a limiting white different from the encoding white
-	static const bool scale_white = false;
-	ColorXYZ = white_limiting(ColorXYZ, PARAMS, scale_white);
+	if(bApplyWhiteLimiting)
+	{
+		ColorXYZ = white_limiting(ColorXYZ, PARAMS, UE_ACES_SCALE_WHITE);
+	}
 
 	// Scale by reference luminance to return values in nits, matching our 1.3 transform
 	return mul(XYZ_2_AP1_MAT, ColorXYZ) * referenceLuminance;
 }
+
+float3 ACESInvOutputTransform(
+	float3 ColorAP1_Nits,
+	const float3x3 AP0_2_WCS,
+	float PeakLuminance,
+	Texture2D<float> ReachTable,
+    Texture2D<float3> GamutTable,
+    Texture2D<float> GammaTable,
+	const bool bApplyWhiteLimiting = true
+)
+{
+	float3 ColorXYZ = mul(AP1_2_XYZ_MAT, ColorAP1_Nits / referenceLuminance);
+
+	const FColorSpace LimitingColorSpace = GetColorSpace(TONEMAPPER_GAMUT_ACEScg_D60);
+	const FColorSpace EncodingColorSpace = GetColorSpace(TONEMAPPER_GAMUT_ACEScg_D60);
+	
+	ODTParams PARAMS = init_ODTParams(
+		PeakLuminance,
+		LimitingColorSpace.RGBtoXYZ,
+		LimitingColorSpace.XYZtoRGB,
+		EncodingColorSpace.RGBtoXYZ,
+		EncodingColorSpace.XYZtoRGB,
+		ReachTable,
+		GamutTable,
+		GammaTable
+	);
+
+	if(bApplyWhiteLimiting)
+	{
+		ColorXYZ = white_limiting(ColorXYZ, PARAMS, UE_ACES_SCALE_WHITE, true);
+	}
+
+    float3 Aces = outputTransform_inv(ColorXYZ, PeakLuminance, PARAMS);
+
+	return mul(AP0_2_WCS, Aces);
+}
+
 #else // USE_ACES_2
 struct FACESTonemapParams
 {
diff --git a/Engine/Source/Runtime/Engine/Private/SceneManagement.cpp b/Engine/Source/Runtime/Engine/Private/SceneManagement.cpp
--- a/Engine/Source/Runtime/Engine/Private/SceneManagement.cpp
+++ b/Engine/Source/Runtime/Engine/Private/SceneManagement.cpp
@@ -146,6 +146,7 @@ void FDefaultWorkingColorSpaceUniformBuffer::Update(FRHICommandListBase& RHICmdL
 	Parameters.FromAP1 = Parameters.ToAP1.Inverse();
 	
 	Parameters.ToAP0 = Transpose<float>(FColorSpaceTransform(InColorSpace, FColorSpace(EColorSpace::ACESAP0)));
+	Parameters.FromAP0 = Parameters.ToAP0.Inverse();
 
 	Parameters.bIsSRGB = InColorSpace.IsSRGB();
 
diff --git a/Engine/Source/Runtime/Engine/Public/SceneManagement.h b/Engine/Source/Runtime/Engine/Public/SceneManagement.h
--- a/Engine/Source/Runtime/Engine/Public/SceneManagement.h
+++ b/Engine/Source/Runtime/Engine/Public/SceneManagement.h
@@ -296,6 +296,7 @@ BEGIN_GLOBAL_SHADER_PARAMETER_STRUCT(FWorkingColorSpaceShaderParameters, ENGINE_
 	SHADER_PARAMETER(FMatrix44f, ToAP1)
 	SHADER_PARAMETER(FMatrix44f, FromAP1)
 	SHADER_PARAMETER(FMatrix44f, ToAP0)
+	SHADER_PARAMETER(FMatrix44f, FromAP0)
 	SHADER_PARAMETER(uint32, bIsSRGB)
 END_SHADER_PARAMETER_STRUCT()
 
