Make virtual texture residency mip bias a per group value instead of a single global value.
Virtual texture pool configs could already opt into tracking residency to generate a global mip bias, but now they can specify one of 4 independent group mip biases.
Each group accumulates the highest requested bias from its associated physical pools. All 4 are then passed to the shader.
When sampling a virtual texture stack we apply its associated mip bias.
This allows us to partition pools with different subscription behaviors so that they only bias themselves and not _all_ virtual textures.

diff --git a/Engine/Shaders/Private/VirtualTextureCommon.ush b/Engine/Shaders/Private/VirtualTextureCommon.ush
--- a/Engine/Shaders/Private/VirtualTextureCommon.ush
+++ b/Engine/Shaders/Private/VirtualTextureCommon.ush
@@ -149,7 +149,8 @@ struct VTPageTableUniform
 	uint XOffsetInPages; // 12
 	uint YOffsetInPages; // 12
 	uint MaxLevel; // 4
-	uint vPageTableMipBias; // 8
+	uint vPageTableMipBias; // 4
+	uint MipBiasGroup; // 2
 	uint ShiftedPageTableID; // 4
 	uint AdaptiveLevelBias; // 4
 
@@ -167,7 +168,8 @@ VTPageTableUniform VTPageTableUniform_Unpack(uint4 PackedPageTableUniform0, uint
 	result.MaxAnisoLog2 = asfloat(PackedPageTableUniform1.x);
 	result.XOffsetInPages = PackedPageTableUniform1.y & 0xfff;
 	result.YOffsetInPages = (PackedPageTableUniform1.y >> 12) & 0xfff;
-	result.vPageTableMipBias = (PackedPageTableUniform1.y >> 24) & 0xff;
+	result.vPageTableMipBias = (PackedPageTableUniform1.y >> 24) & 0xf;
+	result.MipBiasGroup = (PackedPageTableUniform1.y >> 28) & 0x3;
 	result.MaxLevel = PackedPageTableUniform1.z & 0xf;
 	result.AdaptiveLevelBias = (PackedPageTableUniform1.z >> 4) & 0xf;
 	result.ShiftedPageTableID = PackedPageTableUniform1.w;
@@ -184,6 +186,7 @@ VTPageTableUniform VTPageTableUniform_Unpack(uint2 PackedPageTableUniform)
 	Result.XOffsetInPages = PackedPageTableUniform.x & 0xfff;
 	Result.YOffsetInPages = (PackedPageTableUniform.x >> 12) & 0xfff;
 	Result.vPageTableMipBias = (PackedPageTableUniform.x >> 24) & 0xf;
+	Result.MipBiasGroup = 0;
 	Result.MaxLevel = (PackedPageTableUniform.y >> 24) & 0xf;
 	Result.AdaptiveLevelBias = 0;
 	Result.ShiftedPageTableID = PackedPageTableUniform.x & 0xf0000000;
@@ -234,12 +237,12 @@ float GetStochasticMipNoise(float2 SvPositionXY)
 #endif
 }
 
-float GetGlobalVirtualTextureMipBias()
+float GetGlobalVirtualTextureMipBias(int Group)
 {
 #if VT_DISABLE_VIEW_UNIFORM_BUFFER
 	return 0;
 #else
-	return View.GlobalVirtualTextureMipBias;
+	return View.GlobalVirtualTextureMipBias[Group];
 #endif
 }
 
@@ -257,7 +260,7 @@ int TextureComputeVirtualMipLevel(
 	// This way the VT mip level selection will come much closer to HW mip selection, even if we're not sampling the texture using anisotropic filtering
 	const float ComputedLevel = MipLevelAniso2D(OutResult.dUVdx, OutResult.dUVdy, PageTableUniform.MaxAnisoLog2);
 
-	const float GlobalMipBias = GetGlobalVirtualTextureMipBias();
+	const float GlobalMipBias = GetGlobalVirtualTextureMipBias(PageTableUniform.MipBiasGroup);
 #if VIRTUAL_TEXTURE_MANUAL_TRILINEAR_FILTERING
 	const float Noise = 0.f;
 #elif VIRTUAL_TEXTURE_FORCE_BILINEAR_FILTERING
@@ -428,7 +431,7 @@ VTPageTableResult TextureLoadVirtualPageTable(
 	if (IsValid(PageTableUniform))
 	{
 		UV = UV * PageTableUniform.UVScale;
-		int vLevel = GetGlobalVirtualTextureMipBias();
+		int vLevel = GetGlobalVirtualTextureMipBias(PageTableUniform.MipBiasGroup);
 #if PIXELSHADER
 		vLevel = TextureComputeVirtualMipLevel(Result, ddx(UV), ddy(UV), MipBias, SvPositionXY, PageTableUniform);
 #endif // PIXELSHADER
@@ -450,7 +453,7 @@ VTPageTableResult TextureLoadVirtualPageTable(
 	if (IsValid(PageTableUniform))
 	{
 		UV = UV * PageTableUniform.UVScale;
-		int vLevel = GetGlobalVirtualTextureMipBias();
+		int vLevel = GetGlobalVirtualTextureMipBias(PageTableUniform.MipBiasGroup);
 #if PIXELSHADER
 		vLevel = TextureComputeVirtualMipLevel(Result, ddx(UV), ddy(UV), MipBias, SvPositionXY, PageTableUniform);
 #endif // PIXELSHADER
@@ -473,7 +476,7 @@ VTPageTableResult TextureLoadVirtualPageTable(
 	if (IsValid(PageTableUniform))
 	{
 		UV = UV * PageTableUniform.UVScale;
-		int vLevel = GetGlobalVirtualTextureMipBias();
+		int vLevel = GetGlobalVirtualTextureMipBias(PageTableUniform.MipBiasGroup);
 #if PIXELSHADER
 		vLevel = TextureComputeVirtualMipLevel(Result, ddx(UV), ddy(UV), MipBias, SvPositionXY, PageTableUniform);
 #endif // PIXELSHADER
@@ -495,7 +498,7 @@ VTPageTableResult TextureLoadVirtualPageTable(
 	if (IsValid(PageTableUniform))
 	{
 		UV = UV * PageTableUniform.UVScale;
-		int vLevel = GetGlobalVirtualTextureMipBias();
+		int vLevel = GetGlobalVirtualTextureMipBias(PageTableUniform.MipBiasGroup);
 #if PIXELSHADER
 		vLevel = TextureComputeVirtualMipLevel(Result, ddx(UV), ddy(UV), MipBias, SvPositionXY, PageTableUniform);
 #endif // PIXELSHADER
@@ -596,7 +599,7 @@ VTPageTableResult TextureLoadVirtualPageTableLevel(
 	if (IsValid(PageTableUniform))
 	{
 		UV = UV * PageTableUniform.UVScale;
-		int vLevel = (int)floor(Level + GetGlobalVirtualTextureMipBias());
+		int vLevel = (int)floor(Level + GetGlobalVirtualTextureMipBias(PageTableUniform.MipBiasGroup));
 		UV = ApplyAddressMode(UV, AddressU, AddressV);
 		TextureLoadVirtualPageTableInternal(Result, PageTable0, PageTableUniform, UV, vLevel);
 		StoreVirtualTextureFeedback(Feedback, SampleIndex, Result.PackedRequest, Result.PendingMips);
@@ -615,7 +618,7 @@ VTPageTableResult TextureLoadVirtualPageTableLevel(
 	if (IsValid(PageTableUniform))
 	{
 		UV = UV * PageTableUniform.UVScale;
-		int vLevel = (int)floor(Level + GetGlobalVirtualTextureMipBias());
+		int vLevel = (int)floor(Level + GetGlobalVirtualTextureMipBias(PageTableUniform.MipBiasGroup));
 		UV = ApplyAddressMode(UV, AddressU, AddressV);
 		TextureLoadVirtualPageTableInternal(Result, PageTable0, PageTableUniform, UV, vLevel);
 	}
@@ -635,7 +638,7 @@ VTPageTableResult TextureLoadVirtualPageTableLevel(
 	if (IsValid(PageTableUniform))
 	{
 		UV = UV * PageTableUniform.UVScale;
-		int vLevel = (int)floor(Level + GetGlobalVirtualTextureMipBias());
+		int vLevel = (int)floor(Level + GetGlobalVirtualTextureMipBias(PageTableUniform.MipBiasGroup));
 		UV = ApplyAddressMode(UV, AddressU, AddressV);
 		TextureLoadVirtualPageTableInternal(Result, PageTable0, PageTable1, PageTableUniform, UV, vLevel);
 		StoreVirtualTextureFeedback(Feedback, SampleIndex, Result.PackedRequest, Result.PendingMips);
@@ -652,7 +655,7 @@ VTPageTableResult TextureLoadVirtualPageTableLevel(
 {
 	VTPageTableResult Result = (VTPageTableResult)0;
 	UV = UV * PageTableUniform.UVScale;
-	int vLevel = (int)floor(Level + GetGlobalVirtualTextureMipBias());
+	int vLevel = (int)floor(Level + GetGlobalVirtualTextureMipBias(PageTableUniform.MipBiasGroup));
 	UV = ApplyAddressMode(UV, AddressU, AddressV);
 	TextureLoadVirtualPageTableInternal(Result, PageTable0, PageTable1, PageTableUniform, UV, vLevel);
 	return Result;
@@ -670,7 +673,7 @@ VTPageTableResult TextureLoadVirtualPageTableAdaptive(
 {
 	VTPageTableResult Result = (VTPageTableResult)0;
 	UV = UV * PageTableUniform.UVScale;
-	int vLevel = GetGlobalVirtualTextureMipBias();
+	int vLevel = GetGlobalVirtualTextureMipBias(PageTableUniform.MipBiasGroup);
 #if PIXELSHADER
 	vLevel = TextureComputeVirtualMipLevel(Result, ddx(UV), ddy(UV), MipBias, SvPositionXY, PageTableUniform);
 #endif // PIXELSHADER
@@ -713,7 +716,7 @@ VTPageTableResult TextureLoadVirtualPageTableAdaptiveLevel(
 {
 	VTPageTableResult Result = (VTPageTableResult)0;
 	// Level is an index into the full size adaptive VT. AdaptiveLevelBias shifts it relative to the low mips allocated VT.
-	int vLevel = (int)floor(Level + GetGlobalVirtualTextureMipBias()) - PageTableUniform.AdaptiveLevelBias;
+	int vLevel = (int)floor(Level + GetGlobalVirtualTextureMipBias(PageTableUniform.MipBiasGroup)) - PageTableUniform.AdaptiveLevelBias;
 	UV = UV * PageTableUniform.UVScale;
 	UV = ApplyAddressMode(UV, AddressU, AddressV);
 	ApplyAdaptivePageTableUniform(PageTableIndirection, Result, PageTableUniform, UV, vLevel);
@@ -732,7 +735,7 @@ VTPageTableResult TextureLoadVirtualPageTableAdaptiveLevel(
 {
 	VTPageTableResult Result = (VTPageTableResult)0;
 	// Level is an index into the full size adaptive VT. AdaptiveLevelBias shifts it relative to the low mips allocated VT.
-	int vLevel = (int)floor(Level + GetGlobalVirtualTextureMipBias()) - PageTableUniform.AdaptiveLevelBias;
+	int vLevel = (int)floor(Level + GetGlobalVirtualTextureMipBias(PageTableUniform.MipBiasGroup)) - PageTableUniform.AdaptiveLevelBias;
 	UV = UV * PageTableUniform.UVScale;
 	UV = ApplyAddressMode(UV, AddressU, AddressV);
 	ApplyAdaptivePageTableUniform(PageTableIndirection, Result, PageTableUniform, UV, vLevel);
diff --git a/Engine/Source/Runtime/Engine/Classes/VT/VirtualTexturePoolConfig.h b/Engine/Source/Runtime/Engine/Classes/VT/VirtualTexturePoolConfig.h
--- a/Engine/Source/Runtime/Engine/Classes/VT/VirtualTexturePoolConfig.h
+++ b/Engine/Source/Runtime/Engine/Classes/VT/VirtualTexturePoolConfig.h
@@ -33,6 +33,10 @@ struct FVirtualTextureSpacePoolConfig
 	UPROPERTY(EditAnywhere, Category = PoolConfig)
 	bool bEnableResidencyMipMapBias = false;
 
+	/** Group can be in range 0-3. Pools with the same group share a mip map bias. */
+	UPROPERTY(EditAnywhere, Category = PoolConfig, meta = (ClampMin = "0", ClampMax = "3", UIMin = "0", UIMax = "3"))
+	int32 ResidencyMipMapBiasGroup = 0;
+
 	/** Allow the size to allocate for the pool to be scaled by scalability settings. */
 	UPROPERTY(EditAnywhere, Category = "PoolConfig|Scalabiity")
 	bool bAllowSizeScale = false;
diff --git a/Engine/Source/Runtime/Engine/Private/VT/VirtualTexturePoolConfig.cpp b/Engine/Source/Runtime/Engine/Private/VT/VirtualTexturePoolConfig.cpp
--- a/Engine/Source/Runtime/Engine/Private/VT/VirtualTexturePoolConfig.cpp
+++ b/Engine/Source/Runtime/Engine/Private/VT/VirtualTexturePoolConfig.cpp
@@ -97,6 +97,7 @@ void UVirtualTexturePoolConfig::FindPoolConfig(TEnumAsByte<EPixelFormat> const*
 	FVirtualTextureSpacePoolConfig DefaultConfig;
 	DefaultConfig.SizeInMegabyte = DefaultSizeInMegabyte;
 	DefaultConfig.bEnableResidencyMipMapBias = true;
+	DefaultConfig.ResidencyMipMapBiasGroup = 0;
 	DefaultConfig.bAllowSizeScale = true;
 
 	bool bFoundDefaultConfig = false;
diff --git a/Engine/Source/Runtime/Engine/Public/SceneView.h b/Engine/Source/Runtime/Engine/Public/SceneView.h
--- a/Engine/Source/Runtime/Engine/Public/SceneView.h
+++ b/Engine/Source/Runtime/Engine/Public/SceneView.h
@@ -1042,7 +1042,7 @@ enum ETranslucencyVolumeCascade
 	VIEW_UNIFORM_BUFFER_MEMBER(float, IndirectLightingCacheShowFlag) \
 	VIEW_UNIFORM_BUFFER_MEMBER(float, EyeToPixelSpreadAngle) \
 	VIEW_UNIFORM_BUFFER_MEMBER_ARRAY(FVector4f, XRPassthroughCameraUVs, 2) \
-	VIEW_UNIFORM_BUFFER_MEMBER(float, GlobalVirtualTextureMipBias) \
+	VIEW_UNIFORM_BUFFER_MEMBER(FVector4f, GlobalVirtualTextureMipBias) \
 	VIEW_UNIFORM_BUFFER_MEMBER(uint32, VirtualTextureFeedbackBufferSize) \
 	VIEW_UNIFORM_BUFFER_MEMBER(uint32, VirtualTextureFeedbackShift) \
 	VIEW_UNIFORM_BUFFER_MEMBER(uint32, VirtualTextureFeedbackMask) \
diff --git a/Engine/Source/Runtime/Renderer/Private/VT/AllocatedVirtualTexture.cpp b/Engine/Source/Runtime/Renderer/Private/VT/AllocatedVirtualTexture.cpp
--- a/Engine/Source/Runtime/Renderer/Private/VT/AllocatedVirtualTexture.cpp
+++ b/Engine/Source/Runtime/Renderer/Private/VT/AllocatedVirtualTexture.cpp
@@ -458,6 +458,15 @@ FRHIShaderResourceView* FAllocatedVirtualTexture::GetPhysicalTextureSRV(uint32 I
 	return nullptr;
 }
 
+uint32 FAllocatedVirtualTexture::GetResidencyMipMapBiasGroup() const
+{
+	// For correctness we should fill a 4 bit mask of groups from all layers.
+	// But that approach also requires extra shader maths to accumulate the max bias from all 4 groups in the mask.
+	// So for now we are just take the group from layer 0 and pass it up directly.
+	FVirtualTexturePhysicalSpace* PhysicalSpace = GetPhysicalSpaceForPageTableLayer(0);
+	return PhysicalSpace != nullptr ? PhysicalSpace->GetDescriptionExt().ResidencyMipMapBiasGroup : 0;
+}
+
 static inline uint32 BitcastFloatToUInt32(float v)
 {
 	const union
@@ -477,6 +486,7 @@ void FAllocatedVirtualTexture::GetPackedPageTableUniform(FUintVector4* OutUnifor
 	const uint32 WidthInPages = GetWidthInTiles();
 	const uint32 HeightInPages = GetHeightInTiles();
 	const uint32 vPageTableMipBias = FMath::FloorLog2(vPageSize);
+	const uint32 MipBiasGroup = GetResidencyMipMapBiasGroup();
 	const uint32 AdaptiveLevelBias = Description.AdaptiveLevelBias;
 
 	// Here MaxAnisotropy only controls the VT mip level selection
@@ -488,6 +498,7 @@ void FAllocatedVirtualTexture::GetPackedPageTableUniform(FUintVector4* OutUnifor
 	checkSlow(vPageX < 4096u);
 	checkSlow(vPageY < 4096u);
 	checkSlow(vPageTableMipBias < 16u);
+	checkSlow(MipBiasGroup < 4u);
 	checkSlow(MaxLevel < 16u);
 	checkSlow(AdaptiveLevelBias < 16u);
 	checkSlow(SpaceID < 16u);
@@ -499,7 +510,7 @@ void FAllocatedVirtualTexture::GetPackedPageTableUniform(FUintVector4* OutUnifor
 	OutUniform[0].W = BitcastFloatToUInt32((float)HeightInPages);
 
 	OutUniform[1].X = BitcastFloatToUInt32((float)MaxAnisotropyLog2);
-	OutUniform[1].Y = vPageX | (vPageY << 12) | (vPageTableMipBias << 24);
+	OutUniform[1].Y = vPageX | (vPageY << 12) | (vPageTableMipBias << 24) | (MipBiasGroup << 28);
 	OutUniform[1].Z = MaxLevel | (AdaptiveLevelBias << 4);
 	OutUniform[1].W = (SpaceID << 28);
 }
diff --git a/Engine/Source/Runtime/Renderer/Private/VT/AllocatedVirtualTexture.h b/Engine/Source/Runtime/Renderer/Private/VT/AllocatedVirtualTexture.h
--- a/Engine/Source/Runtime/Renderer/Private/VT/AllocatedVirtualTexture.h
+++ b/Engine/Source/Runtime/Renderer/Private/VT/AllocatedVirtualTexture.h
@@ -67,6 +67,7 @@ private:
 	uint32 CalculatePersistentHash(FAllocatedVTDescription const& InDesc, FVirtualTextureProducer* const* InProducers) const;
 	uint32 AddUniqueProducer(FVirtualTextureProducerHandle const& InHandle, const FVirtualTextureProducer* InProducer);
 	uint32 AddUniquePhysicalSpace(FRHICommandListBase& RHICmdList, FVirtualTexturePhysicalSpace* InPhysicalSpace, uint32 InUniqueProducerIndex, uint32 InProducerPhysicalSpaceIndex);
+	uint32 GetResidencyMipMapBiasGroup() const;
 
 	uint32 FrameAllocated;
 	uint32 PersistentHash;
diff --git a/Engine/Source/Runtime/Renderer/Private/VT/VirtualTexturePhysicalSpace.h b/Engine/Source/Runtime/Renderer/Private/VT/VirtualTexturePhysicalSpace.h
--- a/Engine/Source/Runtime/Renderer/Private/VT/VirtualTexturePhysicalSpace.h
+++ b/Engine/Source/Runtime/Renderer/Private/VT/VirtualTexturePhysicalSpace.h
@@ -69,11 +69,13 @@ struct FVTPhysicalSpaceDescriptionExt
 	int32 TileWidthHeight = 0;
 	int32 PoolCount = 1;
 	bool bEnableResidencyMipMapBias = false;
+	int32 ResidencyMipMapBiasGroup = 0;
 };
 
 inline bool operator==(const FVTPhysicalSpaceDescriptionExt& Lhs, const FVTPhysicalSpaceDescriptionExt& Rhs)
 {
-	return Lhs.TileWidthHeight == Rhs.TileWidthHeight && Lhs.PoolCount == Rhs.PoolCount && Lhs.bEnableResidencyMipMapBias == Rhs.bEnableResidencyMipMapBias;
+	return Lhs.TileWidthHeight == Rhs.TileWidthHeight && Lhs.PoolCount == Rhs.PoolCount && 
+		Lhs.bEnableResidencyMipMapBias == Rhs.bEnableResidencyMipMapBias && Lhs.ResidencyMipMapBiasGroup == Rhs.ResidencyMipMapBiasGroup;
 }
 
 /** Virtual texture physical space that contains the physical texture and the page pool that tracks residency. */
diff --git a/Engine/Source/Runtime/Renderer/Private/VT/VirtualTextureSystem.cpp b/Engine/Source/Runtime/Renderer/Private/VT/VirtualTextureSystem.cpp
--- a/Engine/Source/Runtime/Renderer/Private/VT/VirtualTextureSystem.cpp
+++ b/Engine/Source/Runtime/Renderer/Private/VT/VirtualTextureSystem.cpp
@@ -1009,6 +1009,7 @@ void GetPhysicalSpaceExtraDescription(FVTPhysicalSpaceDescription const& InDesc,
 	OutDescExt.TileWidthHeight = TileWidthHeight;
 	OutDescExt.PoolCount = PoolCount;
 	OutDescExt.bEnableResidencyMipMapBias = Config.bEnableResidencyMipMapBias;
+	OutDescExt.ResidencyMipMapBiasGroup = FMath::Clamp(Config.ResidencyMipMapBiasGroup, 0, 3);
 }
 
 /** Cached version of GetPhysicalSpaceExtraDescription() to avoid regularly repeating the heavy work in that function. */
@@ -2677,7 +2678,7 @@ void FVirtualTextureSystem::GatherFeedbackRequests(FConcurrentLinearBulkObjectAl
 			// todo: We can split this into concurrent tasks. 
 			FAddRequestedTilesParameters Parameters;
 			Parameters.System = this;
-			Parameters.LevelBias = FMath::FloorToInt(CVarVTPlaybackMipBias.GetValueOnRenderThread() + GetGlobalMipBias() + 0.5f);
+			Parameters.LevelBias = FMath::FloorToInt(CVarVTPlaybackMipBias.GetValueOnRenderThread() + UTexture2D::GetGlobalMipMapLODBias() + 0.5f);
 			Parameters.RequestBuffer = PageRequestPlaybackBuffer.GetData();
 			Parameters.NumRequests = PageRequestPlaybackBuffer.Num();
 			Parameters.UniquePageList = Allocator.Create<FUniquePageList>();
@@ -3013,17 +3014,21 @@ void FVirtualTextureSystem::ReleasePendingResources()
 	ReleasePendingSpaces();
 }
 
-float FVirtualTextureSystem::GetGlobalMipBias() const
+FVector4f FVirtualTextureSystem::GetGlobalMipBias() const
 {
-	float MaxResidencyMipMapBias = 0.f;
+	FVector4f MaxResidencyMipMapBias(EForceInit::ForceInitToZero);
 	for (int32 SpaceIndex = 0; SpaceIndex < PhysicalSpaces.Num(); ++SpaceIndex)
 	{
-		FVirtualTexturePhysicalSpace const* PhysicalSpace = PhysicalSpaces[SpaceIndex];
-		const float ResidencyMipMapBias = PhysicalSpace != nullptr ? PhysicalSpace->GetResidencyMipMapBias() : 0.f;
-		MaxResidencyMipMapBias = FMath::Max(MaxResidencyMipMapBias, ResidencyMipMapBias);
+		if (FVirtualTexturePhysicalSpace const* PhysicalSpace = PhysicalSpaces[SpaceIndex])
+		{
+			const float ResidencyMipMapBias = PhysicalSpace->GetResidencyMipMapBias();
+			const int32 GroupIndex = PhysicalSpace->GetDescriptionExt().ResidencyMipMapBiasGroup;
+			MaxResidencyMipMapBias[GroupIndex] = FMath::Max(MaxResidencyMipMapBias[GroupIndex], ResidencyMipMapBias);
+		}
 	}
 
-	return UTexture2D::GetGlobalMipMapLODBias() + MaxResidencyMipMapBias;
+	const float GlobalBias = UTexture2D::GetGlobalMipMapLODBias();
+	return MaxResidencyMipMapBias + FVector4f(GlobalBias, GlobalBias, GlobalBias, GlobalBias);
 }
 
 bool FVirtualTextureSystem::IsPendingRootPageMap(IAllocatedVirtualTexture* AllocatedVT) const
@@ -3048,10 +3053,11 @@ void FVirtualTextureSystem::GetOnScreenMessages(FCoreDelegates::FSeverityMessage
 		{
 			FString const& FormatString = PhysicalSpace->GetFormatString();
 			const float MipBias = PhysicalSpace->GetResidencyMipMapBias();
+			const int32 GroupIndex = PhysicalSpace->GetDescriptionExt().ResidencyMipMapBiasGroup;
 
 			OutMessages.Add(
 				FCoreDelegates::EOnScreenMessageSeverity::Warning,
-				FText::Format(LOCTEXT("VTOversubscribed", "VT Pool [{0}] is oversubscribed. Setting MipBias {1}"), FText::FromString(FormatString), FText::AsNumber(MipBias)));
+				FText::Format(LOCTEXT("VTOversubscribed", "VT Pool [{0}] is oversubscribed. Setting MipBias {1} on Group {2}"), FText::FromString(FormatString), FText::AsNumber(MipBias), FText::AsNumber(GroupIndex)));
 		}
 	}
 }
diff --git a/Engine/Source/Runtime/Renderer/Private/VT/VirtualTextureSystem.h b/Engine/Source/Runtime/Renderer/Private/VT/VirtualTextureSystem.h
--- a/Engine/Source/Runtime/Renderer/Private/VT/VirtualTextureSystem.h
+++ b/Engine/Source/Runtime/Renderer/Private/VT/VirtualTextureSystem.h
@@ -164,7 +164,7 @@ public:
 	void FlushCache();
 	void FlushCache(FVirtualTextureProducerHandle const& ProducerHandle, int32 SpaceID, FIntRect const& TextureRegion, uint32 MaxLevelToEvict, uint32 MaxAgeToKeepMapped, EVTInvalidatePriority InvalidatePriority);
 
-	float GetGlobalMipBias() const;
+	FVector4f GetGlobalMipBias() const;
 
 	bool IsPendingRootPageMap(IAllocatedVirtualTexture* AllocatedVT) const;
 
