Path Tracer: Fix fireflies in volumetric clouds when enabling full bounces

The payload returned from TraceTransparentRay was not always initialized. This could lead to un-initialized values creeping into the rendered image. The code is now refactored to make it clearer and always ensure that the payload is marked invalid at the start in case it is not written to.

In addition, the ray marching logic in PathTracingCore is now fully split out from the cloud case to make it clear these are mutually exclusive.

#rnx

diff --git a/Engine/Shaders/Private/PathTracing/PathTracingCore.ush b/Engine/Shaders/Private/PathTracing/PathTracingCore.ush
--- a/Engine/Shaders/Private/PathTracing/PathTracingCore.ush
+++ b/Engine/Shaders/Private/PathTracing/PathTracingCore.ush
@@ -491,7 +491,14 @@ FPackedPathTracingPayload TraceTransparentRay(inout FPathState PathState, int Bo
 	FRISContext HitSample = InitRISContext(RISRandSample.x);
 	FRISContext VolSample = InitRISContext(RISRandSample.y);
 
-
+	float3 PayloadThroughput = 0;
+	FPackedPathTracingPayload Payload;
+	Payload.SetMiss();
+	if (bLastBounce && !PathState.ShouldAccumulateEmissive())
+	{
+		PathState.PathThroughput = 0;
+		return Payload;
+	}
 #if PATH_TRACER_TRACE_OPAQUE_FIRST
 	FPackedPathTracingPayload OpaquePayload = InitPathTracingPayload(PathState.GetFirstScatterType(), ApproximateCaustics ? PathState.GetPathRoughness() : 0.0);
 	OpaquePayload.SetDBufferA(float4(0, 0, 0, 1));
@@ -577,19 +584,10 @@ FPackedPathTracingPayload TraceTransparentRay(inout FPathState PathState, int Bo
 		PathState.Ray.TMax = asfloat(asuint(OpaquePayload.HitT) - 1);
 	}
 	
-
-	float3 PayloadThroughput = PathState.PathThroughput;
-	FPackedPathTracingPayload Payload;
 #else // PATH_TRACER_TRACE_OPAQUE_FIRST
+
 	InstanceInclusionMask |= bUseCameraRayType ? PATHTRACER_MASK_CAMERA_TRANSLUCENT : PATHTRACER_MASK_INDIRECT_TRANSLUCENT;
-	float3 PayloadThroughput = PathState.PathThroughput;
-	FPackedPathTracingPayload Payload;
-	if (bLastBounce && !PathState.ShouldAccumulateEmissive())
-	{
-		Payload.SetMiss();
-		PathState.PathThroughput = 0;
-		return Payload;
-	}
+
 #endif
 
 	for (;;)
@@ -722,12 +720,64 @@ FPackedPathTracingPayload TraceTransparentRay(inout FPathState PathState, int Bo
 				}
 			}
 
-			{
-				// remember the start of our interval before ray marching messes with it
-				float StartVolumeTMin = Interval.VolumeTMin;
-				// Ray marching loop
-				float3 PathThroughput = PathState.PathThroughput;
 
+#if PATH_TRACER_USE_CLOUD_SHADER
+			if (Interval.VolumeMask & PATH_TRACER_VOLUME_ENABLE_CLOUDS)
+			{
+				// let callable shader handle cloud sampling
+				FPackedPathTracingPayload CloudPayload = (FPackedPathTracingPayload) 0;
+				const bool bNeedSample = !bLastBounce;
+				CloudPayload.SetVolumetricCallableShaderInput(
+					PathState.Ray.Origin,
+					PathState.Ray.Direction,
+					Interval.VolumeTMin,
+					Interval.VolumeTMax,
+					PathState.RandSequence,
+					Interval.VolumeMask |
+					(VolumeFlags & ~PATH_TRACER_VOLUME_ENABLE_MASK) |
+					(min(Bounce, 31) << PATH_TRACER_VOLUME_CALLABLE_FLAGS_BOUNCE_SHIFT) |
+					(bNeedSample  ? PATH_TRACER_VOLUME_CALLABLE_FLAGS_GET_SAMPLE    : 0),
+					PathState.PathThroughput,
+					Interval.CloudDensity
+				);
+				CloudPayload.SetVolumetricCallableShaderRISContext(HitSample);
+				CallShader(CloudCallableShaderId, CloudPayload);
+				PathState.RandSequence = CloudPayload.GetVolumetricCallableShaderOutputRandSeq();
+				PathState.PathThroughput = CloudPayload.GetVolumetricCallableShaderOutputThroughput();
+				PathState.SetAlbedoNormal(
+					PathState.GetAlbedo() + CloudPayload.GetVolumetricCallableShaderOutputAlbedo() * DenoiserAOVWeight(PathState.GetPathRoughness()),
+					PathState.GetNormal()
+				);
+				HitSample = CloudPayload.GetVolumetricCallableShaderRISContext();
+				AccumulateRadiance(PathState.Radiance, CloudPayload.GetVolumetricCallableShaderOutputRadiance(), bIsCameraRay);
+				if (bIsCameraRay)
+				{
+					PathState.Alpha += CloudPayload.GetVolumetricCallableShaderOutputAlpha();
+				}
+				if (CloudPayload.HitT > 0)
+				{
+					// If we registered a hit, transfer it to the currently sampled payload
+					PayloadThroughput = CloudPayload.GetVolumetricCallableShaderOutputPayloadThroughput();
+					FPathTracingPayload Result = (FPathTracingPayload)0; // clear all fields
+					Result.HitT = CloudPayload.HitT;
+					Result.TranslatedWorldPos = PathState.Ray.Origin + PathState.Ray.Direction * CloudPayload.HitT;
+					Result.ShadingModelID = SHADINGMODELID_MEDIUM;
+					Result.BSDFOpacity = 1.0;
+					Result.PrimitiveLightingChannelMask = 7;
+					Result.SetFrontFace();
+					Result.SetBaseColor(CloudPayload.GetVolumetricCallableShaderOutputRayleighWeight());
+					float4 HGData = CloudPayload.GetVolumetricCallableShaderOutputHG();
+					Result.SetHG(HGData.xyz, HGData.w);
+					Result.SetDualHG(
+						CloudPayload.GetVolumetricCallableShaderOutputDualHGWeight(),
+						CloudPayload.GetVolumetricCallableShaderOutputDualHGPhaseData()
+					);
+					Payload = PackPathTracingPayload(Result);
+				}
+			}
+			else
+#endif // PATH_TRACER_USE_CLOUD_SHADER
+			{
 				// In order to use analytic transmittance, the user must have requested it _and_ we must be on an interval with no clouds or heterogeneous volumes
 				const bool bUseAnalyticTransmittance =
 					(VolumeFlags & PATH_TRACER_VOLUME_USE_ANALYTIC_TRANSMITTANCE) != 0 &&
@@ -735,69 +785,17 @@ FPackedPathTracingPayload TraceTransparentRay(inout FPathState PathState, int Bo
 				// are there any holdouts on the current segment?
 				const bool bHaveHoldouts =
 					(Interval.VolumeMask << PATH_TRACER_MAX_VOLUMES) & VolumeFlags & PATH_TRACER_VOLUME_HOLDOUT_MASK != 0;
-
-#if PATH_TRACER_USE_CLOUD_SHADER
-				if (Interval.VolumeMask & PATH_TRACER_VOLUME_ENABLE_CLOUDS)
-				{
-					// let callable shader handle cloud sampling
-					FPackedPathTracingPayload CloudPayload = (FPackedPathTracingPayload) 0;
-					const bool bNeedSample = !bLastBounce;
-					CloudPayload.SetVolumetricCallableShaderInput(
-						PathState.Ray.Origin,
-						PathState.Ray.Direction,
-						Interval.VolumeTMin,
-						Interval.VolumeTMax,
-						PathState.RandSequence,
-						Interval.VolumeMask |
-						  (VolumeFlags & ~PATH_TRACER_VOLUME_ENABLE_MASK) |
-						  (min(Bounce, 31) << PATH_TRACER_VOLUME_CALLABLE_FLAGS_BOUNCE_SHIFT) |
-						  (bNeedSample  ? PATH_TRACER_VOLUME_CALLABLE_FLAGS_GET_SAMPLE    : 0),
-						PathThroughput,
-						Interval.CloudDensity
-					);
-					CloudPayload.SetVolumetricCallableShaderRISContext(HitSample);
-					CallShader(CloudCallableShaderId, CloudPayload);
-					PathState.RandSequence = CloudPayload.GetVolumetricCallableShaderOutputRandSeq();
-					PathThroughput = CloudPayload.GetVolumetricCallableShaderOutputThroughput();
-					PathState.SetAlbedoNormal(
-						PathState.GetAlbedo() + CloudPayload.GetVolumetricCallableShaderOutputAlbedo() * DenoiserAOVWeight(PathState.GetPathRoughness()),
-						PathState.GetNormal()
-					);
-					HitSample = CloudPayload.GetVolumetricCallableShaderRISContext();
-					AccumulateRadiance(PathState.Radiance, CloudPayload.GetVolumetricCallableShaderOutputRadiance(), bIsCameraRay);
-					if (bIsCameraRay)
-					{
-						PathState.Alpha += CloudPayload.GetVolumetricCallableShaderOutputAlpha();
-					}
-					if (CloudPayload.HitT > 0)
-					{
-						// If we registered a hit, transfer it to the currently sampled payload
-						PayloadThroughput = CloudPayload.GetVolumetricCallableShaderOutputPayloadThroughput();
-						FPathTracingPayload Result = (FPathTracingPayload)0; // clear all fields
-						Result.HitT = CloudPayload.HitT;
-						Result.TranslatedWorldPos = PathState.Ray.Origin + PathState.Ray.Direction * CloudPayload.HitT;
-						Result.ShadingModelID = SHADINGMODELID_MEDIUM;
-						Result.BSDFOpacity = 1.0;
-						Result.PrimitiveLightingChannelMask = 7;
-						Result.SetFrontFace();
-						Result.SetBaseColor(CloudPayload.GetVolumetricCallableShaderOutputRayleighWeight());
-						float4 HGData = CloudPayload.GetVolumetricCallableShaderOutputHG();
-						Result.SetHG(HGData.xyz, HGData.w);
-						Result.SetDualHG(
-							CloudPayload.GetVolumetricCallableShaderOutputDualHGWeight(),
-							CloudPayload.GetVolumetricCallableShaderOutputDualHGPhaseData()
-						);
-						Payload = PackPathTracingPayload(Result);
-					}
-				}
-				else
-#endif // PATH_TRACER_USE_CLOUD_SHADER
 				// reasons to enter the ray marching loop:
 				//  - we can't use analytic transmittance (need ray marching for throughput or emission)
 				//  - we are a camera ray and some volumes are flagged as holdouts
 				//  - this is not the last bounce, so we need to select a point for scattering
 				if (!bUseAnalyticTransmittance || (bIsCameraRay && bHaveHoldouts) || !bLastBounce)
 				{
+					// remember the start of our interval before ray marching messes with it
+					float StartVolumeTMin = Interval.VolumeTMin;
+					// Ray marching loop
+					float3 PathThroughput = PathState.PathThroughput;
+
 					float3 VolumeRadiance = 0;
 					// Limit number of steps to prevent timeouts
 					// TODO: This biases the result! Is there a better way to limit the number of steps?
@@ -867,20 +865,19 @@ FPackedPathTracingPayload TraceTransparentRay(inout FPathState PathState, int Bo
 					}
 					// accumulate all at once so clamping applies to entire volume contribution
 					AccumulateRadiance(PathState.Radiance, VolumeRadiance, bIsCameraRay);
-				}
-
-				// restore interval start now that we are finished with ray marching
-				Interval.VolumeTMin = StartVolumeTMin;
+					// restore interval start now that we are finished with ray marching
+					Interval.VolumeTMin = StartVolumeTMin;
 
-				if (bUseAnalyticTransmittance)
-				{
-					// analytically handle the transmittance to the next surface (which will be higher quality than the transmittance implicitly computed above and prevent noise with emissive surfaces in volumes)
-					PathState.PathThroughput = VolumeGetTransmittance(PathState.PathThroughput, PathState.Ray.Origin, PathState.Ray.Direction, Interval, PathState.RandSequence);
-				}
-				else
-				{
-					// record the change in transmittance since we already computed it
-					PathState.PathThroughput = PathThroughput;
+					if (bUseAnalyticTransmittance)
+					{
+						// analytically handle the transmittance to the next surface (which will be higher quality than the transmittance implicitly computed above and prevent noise with emissive surfaces in volumes)
+						PathState.PathThroughput = VolumeGetTransmittance(PathState.PathThroughput, PathState.Ray.Origin, PathState.Ray.Direction, Interval, PathState.RandSequence);
+					}
+					else
+					{
+						// record the change in transmittance since we already computed it
+						PathState.PathThroughput = PathThroughput;
+					}
 				}
 			}
 
@@ -1014,16 +1011,9 @@ FPackedPathTracingPayload TraceTransparentRay(inout FPathState PathState, int Bo
 	// normalization will be 0 if we didn't pick any volume segment
 	VolumeSegment.Throughput *= VolSample.GetNormalization();
 
-	if (HitSample.HasSample())
-	{
-		// if we stored a valid hit in the payload, reset the path throughput to this point
-		PathState.PathThroughput = PayloadThroughput * HitSample.GetNormalization();
-	}
-	else
-	{
-		PathState.PathThroughput = 0;
-		Payload.SetMiss();
-	}
+	// if we stored a valid hit in the payload, reset the path throughput to this point
+	// normalization will be 0 if we didn't pick any hit
+	PathState.PathThroughput = PayloadThroughput * HitSample.GetNormalization();
 
 	return Payload;
 }
@@ -1998,7 +1988,7 @@ bool PathTracingKernel(inout FPathState PathState, int Bounce)
 		}
 	}
 
-	if (MaterialSample.Pdf < 0 || asuint(MaterialSample.Pdf) > 0x7F800000)
+	if (asuint(MaterialSample.Pdf) > 0x7F800000)
 	{
 		// Pdf became invalid (either negative or NaN)
 		PathState.Radiance = float3(1, 0, 1);
diff --git a/Engine/Shaders/Private/PathTracing/PathTracingVolumetricCloudMaterialShader.usf b/Engine/Shaders/Private/PathTracing/PathTracingVolumetricCloudMaterialShader.usf
--- a/Engine/Shaders/Private/PathTracing/PathTracingVolumetricCloudMaterialShader.usf
+++ b/Engine/Shaders/Private/PathTracing/PathTracingVolumetricCloudMaterialShader.usf
@@ -42,8 +42,6 @@
 
 #include "./PathTracingCommon.ush"
 
-uint MaxRaymarchSteps;
-
 FVolumeShadedResult CloudsGetDensity(float3 TranslatedWorldPos)
 {
 	FDFVector3 AbsoluteWorldPosition = DFFastSubtract(TranslatedWorldPos, ResolvedView.PreViewTranslation);
