Path Tracer: Add support for purely emissive volumetric clouds

This can be useful for effects like auroras which have no density, but still emit light. In this case, we must artificially introduce some density in the cloud acceleration texture to make sure ray marching has something to step through.

diff --git a/Engine/Shaders/Private/PathTracing/PathTracingBuildCloudAccelerationMap.usf b/Engine/Shaders/Private/PathTracing/PathTracingBuildCloudAccelerationMap.usf
--- a/Engine/Shaders/Private/PathTracing/PathTracingBuildCloudAccelerationMap.usf
+++ b/Engine/Shaders/Private/PathTracing/PathTracingBuildCloudAccelerationMap.usf
@@ -63,10 +63,17 @@ void PathTracingBuildCloudAccelerationMapCS(uint2 DispatchThreadId : SV_Dispatch
 		float3 WorldPos = mul(ClipPos, GetCloudClipBasis()) + CloudClipCenterKm; // TangentToWorld
 		FDFVector3 AbsoluteWorldPosition = DFMultiply(WorldPos, SKY_UNIT_TO_CM);
 		// run the material and encode the result, which includes the effect that quantization will have during rendering
-		float3 Extinction = SampleCloudMaterial(AbsoluteWorldPosition, H, B, T).Extinction;
+		FSampleCloudMaterialResult Result = SampleCloudMaterial(AbsoluteWorldPosition, H, B, T);
+		float3 Extinction = Result.Extinction;
 		
-		const float Density   = max3(Extinction.x, Extinction.y, Extinction.z);
-		const float DensityLo = min3(Extinction.x, Extinction.y, Extinction.z);
+		float Density   = max3(Extinction.x, Extinction.y, Extinction.z);
+		float DensityLo = min3(Extinction.x, Extinction.y, Extinction.z);
+
+		// If shader had no density but some emission, we need to artificially boost the density so that we get at least _some_ raymarch steps in the empty regions
+		if (Density == 0 && any(Result.EmissiveColor > 0))
+		{
+			Density = CloudInvVoxelWidth; // ~1 step per voxel
+		}
 
 		if (Density > 0)
 		{
