Fixed: Incorrect 'Pre-Skinned Local Bounds' in a mobile specific GPUScene implementation (2nd version without a branch)

diff --git a/Engine/Shaders/Private/SceneDataMobileLoader.ush b/Engine/Shaders/Private/SceneDataMobileLoader.ush
--- a/Engine/Shaders/Private/SceneDataMobileLoader.ush
+++ b/Engine/Shaders/Private/SceneDataMobileLoader.ush
@@ -113,12 +113,6 @@ FPrimitiveSceneData LoadPrimitiveDataUBO(uint DrawInstanceId)
 		DataIdx += 1u;
 	}
 	
-	// PreSkinnedLocalBounds, TODO: this works until we actually support GPUScene for skinned primitives
-	{
-		PrimitiveData.PreSkinnedLocalBoundsMin = PrimitiveData.LocalObjectBoundsMin;
-		PrimitiveData.PreSkinnedLocalBoundsMax = PrimitiveData.LocalObjectBoundsMax;
-	}
-
 	// InstanceLocalBounds
 	{
 		PrimitiveData.InstanceLocalBoundsCenter = (PrimitiveData.LocalObjectBoundsMin + PrimitiveData.LocalObjectBoundsMax) * 0.5f;
@@ -164,6 +158,14 @@ FPrimitiveSceneData LoadPrimitiveDataUBO(uint DrawInstanceId)
 	// Should not be used with regular primitives
 	PrimitiveData.WorldToPreviousWorld = float4x4(1,0,0,0, 0,1,0,0, 0,0,1,0, 0,0,0,1);
 
+	// PreSkinnedLocalBounds, TODO.w
+	// TODO: this is needed only for skinned meshes
+	{
+		PrimitiveData.PreSkinnedLocalBoundsMin = LoadPrimitiveDataElementUBO(DrawInstanceId, DataIdx + 0);
+		PrimitiveData.PreSkinnedLocalBoundsMax = LoadPrimitiveDataElementUBO(DrawInstanceId, DataIdx + 1);
+		DataIdx += 2u;
+	}
+
 	// CustomData, TODO make optional
 	{
 		uint NumCustomData = min(NUM_CUSTOM_PRIMITIVE_DATA, BATCHED_PRIMITIVE_DATA_STRIDE_FLOAT4 - DataIdx); //23u
diff --git a/Engine/Shaders/Private/SceneDataMobileWriter.ush b/Engine/Shaders/Private/SceneDataMobileWriter.ush
--- a/Engine/Shaders/Private/SceneDataMobileWriter.ush
+++ b/Engine/Shaders/Private/SceneDataMobileWriter.ush
@@ -138,6 +138,14 @@ void WritePrimitiveDataUBO(RWStructuredBuffer<float4> Output, uint Offset, FPrim
 		WriteDataElementUBO(Output, Offset, DataIdx + 2, PreviousRelativeWorldToLocal[2]);
 		DataIdx += 3u;
 	}
+	
+	// PreSkinnedLocalBounds, TODO.w
+	// TODO: this is needed only for skinned meshes
+	{
+		WriteDataElementUBO(Output, Offset, DataIdx + 0, float4(PrimitiveData.PreSkinnedLocalBoundsMin, 0.f));
+		WriteDataElementUBO(Output, Offset, DataIdx + 1, float4(PrimitiveData.PreSkinnedLocalBoundsMax, 0.f));
+		DataIdx += 2u;
+	}
 
 	// CustomData, TODO make optional
 	{
