Fix for object position values not propagating from the decal primitive buffers due to missing macro in materials.

diff --git a/Engine/Shaders/Private/MaterialTemplate.ush b/Engine/Shaders/Private/MaterialTemplate.ush
--- a/Engine/Shaders/Private/MaterialTemplate.ush
+++ b/Engine/Shaders/Private/MaterialTemplate.ush
@@ -1717,7 +1717,7 @@ FWSVector3 TransformLocalPositionToPrevWorld(FMaterialVertexParameters Parameter
 /** Return the object's position in world space */
 FWSVector3 GetObjectWorldPosition(FMaterialPixelParameters Parameters)
 {
-#if HAS_PRIMITIVE_UNIFORM_BUFFER
+#if HAS_PRIMITIVE_UNIFORM_BUFFER || DECAL_PRIMITIVE
 	return Parameters.LWCData.ObjectWorldPosition;
 #elif MATERIAL_DOMAIN_LIGHTFUNCTION && !LIGHT_ATLAS
 	/**
@@ -1737,7 +1737,7 @@ FWSVector3 GetObjectWorldPosition(FMaterialVertexParameters Parameters)
 }
 float3 GetObjectTranslatedWorldPosition(FMaterialPixelParameters Parameters)
 {
-#if HAS_PRIMITIVE_UNIFORM_BUFFER
+#if HAS_PRIMITIVE_UNIFORM_BUFFER || DECAL_PRIMITIVE
 
 #if PRIMITIVE_HAS_TILEOFFSET_DATA && VIEW_HAS_TILEOFFSET_DATA
 	return LWCToFloat(LWCAdd(GetPrimitiveData(Parameters).ObjectWorldPositionTO, ResolvedView.TileOffset.PreViewTranslation));
