Volumetric Cloud: Fix minor shader truncation warnings when running clouds in greyscale mode

Cleanup handling of MATERIAL_VOLUMETRIC_ADVANCED_GRAYSCALE_MATERIAL so that we don't need to do min/max/avg across channels when the values are tracked as floats.

Also fix the transmittance calculation to lights to not be done in greyscale since it includes the atmosphere transmittance. This may result in a slight visual change for cloud materials flagged as greyscale.

diff --git a/Engine/Shaders/Private/VolumetricCloud.usf b/Engine/Shaders/Private/VolumetricCloud.usf
--- a/Engine/Shaders/Private/VolumetricCloud.usf
+++ b/Engine/Shaders/Private/VolumetricCloud.usf
@@ -349,11 +349,21 @@ float SamplePhaseFunction(in float PhaseCosTheta, in float PhaseG, in float Phas
 #endif
 
 #if MATERIAL_VOLUMETRIC_ADVANCED_GRAYSCALE_MATERIAL
-	#define MATVEC float
-	#define MATSUF rrr
+
+#define MATVEC              float
+#define ColorToMATVEC(C)    ((C).r)   // Assume all 3 floats are the same
+#define MATVEC_Max(C)       (C)
+#define MATVEC_Min(C)       (C)
+#define MATVEC_Avg(C)       (C)
+
 #else
-	#define MATVEC float3
-	#define MATSUF rgb
+
+#define MATVEC              float3
+#define ColorToMATVEC(C)    ((C).rgb)
+float MATVEC_Max(float3 C) { return max3(C.r, C.g, C.b); }
+float MATVEC_Min(float3 C) { return min3(C.r, C.g, C.b); }
+float MATVEC_Avg(float3 C) { return dot(C, 1.0f / 3.0f); }
+
 #endif
 
 struct ParticipatingMediaContext
@@ -361,13 +371,13 @@ struct ParticipatingMediaContext
 	MATVEC ScatteringCoefficients[MSCOUNT];
 	MATVEC ExtinctionCoefficients[MSCOUNT];
 
-	MATVEC TransmittanceToLight0[MSCOUNT];
+	float3 TransmittanceToLight0[MSCOUNT];
 #if CLOUD_SAMPLE_SECOND_LIGHT
-	MATVEC TransmittanceToLight1[MSCOUNT];
+	float3 TransmittanceToLight1[MSCOUNT];
 #endif
 };
 
-ParticipatingMediaContext SetupParticipatingMediaContext(float3 BaseAlbedo, float3 BaseExtinctionCoefficients, float MsSFactor, float MsEFactor, float3 InitialTransmittanceToLight0, float3 InitialTransmittanceToLight1)
+ParticipatingMediaContext SetupParticipatingMediaContext(MATVEC BaseAlbedo, MATVEC BaseExtinctionCoefficients, float MsSFactor, float MsEFactor, float3 InitialTransmittanceToLight0, float3 InitialTransmittanceToLight1)
 {
 	const MATVEC ScatteringCoefficients = BaseAlbedo * BaseExtinctionCoefficients;
 	//const float3 AbsorptionCoefficients = max(0.0f, BaseExtinctionCoefficients - ScatteringCoefficients);
@@ -898,9 +908,9 @@ void MainCommon(in FMaterialPixelParameters MaterialParameters, in float4 SvPosi
 			//////////////////////////////
 			VolumetricCloudCalcPixelMaterialInputs(MaterialParameters, PixelMaterialInputs);
 			ConvertCloudPixelMaterialInputsToWorkingColorSpace(PixelMaterialInputs);
-			const MATVEC ExtinctionCoefficients = SampleExtinctionCoefficients(PixelMaterialInputs);
-			const MATVEC EmissiveLuminance = USES_EMISSIVE_COLOR ? SampleEmissive(PixelMaterialInputs) : 0.0f;
-			const MATVEC Albedo = SampleAlbedo(PixelMaterialInputs);
+			const MATVEC ExtinctionCoefficients = ColorToMATVEC(SampleExtinctionCoefficients(PixelMaterialInputs));
+			const MATVEC EmissiveLuminance = USES_EMISSIVE_COLOR ? ColorToMATVEC(SampleEmissive(PixelMaterialInputs)) : 0.0f;
+			const MATVEC Albedo = ColorToMATVEC(SampleAlbedo(PixelMaterialInputs));
 
 #if MATERIAL_VOLUMETRIC_ADVANCED_PHASE_PERSAMPLE
 			const float PhaseG = GetVolumetricAdvancedMaterialOutput0(MaterialParameters);
@@ -955,7 +965,7 @@ void MainCommon(in FMaterialPixelParameters MaterialParameters, in float4 SvPosi
 			//////////////////////////////
 			if (any(PMC.ScatteringCoefficients[0] > 0.0f))
 			{
-				const float MaxTransmittanceToView = max(max(TransmittanceToView.MATSUF.x, TransmittanceToView.MATSUF.y), TransmittanceToView.MATSUF.z);
+				const float MaxTransmittanceToView = MATVEC_Max(TransmittanceToView);
 
 				//
 				// A- we compute lighting bouncing off the ground: affected by light direction, transmittance in the atmosphere and albedo.
@@ -1121,12 +1131,12 @@ void MainCommon(in FMaterialPixelParameters MaterialParameters, in float4 SvPosi
 
 					if (MaterialParameters.CloudSampleNormAltitudeInLayer <= 0.0f || MaterialParameters.CloudSampleNormAltitudeInLayer >= 1.0f)
 					{
-						break; // Ignore remaining samples since we have just traveld out of the cloud layer.
+						break; // Ignore remaining samples since we have just traveled out of the cloud layer.
 					}
 
 					VolumetricCloudCalcPixelMaterialInputs(MaterialParameters, PixelMaterialInputs);
 					ConvertCloudPixelMaterialInputsToWorkingColorSpace(PixelMaterialInputs);
-					MATVEC ShadowExtinctionCoefficients = SampleExtinctionCoefficients(PixelMaterialInputs);
+					MATVEC ShadowExtinctionCoefficients = ColorToMATVEC(SampleExtinctionCoefficients(PixelMaterialInputs));
 
 					ParticipatingMediaContext ShadowPMC = SetupParticipatingMediaContext(0.0f, ShadowExtinctionCoefficients, MsScattFactor, MsExtinFactor, 0.0f, 0.0f);
 
@@ -1191,7 +1201,7 @@ void MainCommon(in FMaterialPixelParameters MaterialParameters, in float4 SvPosi
 
 						VolumetricCloudCalcPixelMaterialInputs(MaterialParameters, PixelMaterialInputs);
 						ConvertCloudPixelMaterialInputsToWorkingColorSpace(PixelMaterialInputs);
-						MATVEC ShadowExtinctionCoefficients = SampleExtinctionCoefficients(PixelMaterialInputs);
+						MATVEC ShadowExtinctionCoefficients = ColorToMATVEC(SampleExtinctionCoefficients(PixelMaterialInputs));
 
 						ParticipatingMediaContext ShadowPMC = SetupParticipatingMediaContext(0.0f, ShadowExtinctionCoefficients, MsScattFactor, MsExtinFactor, 0.0f, 0.0f);
 
@@ -1221,7 +1231,7 @@ void MainCommon(in FMaterialPixelParameters MaterialParameters, in float4 SvPosi
 			// Compute the weighted average of t for the aerial perspective evaluation.
 			if (any(PMC.ExtinctionCoefficients[0] > 0.0))
 			{
-				float tAPWeight = min(TransmittanceToView.MATSUF.r, min(TransmittanceToView.MATSUF.g, TransmittanceToView.MATSUF.b));
+				float tAPWeight = MATVEC_Min(TransmittanceToView);
 				tAPWeightedSum += t * tAPWeight;
 				tAPWeightsSum += tAPWeight;
 
@@ -1315,7 +1325,7 @@ void MainCommon(in FMaterialPixelParameters MaterialParameters, in float4 SvPosi
 
 								VolumetricCloudCalcPixelMaterialInputs(MaterialParameters, PixelMaterialInputs);
 								ConvertCloudPixelMaterialInputsToWorkingColorSpace(PixelMaterialInputs);
-								MATVEC ShadowExtinctionCoefficients = SampleExtinctionCoefficients(PixelMaterialInputs);
+								MATVEC ShadowExtinctionCoefficients = ColorToMATVEC(SampleExtinctionCoefficients(PixelMaterialInputs));
 
 								if (MaterialParameters.CloudSampleNormAltitudeInLayer <= 0.0f || MaterialParameters.CloudSampleNormAltitudeInLayer >= 1.0f)
 								{
@@ -1474,7 +1484,7 @@ void MainCommon(in FMaterialPixelParameters MaterialParameters, in float4 SvPosi
 		ClosestLuminance *= 0;
 	#endif
 
-		OutHoldOutCoverage = saturate(1.0f - dot(TransmittanceToView.MATSUF, 1.0f / 3.0f));
+		OutHoldOutCoverage = saturate(1.0f - MATVEC_Avg(TransmittanceToView));
 	}
 #endif
 
@@ -1487,10 +1497,10 @@ void MainCommon(in FMaterialPixelParameters MaterialParameters, in float4 SvPosi
 	// No visual issues have been noticed with reprojection+TAA so far.
 	const float NoCloudDepth = TMax;
 	const float tAP = tAPWeightsSum==0.0f ? NoCloudDepth : tAPWeightedSum / max(0.0000000001f, tAPWeightsSum);
-	float MeanTransmittance = dot(TransmittanceToView.MATSUF, 1.0f / 3.0f);
+	float MeanTransmittance = MATVEC_Avg(TransmittanceToView);
 #if CLOUD_MIN_AND_MAX_DEPTH
 	const float ClosestTAP = ClostestTAPWeightsSum == 0.0f ? NoCloudDepth : ClostestTAPWeightedSum / max(0.0000000001f, ClostestTAPWeightsSum);
-	float ClosestMeanTransmittance = dot(ClosestTransmittanceToView.MATSUF, 1.0f / 3.0f);
+	float ClosestMeanTransmittance = MATVEC_Avg(ClosestTransmittanceToView);
 #endif // CLOUD_MIN_AND_MAX
 
 	float3 FogSampleWorldPositionRelativeToCameraCm = tAP * Raydir;
