Path Tracer: Fix rare fireflies in clouds along the axes where ray direction becomes exactly 0.0

Simply nudge the ray direction to avoid this case

#rnx

diff --git a/Engine/Shaders/Private/PathTracing/PathTracingVolumetricCloudMaterialShader.usf b/Engine/Shaders/Private/PathTracing/PathTracingVolumetricCloudMaterialShader.usf
--- a/Engine/Shaders/Private/PathTracing/PathTracingVolumetricCloudMaterialShader.usf
+++ b/Engine/Shaders/Private/PathTracing/PathTracingVolumetricCloudMaterialShader.usf
@@ -141,7 +141,7 @@ RAY_TRACING_ENTRY_CALLABLE(PathTracingVolumetricCloudMaterialShader, FPackedPath
 		// isect clip slabs, the AABB our clouds are enclosed in
 		float3 Ro = mul(GetCloudClipBasis(), Ray.Origin);
 		float3 Rd = mul(GetCloudClipBasis(), Ray.Direction);
-		float3 InvRd = rcp(Rd);
+		float3 InvRd = rcp(max(abs(Rd), 1e-6)) * select(Rd >= 0, 1.0, -1.0); // protect against division by 0
 
 		// now that we know the ray length, lets march through the grid to see if we can clip away anything
 		float2 P = Ro.xy + Ray.TMin * Rd.xy;
diff --git a/Engine/Shaders/Private/PathTracing/Volume/PathTracingClouds.ush b/Engine/Shaders/Private/PathTracing/Volume/PathTracingClouds.ush
--- a/Engine/Shaders/Private/PathTracing/Volume/PathTracingClouds.ush
+++ b/Engine/Shaders/Private/PathTracing/Volume/PathTracingClouds.ush
@@ -84,7 +84,7 @@ FVolumeIntersection CloudsIntersect(float3 Origin, float3 Direction, float TMin,
 		// isect clip slabs, the AABB our clouds are enclosed in
 		float3 Ro = mul(GetCloudClipBasis(), Origin);
 		float3 Rd = mul(GetCloudClipBasis(), Direction);
-		float3 InvRd = rcp(Rd);
+		float3 InvRd = rcp(max(abs(Rd), 1e-6)) * select(Rd >= 0, 1.0, -1.0); // protect against division by 0
 		float2 Slab0 = (-Ro.xy - CloudClipDistKm * SKY_UNIT_TO_CM) * InvRd.xy;
 		float2 Slab1 = (-Ro.xy + CloudClipDistKm * SKY_UNIT_TO_CM) * InvRd.xy;
 		VolumeTMin = max(VolumeTMin, max(min(Slab0.x, Slab1.x), min(Slab0.y, Slab1.y)));
