Also fix landscape position for RT hit materials and not only for dynamic compute pass

diff --git a/Engine/Shaders/Private/LandscapeVertexFactory.ush b/Engine/Shaders/Private/LandscapeVertexFactory.ush
--- a/Engine/Shaders/Private/LandscapeVertexFactory.ush
+++ b/Engine/Shaders/Private/LandscapeVertexFactory.ush
@@ -179,22 +179,24 @@ struct FVertexFactoryInput
 FVertexFactoryInput LoadVertexFactoryInputForHGS(uint TriangleIndex, int VertexIndex)
 {
 	FVertexFactoryInput Input;
+		
+	uint RTSectionSizeVertsCurrentLOD = GetRayTracingSectionVertCountForLOD(GetComponentLinearIndex());
+	uint RTSectionSizeQuadsCurrentLOD = RTSectionSizeVertsCurrentLOD - 1;
+	uint NumRayTracingSections = LandscapeParameters.NumRayTracingSections;
 
-	float4 LodValues = GetLodValues(GetComponentLinearIndex());
-	uint SubsectionSizeQuadsCurrentLOD = LodValues.z;
-	uint SubsectionSizeVertsCurrentLOD = LodValues.z + 1;
-
+	float RTSectionSizeScaler = (RTSectionSizeQuadsCurrentLOD - 1.0f / NumRayTracingSections) / RTSectionSizeQuadsCurrentLOD ;
+	
 	#if 1
 		FTriangleBaseAttributes Tri = LoadTriangleBaseAttributes(TriangleIndex);
 		uint VertexId = Tri.Indices[VertexIndex];
 	#else
 		uint QuadId = TriangleIndex / 2;
-		uint QuadX = QuadId % SubsectionSizeQuadsCurrentLOD;
-		uint QuadY = QuadId / SubsectionSizeQuadsCurrentLOD;
-		uint i00 = (QuadX + 0) + (QuadY + 0) * SubsectionSizeVertsCurrentLOD;
-		uint i10 = (QuadX + 1) + (QuadY + 0) * SubsectionSizeVertsCurrentLOD;
-		uint i11 = (QuadX + 1) + (QuadY + 1) * SubsectionSizeVertsCurrentLOD;
-		uint i01 = (QuadX + 0) + (QuadY + 1) * SubsectionSizeVertsCurrentLOD;
+		uint QuadX = QuadId % RTSectionSizeQuadsCurrentLOD;
+		uint QuadY = QuadId / RTSectionSizeQuadsCurrentLOD;
+		uint i00 = (QuadX + 0) + (QuadY + 0) * RTSectionSizeVertsCurrentLOD;
+		uint i10 = (QuadX + 1) + (QuadY + 0) * RTSectionSizeVertsCurrentLOD;
+		uint i11 = (QuadX + 1) + (QuadY + 1) * RTSectionSizeVertsCurrentLOD;
+		uint i01 = (QuadX + 0) + (QuadY + 1) * RTSectionSizeVertsCurrentLOD;
 		uint Indices[6] = {i00, i11, i10, i00, i01, i10};
 		uint VertexId = Indices[TriangleIndex % 2 * 3 + VertexIndex];
 	#endif
@@ -202,8 +204,8 @@ FVertexFactoryInput LoadVertexFactoryInputForHGS(uint TriangleIndex, int VertexI
 	// Note: GetInstanceUserData() stores the GPU-Scene instance ID
 	VF_GPUSCENE_SET_INPUT_FOR_RT(Input, GetInstanceUserData(), 0U);
 
-	Input.Position.x = VertexId % SubsectionSizeVertsCurrentLOD;
-	Input.Position.y = VertexId / SubsectionSizeVertsCurrentLOD;
+	Input.Position.x = (VertexId % RTSectionSizeVertsCurrentLOD + LandscapeMVF.RayTracingSectionXY.x * RTSectionSizeQuadsCurrentLOD) * RTSectionSizeScaler;
+	Input.Position.y = (VertexId / RTSectionSizeVertsCurrentLOD + LandscapeMVF.RayTracingSectionXY.y * RTSectionSizeQuadsCurrentLOD) * RTSectionSizeScaler;	
 	Input.Position.z = LandscapeMVF.SubXY.x;
 	Input.Position.w = LandscapeMVF.SubXY.y;
 
