[ShaderPrint] Add border to ShaderPrint widget for better visibility.

diff --git a/Engine/Shaders/Private/ShaderPrintDraw.usf b/Engine/Shaders/Private/ShaderPrintDraw.usf
--- a/Engine/Shaders/Private/ShaderPrintDraw.usf
+++ b/Engine/Shaders/Private/ShaderPrintDraw.usf
@@ -333,6 +333,20 @@ void DrawSymbolsVS(
 	}
 }
 
+bool IsInside(float2 InUV, float2 InDistToBorder)
+{
+	return 
+		InUV.x > InDistToBorder.x && InUV.x < (1.f - InDistToBorder.x) &&
+		InUV.y > InDistToBorder.y && InUV.y < (1.f - InDistToBorder.y);
+}
+
+bool IsOutside(float2 InUV, float2 InDistToBorder)
+{
+	return 
+		InUV.x < InDistToBorder.x || InUV.x > (1.f - InDistToBorder.x) ||
+		InUV.y < InDistToBorder.y || InUV.y > (1.f - InDistToBorder.y);
+}
+
 void DrawSymbolsPS(
 	in float4 InPosition : SV_POSITION,
 	in noperspective float2 InUV : TEXCOORD0,
@@ -345,11 +359,36 @@ void DrawSymbolsPS(
 	float Alpha = 1.0f;
 	if (InType == SHADER_PRINT_TYPE_SLIDER)
 	{
-		Alpha = InUV.x < asfloat(InValue) ? 1.f : 0.25f;
+		const float2 PixelSizeInUV = 1.f / GetSliderWidgetSize(ShaderPrintData.FontSpacing, ShaderPrintData.Resolution).x;
+		const float2 OuterBorder = float2(PixelSizeInUV.x, PixelSizeInUV.y * 10);
+		const float2 InnerBorder = OuterBorder * 2;
+
+		Alpha = 0.f;
+		if (IsInside(InUV, InnerBorder))
+		{
+			Alpha = InUV.x < asfloat(InValue) ? 1.f : 0.25f;
+		}
+		if (IsOutside(InUV, OuterBorder))
+		{
+			Alpha = 1.f;
+		}
+
 	}
 	else if (InType == SHADER_PRINT_TYPE_CHECK)
 	{
-		Alpha = (asuint(InValue) & 0x1) > 0 ? 1.f : 0.f;
+		const float2 PixelSizeInUV = 1.f / GetCheckboxWidgetSize(ShaderPrintData.FontSpacing, ShaderPrintData.Resolution).x;
+		const float2 OuterBorder = PixelSizeInUV;
+		const float2 InnerBorder = OuterBorder * 2;
+
+		Alpha = 0.f;
+		if (IsInside(InUV, InnerBorder))
+		{
+			Alpha = (asuint(InValue) & 0x1) > 0 ? 1.f : 0.f;
+		}
+		if (IsOutside(InUV, OuterBorder))
+		{
+			Alpha = 1.f;
+		}
 	}
 	else
 	{
