[Hair] Reenable voxel display by default when enabling voxelization view

diff --git a/Engine/Shaders/Private/HairStrands/HairStrandsVoxelDebug.usf b/Engine/Shaders/Private/HairStrands/HairStrandsVoxelDebug.usf
--- a/Engine/Shaders/Private/HairStrands/HairStrandsVoxelDebug.usf
+++ b/Engine/Shaders/Private/HairStrands/HairStrandsVoxelDebug.usf
@@ -252,7 +252,7 @@ void MainCS(uint3 DispatchThreadId : SV_DispatchThreadID)
 	int ForcedMipLevel = -1;
 	{
 		FShaderPrintContext WidgetContext = InitShaderPrintContext(MacroGroupId == 0 && bPrint, uint2(500, 100));
-		bDrawVoxel = AddCheckbox(WidgetContext, TEXT("Draw Voxels"), false, FontWhite); Newline(WidgetContext);
+		bDrawVoxel = AddCheckbox(WidgetContext, TEXT("Draw Voxels"), true, FontWhite); Newline(WidgetContext);
 		bDrawPage = AddCheckbox(WidgetContext, TEXT("Draw Pages"), false, FontWhite); Newline(WidgetContext);
 		bool  bForceMipLevel = AddCheckbox(WidgetContext, TEXT("Force MIP"), false, FontWhite); Newline(WidgetContext);
 		if (bForceMipLevel)
