[VSM] [DOF] Lower the VSM resolution in areas that are out of focus, as those will be blurred anyway.
See r.Shadow.Virtual.MaxDOFResolutionBias.

diff --git a/Engine/Shaders/Private/VirtualShadowMaps/VirtualShadowMapPageMarking.usf b/Engine/Shaders/Private/VirtualShadowMaps/VirtualShadowMapPageMarking.usf
--- a/Engine/Shaders/Private/VirtualShadowMaps/VirtualShadowMapPageMarking.usf
+++ b/Engine/Shaders/Private/VirtualShadowMaps/VirtualShadowMapPageMarking.usf
@@ -22,6 +22,7 @@
 #include "../ShaderPrint.ush"
 #include "../Froxel/Froxel.ush"
 #include "../Nanite/NaniteHZBCull.ush"
+#include "../DiaphragmDOF/DOFCommon.ush"
 
 // Type of input data consume by the page allocation (i.e., data read from the source buffer: Gbuffer, HairStrands data, ...)
 #define INPUT_TYPE_GBUFFER 0
@@ -40,6 +41,8 @@ uint bCullBackfacingPixels;
 uint NumDirectionalLightSmInds;
 uint MipModeLocal;
 
+float DOFBiasStrength;
+
 #ifdef GeneratePageFlagsFromFroxelsCS
 int PassId;
 #endif
@@ -50,19 +53,19 @@ int GetBiasedClipmapLevel(FVirtualShadowMapProjectionShaderData BaseProjectionDa
 	return int(floor(BiasedLevel));
 }
 
-uint GetMipLevelLocal(FVirtualShadowMapProjectionShaderData ProjectionData, float3 TranslatedWorldPosition, float SceneDepth)
+uint GetMipLevelLocal(FVirtualShadowMapProjectionShaderData ProjectionData, float3 TranslatedWorldPosition, float SceneDepth, /* >= 0 */ float ExtraBias = 0.0f)
 {
 	// If local lights are near the primary view the combined offset should be small
 	float3 ViewToShadowTranslation = DFFastLocalSubtractDemote(ProjectionData.PreViewTranslation, PrimaryView.PreViewTranslation);
 	float3 ShadowTranslatedWorldPosition = TranslatedWorldPosition + ViewToShadowTranslation;
 
 	float Footprint = VirtualShadowMapCalcPixelFootprintLocal(ProjectionData, ShadowTranslatedWorldPosition, SceneDepth);
-	return GetMipLevelLocal(Footprint, MipModeLocal, ProjectionData.ResolutionLodBias, VirtualShadowMap.GlobalResolutionLodBias);
+	return GetMipLevelLocal(Footprint, MipModeLocal, ProjectionData.ResolutionLodBias, VirtualShadowMap.GlobalResolutionLodBias, ExtraBias);
 }
 
-uint GetMipLevelLocal(int VirtualShadowMapId, float3 TranslatedWorldPosition, float SceneDepth)
+uint GetMipLevelLocal(int VirtualShadowMapId, float3 TranslatedWorldPosition, float SceneDepth, /* >= 0 */ float ExtraBias = 0.0f)
 {
-	return GetMipLevelLocal(GetVirtualShadowMapProjectionData(VirtualShadowMapId), TranslatedWorldPosition, SceneDepth);
+	return GetMipLevelLocal(GetVirtualShadowMapProjectionData(VirtualShadowMapId), TranslatedWorldPosition, SceneDepth, ExtraBias);
 }
 
 void MarkPageAddress(FVSMPageOffset PageOffset, uint Flags)
@@ -417,6 +420,22 @@ FPositionData InitPositionData(uint2 PixelPos, float DeviceZ, bool bIsValid)
 	return Out;
 }
 
+// Areas that are out of focus get a lower VSM resolution, as the resulting image will get blurred anyway
+float CalculateResolutionBiasFromDepthOfField(float SceneDepth)
+{
+	bool bApplyBiasFromDOFCoCRadius = DOFBiasStrength != 0.0f;
+	if (bApplyBiasFromDOFCoCRadius)
+	{
+		float AbsCocRadius = abs(SceneDepthToCocRadius(SceneDepth));
+		const float BiasRampUpDistance = 8.0f;
+		return min(DOFBiasStrength, lerp(0.0f, DOFBiasStrength, AbsCocRadius / BiasRampUpDistance));
+	}
+	else
+	{
+		return 0.0f;
+	}
+}
+
 // Bias applied to pixels originating from first-person geometry, if negative marking is skipped for these
 float FirstPersonPixelRequestBias;
 int FirstPersonPixelRequestLevelClamp;
@@ -595,6 +614,11 @@ void GeneratePageFlagsFromPixels(
 	// We can't outright skip the pixels as we want environment shadow cast onto the FP geometry.
 	float PrimaryExtraBias = bIsFirstPersonPixel ? FirstPersonPixelRequestBias : 0.0f;
 	int PrimaryMinLevelClamp = bIsFirstPersonPixel ? FirstPersonPixelRequestLevelClamp : 0;
+	float LocalExtraBias = 0.0f;
+	
+	float DOFResolutionBias = CalculateResolutionBiasFromDepthOfField(Primary.SceneDepth);
+	PrimaryExtraBias += DOFResolutionBias;
+	LocalExtraBias += DOFResolutionBias;
 
 	bool bMarkFirstPersonPixels = FirstPersonPixelRequestBias >= 0.0f;
 	// Directional lights
@@ -713,7 +737,7 @@ void GeneratePageFlagsFromPixels(
 				VirtualShadowMapId += VirtualShadowMapGetCubeFace(-ToLight);
 			}
 
-			uint MipLevel = GetMipLevelLocal(VirtualShadowMapId, Primary.TranslatedWorldPosition, Primary.SceneDepth);
+			uint MipLevel = GetMipLevelLocal(VirtualShadowMapId, Primary.TranslatedWorldPosition, Primary.SceneDepth, LocalExtraBias);
 			MarkPage(VirtualShadowMapId, MipLevel, Primary.TranslatedWorldPosition, bUsePageDilation, PageDilationOffset);
 		}
 	}
@@ -806,6 +830,13 @@ void GeneratePageFlagsFromFroxelsCS(uint FroxelIndex : SV_DispatchThreadID)
 	const float2 PageDilationDither = float2(
 		(FroxelIndex & 1) ? 1.0f : -1.0f,
 		(FroxelIndex & 2) ? 1.0f : -1.0f);
+
+	float PrimaryExtraBias = 0.0f;
+	float LocalExtraBias = 0.0f;
+	
+	float DOFResolutionBias = CalculateResolutionBiasFromDepthOfField(FroxelCenterView.z);
+	PrimaryExtraBias += DOFResolutionBias;
+	LocalExtraBias += DOFResolutionBias;
 		
 	// Directional lights
 	{
@@ -817,7 +848,7 @@ void GeneratePageFlagsFromFroxelsCS(uint FroxelIndex : SV_DispatchThreadID)
 			int VirtualShadowMapId = DirectionalLightIds[Index];
 			FVirtualShadowMapProjectionShaderData ProjectionData = GetVirtualShadowMapProjectionData(VirtualShadowMapId);
 
-			MarkPageClipmapFroxel(ProjectionData, bUsePageDilation, PageDilationOffset, FroxelCenterTranslatedWorld, FroxelViewAabb, DebugDrawSetup);
+			MarkPageClipmapFroxel(ProjectionData, bUsePageDilation, PageDilationOffset, FroxelCenterTranslatedWorld, FroxelViewAabb, DebugDrawSetup, PrimaryExtraBias);
 		}
 	}
 
@@ -883,7 +914,7 @@ void GeneratePageFlagsFromFroxelsCS(uint FroxelIndex : SV_DispatchThreadID)
 			}
 
 			FVirtualShadowMapProjectionShaderData ProjectionData = GetVirtualShadowMapProjectionData(VirtualShadowMapId);
-			uint MipLevel = GetMipLevelLocal(ProjectionData, FroxelCenterTranslatedWorld, FroxelCenterView.z);
+			uint MipLevel = GetMipLevelLocal(ProjectionData, FroxelCenterTranslatedWorld, FroxelCenterView.z, LocalExtraBias);
 
 			MarkPageFroxel(ProjectionData, bUsePageDilation, PageDilationOffset, FroxelCenterTranslatedWorld, FroxelViewAabb, MipLevel, false, DebugDrawSetup);
 		}
diff --git a/Engine/Shaders/Shared/VirtualShadowMapDefinitions.h b/Engine/Shaders/Shared/VirtualShadowMapDefinitions.h
--- a/Engine/Shaders/Shared/VirtualShadowMapDefinitions.h
+++ b/Engine/Shaders/Shared/VirtualShadowMapDefinitions.h
@@ -122,10 +122,10 @@ struct FNextVirtualShadowMapData
 	int _Padding;
 };
 
-
-inline uint GetMipLevelLocal(float Footprint, uint MipMode, float ShadowMapResolutionLodBias, float GlobalResolutionLodBias)
+//ExtraBias must >= 0
+inline uint GetMipLevelLocal(float Footprint, uint MipMode, float ShadowMapResolutionLodBias, float GlobalResolutionLodBias, float ExtraBias = 0.0f)
 {
-	float MipLevelFloat = log2(Footprint) + ShadowMapResolutionLodBias + GlobalResolutionLodBias;
+	float MipLevelFloat = log2(Footprint) + ShadowMapResolutionLodBias + GlobalResolutionLodBias + ExtraBias;
 	uint MipLevel = uint(max(floor(MipLevelFloat), 0.0f));
 	MipLevel = min(MipLevel, (VSM_MAX_MIP_LEVELS - 1U));
 
diff --git a/Engine/Source/Runtime/Renderer/Private/PostProcess/DiaphragmDOF.cpp b/Engine/Source/Runtime/Renderer/Private/PostProcess/DiaphragmDOF.cpp
--- a/Engine/Source/Runtime/Renderer/Private/PostProcess/DiaphragmDOF.cpp
+++ b/Engine/Source/Runtime/Renderer/Private/PostProcess/DiaphragmDOF.cpp
@@ -685,57 +685,45 @@ BEGIN_SHADER_PARAMETER_STRUCT(FDOFCommonShaderParameters, )
 END_SHADER_PARAMETER_STRUCT()
 
 
-BEGIN_SHADER_PARAMETER_STRUCT(FDOFCocModelShaderParameters, )
-	SHADER_PARAMETER(float, CocInfinityRadius)
-	SHADER_PARAMETER(float, CocInFocusRadius)
-	SHADER_PARAMETER_RDG_TEXTURE(Texture2D, DynamicRadiusOffsetLUT)
-	SHADER_PARAMETER_SAMPLER(SamplerState, DynamicRadiusOffsetLUTSampler)
-	SHADER_PARAMETER(uint32, bCocEnableDynamicRadiusOffset)
-	SHADER_PARAMETER(float, CocMinRadius)
-	SHADER_PARAMETER(float, CocMaxRadius)
-	SHADER_PARAMETER(float, CocSqueeze)
-	SHADER_PARAMETER(float, CocInvSqueeze)
-	SHADER_PARAMETER(float, DepthBlurRadius)
-	SHADER_PARAMETER(float, DepthBlurExponent)
+BEGIN_SHADER_PARAMETER_STRUCT(FDOFTileDecisionParameters, )
+	SHADER_PARAMETER(float, MinGatherRadius)
+	SHADER_PARAMETER(float, SlightOutOfFocusRadiusBoundary)
 END_SHADER_PARAMETER_STRUCT()
 
-void SetCocModelParameters(
-	FRDGBuilder& GraphBuilder,
-	FDOFCocModelShaderParameters* OutParameters,
-	const DiaphragmDOF::FPhysicalCocModel& CocModel,
-	float CocRadiusBasis = 1.0f)
+} // namespace
+
+namespace DiaphragmDOF
 {
-	OutParameters->CocInfinityRadius = CocRadiusBasis * CocModel.InfinityBackgroundCocRadius;
-	OutParameters->CocInFocusRadius = CocRadiusBasis * CocModel.InFocusRadius;
-	OutParameters->bCocEnableDynamicRadiusOffset = CocModel.bEnableDynamicOffset;
-	OutParameters->CocMinRadius = CocRadiusBasis * CocModel.MinForegroundCocRadius;
-	OutParameters->CocMaxRadius = CocRadiusBasis * CocModel.MaxBackgroundCocRadius;
-	OutParameters->CocSqueeze = CocModel.Squeeze;
-	OutParameters->CocInvSqueeze = 1.0f / CocModel.Squeeze;
-	OutParameters->DepthBlurRadius = CocRadiusBasis * CocModel.MaxDepthBlurRadius;
-	OutParameters->DepthBlurExponent = CocModel.DepthBlurExponent;
-
-	if (CocModel.DynamicRadiusOffsetLUT)
+	void SetCocModelParameters(
+		FRDGBuilder& GraphBuilder,
+		FDOFCocModelShaderParameters* OutParameters,
+		const FPhysicalCocModel& CocModel,
+		float CocRadiusBasis)
 	{
-		OutParameters->DynamicRadiusOffsetLUT = GraphBuilder.RegisterExternalTexture(CreateRenderTarget(CocModel.DynamicRadiusOffsetLUT, TEXT("DynamicRadiusOffsetLUT")));
-	}
-	else
-	{
-		const FRDGSystemTextures& SystemTextures = FRDGSystemTextures::Get(GraphBuilder);
-
-		OutParameters->DynamicRadiusOffsetLUT = SystemTextures.Black;
-	}
-
-	OutParameters->DynamicRadiusOffsetLUTSampler = TStaticSamplerState<SF_Bilinear, AM_Clamp, AM_Wrap>::GetRHI();
-}
+		OutParameters->CocInfinityRadius = CocRadiusBasis * CocModel.InfinityBackgroundCocRadius;
+		OutParameters->CocInFocusRadius = CocRadiusBasis * CocModel.InFocusRadius;
+		OutParameters->bCocEnableDynamicRadiusOffset = CocModel.bEnableDynamicOffset;
+		OutParameters->CocMinRadius = CocRadiusBasis * CocModel.MinForegroundCocRadius;
+		OutParameters->CocMaxRadius = CocRadiusBasis * CocModel.MaxBackgroundCocRadius;
+		OutParameters->CocSqueeze = CocModel.Squeeze;
+		OutParameters->CocInvSqueeze = 1.0f / CocModel.Squeeze;
+		OutParameters->DepthBlurRadius = CocRadiusBasis * CocModel.MaxDepthBlurRadius;
+		OutParameters->DepthBlurExponent = CocModel.DepthBlurExponent;
 
+		if (CocModel.DynamicRadiusOffsetLUT)
+		{
+			OutParameters->DynamicRadiusOffsetLUT = GraphBuilder.RegisterExternalTexture(CreateRenderTarget(CocModel.DynamicRadiusOffsetLUT, TEXT("DynamicRadiusOffsetLUT")));
+		}
+		else
+		{
+			const FRDGSystemTextures& SystemTextures = FRDGSystemTextures::Get(GraphBuilder);
 
-BEGIN_SHADER_PARAMETER_STRUCT(FDOFTileDecisionParameters, )
-	SHADER_PARAMETER(float, MinGatherRadius)
-	SHADER_PARAMETER(float, SlightOutOfFocusRadiusBoundary)
-END_SHADER_PARAMETER_STRUCT()
+			OutParameters->DynamicRadiusOffsetLUT = SystemTextures.Black;
+		}
 
-} // namespace
+		OutParameters->DynamicRadiusOffsetLUTSampler = TStaticSamplerState<SF_Bilinear, AM_Clamp, AM_Wrap>::GetRHI();
+	}
+}
 
 
 // ---------------------------------------------------- Shaders
@@ -792,7 +780,7 @@ class FDiaphragmDOFSetupCS : public FDiaphragmDOFShader
 
 	BEGIN_SHADER_PARAMETER_STRUCT(FParameters, )
 		SHADER_PARAMETER_STRUCT_INCLUDE(FDOFCommonShaderParameters, CommonParameters)
-		SHADER_PARAMETER_STRUCT_INCLUDE(FDOFCocModelShaderParameters, CocModel)
+		SHADER_PARAMETER_STRUCT_INCLUDE(DiaphragmDOF::FDOFCocModelShaderParameters, CocModel)
 
 		SHADER_PARAMETER(FVector2f, CocRadiusBasis) // TODO: decompose
 
@@ -963,7 +951,7 @@ class FDiaphragmDOFReduceCS : public FDiaphragmDOFShader
 		SHADER_PARAMETER(float, MinScatteringCocRadius)
 		SHADER_PARAMETER(float, NeighborCompareMaxColor)
 
-		SHADER_PARAMETER_STRUCT_INCLUDE(FDOFCocModelShaderParameters, CocModel)
+		SHADER_PARAMETER_STRUCT_INCLUDE(DiaphragmDOF::FDOFCocModelShaderParameters, CocModel)
 
 		SHADER_PARAMETER(FVector2f, SensorSize)
 		SHADER_PARAMETER(float, Aperture)
@@ -1279,7 +1267,7 @@ BEGIN_SHADER_PARAMETER_STRUCT(FDOFHybridScatterParameters, )
 	SHADER_PARAMETER(float, PetzvalExclusionBoxRadius)
 	
 	SHADER_PARAMETER_STRUCT_INCLUDE(FDOFCommonShaderParameters, CommonParameters)
-	SHADER_PARAMETER_STRUCT_INCLUDE(FDOFCocModelShaderParameters, CocModel)
+	SHADER_PARAMETER_STRUCT_INCLUDE(DiaphragmDOF::FDOFCocModelShaderParameters, CocModel)
 	
 	SHADER_PARAMETER(FVector4f, ScatterOcclusionSize)
 	SHADER_PARAMETER_RDG_TEXTURE(Texture2D, ScatterOcclusion)
@@ -1412,7 +1400,7 @@ class FDiaphragmDOFRecombineCS : public FDiaphragmDOFShader
 	
 	BEGIN_SHADER_PARAMETER_STRUCT(FParameters, )
 		SHADER_PARAMETER_STRUCT_INCLUDE(FDOFCommonShaderParameters, CommonParameters)
-		SHADER_PARAMETER_STRUCT_INCLUDE(FDOFCocModelShaderParameters, CocModel)
+		SHADER_PARAMETER_STRUCT_INCLUDE(DiaphragmDOF::FDOFCocModelShaderParameters, CocModel)
 
 		SHADER_PARAMETER(FIntRect, ViewportRect)
 		SHADER_PARAMETER(FVector4f, ViewportSize)
diff --git a/Engine/Source/Runtime/Renderer/Private/PostProcess/DiaphragmDOF.h b/Engine/Source/Runtime/Renderer/Private/PostProcess/DiaphragmDOF.h
--- a/Engine/Source/Runtime/Renderer/Private/PostProcess/DiaphragmDOF.h
+++ b/Engine/Source/Runtime/Renderer/Private/PostProcess/DiaphragmDOF.h
@@ -190,4 +190,26 @@ RENDERER_API bool AddPasses(
 	const FTranslucencyPassResources& TranslucencyViewResources,
 	FRDGTextureRef& OutputColor);
 
+
+BEGIN_SHADER_PARAMETER_STRUCT(FDOFCocModelShaderParameters, )
+	SHADER_PARAMETER(float, CocInfinityRadius)
+	SHADER_PARAMETER(float, CocInFocusRadius)
+	SHADER_PARAMETER_RDG_TEXTURE(Texture2D, DynamicRadiusOffsetLUT)
+	SHADER_PARAMETER_SAMPLER(SamplerState, DynamicRadiusOffsetLUTSampler)
+	SHADER_PARAMETER(uint32, bCocEnableDynamicRadiusOffset)
+	SHADER_PARAMETER(float, CocMinRadius)
+	SHADER_PARAMETER(float, CocMaxRadius)
+	SHADER_PARAMETER(float, CocSqueeze)
+	SHADER_PARAMETER(float, CocInvSqueeze)
+	SHADER_PARAMETER(float, DepthBlurRadius)
+	SHADER_PARAMETER(float, DepthBlurExponent)
+END_SHADER_PARAMETER_STRUCT()
+
+void SetCocModelParameters(
+	FRDGBuilder& GraphBuilder,
+	FDOFCocModelShaderParameters* OutParameters,
+	const FPhysicalCocModel& CocModel,
+	float CocRadiusBasis = 1.0f);
+
 } // namespace DiaphragmDOF
+
diff --git a/Engine/Source/Runtime/Renderer/Private/VirtualShadowMaps/VirtualShadowMapArray.cpp b/Engine/Source/Runtime/Renderer/Private/VirtualShadowMaps/VirtualShadowMapArray.cpp
--- a/Engine/Source/Runtime/Renderer/Private/VirtualShadowMaps/VirtualShadowMapArray.cpp
+++ b/Engine/Source/Runtime/Renderer/Private/VirtualShadowMaps/VirtualShadowMapArray.cpp
@@ -25,6 +25,7 @@
 #include "SceneCulling/SceneCullingRenderer.h"
 #include "Shadows/ShadowSceneRenderer.h"
 #include "Froxel/Froxel.h"
+#include "PostProcess/DiaphragmDOF.h"
 
 #define DEBUG_ALLOW_STATIC_SEPARATE_WITHOUT_CACHING 0
 
@@ -146,6 +147,15 @@ static TAutoConsoleVariable<int32> CVarFirstPersonPixelRequestLevelClamp(
 	ECVF_Scalability | ECVF_RenderThreadSafe
 );
 
+static TAutoConsoleVariable<float> CVarMaxDOFResolutionBias(
+	TEXT("r.Shadow.Virtual.MaxDOFResolutionBias"),
+	1.0f,
+	TEXT("Determine which pixels are out of focus, and request a lower VSM resolution in those areas.\n")
+	TEXT("Since DOF will blur these anyway, the lowered resolution should not be noticable.\n")
+	TEXT("Set to 0 to turn off this feature. A higher value more aggressively lowers resolution."),
+	ECVF_RenderThreadSafe
+);
+
 static TAutoConsoleVariable<float> CVarPageDilationBorderSizeLocal(
 	TEXT("r.Shadow.Virtual.PageDilationBorderSizeLocal"),
 	0.05f,
@@ -939,6 +949,8 @@ BEGIN_SHADER_PARAMETER_STRUCT(FBaseGeneratePageFlagsParameters, )
 	SHADER_PARAMETER(uint32, MipModeLocal)
 	SHADER_PARAMETER(float, FirstPersonPixelRequestBias)	
 	SHADER_PARAMETER(uint32, FirstPersonPixelRequestLevelClamp)	
+	SHADER_PARAMETER(float, DOFBiasStrength)
+	SHADER_PARAMETER_STRUCT_INCLUDE(DiaphragmDOF::FDOFCocModelShaderParameters, CocModel)
 END_SHADER_PARAMETER_STRUCT()
 
 class FGeneratePageFlagsFromPixelsCS : public FVirtualShadowMapPageManagementShader
@@ -2178,6 +2190,11 @@ void FVirtualShadowMapArray::BuildPageAllocations(
 				BaseParameters.FirstPersonPixelRequestBias = CVarFirstPersonPixelRequestBias.GetValueOnRenderThread();
 				BaseParameters.FirstPersonPixelRequestLevelClamp = CVarFirstPersonPixelRequestLevelClamp.GetValueOnRenderThread();
 
+				DiaphragmDOF::FPhysicalCocModel CocModel;
+				CocModel.Compile(View);
+				DiaphragmDOF::SetCocModelParameters(GraphBuilder, &BaseParameters.CocModel, CocModel, View.ViewRect.Size().X);
+				BaseParameters.DOFBiasStrength = FMath::Max(0.0f, DiaphragmDOF::IsEnabled(View) ? CVarMaxDOFResolutionBias.GetValueOnRenderThread() : 0.0f);
+
 				auto GeneratePageFlags = [&](const EVirtualShadowMapProjectionInputType InputType)
 				{
 					FGeneratePageFlagsFromPixelsCS::FParameters* PassParameters = GraphBuilder.AllocParameters< FGeneratePageFlagsFromPixelsCS::FParameters >();
