Volumetric Cloud - completely reworked the way clouds are reconstructed and upsampled. This is roughly the same cost but  higher quality, with less flickering on fixed image.

#tests PC editor dx12, xsx, EngineTest screenshot
#rnx

diff --git a/Engine/Shaders/Private/DownsampleDepthPixelShader.usf b/Engine/Shaders/Private/DownsampleDepthPixelShader.usf
--- a/Engine/Shaders/Private/DownsampleDepthPixelShader.usf
+++ b/Engine/Shaders/Private/DownsampleDepthPixelShader.usf
@@ -18,6 +18,9 @@
 
 Texture2D<uint> NaniteShadingMask;
 Texture2D<float> DepthTexture;
+#if OUTPUT_MINMAXDEPTH_FROM_MINMAXDEPTH
+Texture2D<float2> MinMaxDepthTexture;
+#endif
 float2 DstToSrcPixelScale;
 float2 SourceMaxUV;
 float2 DestinationResolution;
@@ -36,7 +39,7 @@ void Main(
 	noperspective float4 OutUVAndScreenPos : TEXCOORD0,
 	in FStereoPSInput StereoInput,
 	in float4 SVPos : SV_POSITION,
-#if OUTPUT_MIN_AND_MAX_DEPTH
+#if OUTPUT_MIN_AND_MAX_DEPTH || OUTPUT_MINMAXDEPTH_FROM_MINMAXDEPTH
 	out float4 OutColor : SV_Target0
 #else
 	out float OutDepth : SV_DEPTH
@@ -48,6 +51,19 @@ void Main(
 	// Lower left corner from the input 2x2 quad
 	const int2 PixelPos = float2(DstPixelPos) * DstToSrcPixelScale;
 
+#if OUTPUT_MINMAXDEPTH_FROM_MINMAXDEPTH
+
+	float2 MinMaxDepth0 = MinMaxDepthTexture.Load(uint3(clamp(PixelPos + int2(0, 0), SrcPixelCoordMinAndMax.xy, SrcPixelCoordMinAndMax.zw), 0));
+	float2 MinMaxDepth1 = MinMaxDepthTexture.Load(uint3(clamp(PixelPos + int2(0, 1), SrcPixelCoordMinAndMax.xy, SrcPixelCoordMinAndMax.zw), 0));
+	float2 MinMaxDepth2 = MinMaxDepthTexture.Load(uint3(clamp(PixelPos + int2(1, 1), SrcPixelCoordMinAndMax.xy, SrcPixelCoordMinAndMax.zw), 0));
+	float2 MinMaxDepth3 = MinMaxDepthTexture.Load(uint3(clamp(PixelPos + int2(1, 0), SrcPixelCoordMinAndMax.xy, SrcPixelCoordMinAndMax.zw), 0));
+
+	const float MaxDeviceZ = max(max(MinMaxDepth0.y, MinMaxDepth1.y), max(MinMaxDepth2.y, MinMaxDepth3.y));
+	const float MinDeviceZ = min(min(MinMaxDepth0.x, MinMaxDepth1.x), min(MinMaxDepth2.x, MinMaxDepth3.x));
+	OutColor = float4(MinDeviceZ, MaxDeviceZ, 0.0f, 0.0f);
+
+#else // OUTPUT_MINMAXDEPTH_FROM_MINMAXDEPTH
+
 	float DeviceZ0 = DepthTexture.Load(uint3(clamp(PixelPos + int2(0, 0), SrcPixelCoordMinAndMax.xy, SrcPixelCoordMinAndMax.zw), 0));
 	float DeviceZ1 = DepthTexture.Load(uint3(clamp(PixelPos + int2(0, 1), SrcPixelCoordMinAndMax.xy, SrcPixelCoordMinAndMax.zw), 0));
 	float DeviceZ2 = DepthTexture.Load(uint3(clamp(PixelPos + int2(1, 1), SrcPixelCoordMinAndMax.xy, SrcPixelCoordMinAndMax.zw), 0));
@@ -90,6 +106,8 @@ void Main(
 	OutDepth = Depth;
 
 #endif // OUTPUT_MIN_AND_MAX_DEPTH
+
+#endif // OUTPUT_MINMAXDEPTH_FROM_MINMAXDEPTH
 }
 
 #ifndef STENCIL_LIGHTING_CHANNELS_SHIFT 
diff --git a/Engine/Shaders/Private/VolumetricCloud.usf b/Engine/Shaders/Private/VolumetricCloud.usf
--- a/Engine/Shaders/Private/VolumetricCloud.usf
+++ b/Engine/Shaders/Private/VolumetricCloud.usf
@@ -522,15 +522,6 @@ void MainCommon(in FMaterialPixelParameters MaterialParameters, in float4 SvPosi
 	const float3 RayTranslatedWorldOriginOffsetY = GetTranslatedWorldCameraPosFromView(ResolvedView, SvPosition.xy + float2(0.0f, 1.0f));
 #endif // USE_ANALYTIC_DERIVATIVES
 
-	uint4 TracingCoordToZbufferCoordScaleBias = RenderVolumetricCloudParameters.TracingCoordToZbufferCoordScaleBias;
-	uint2 SceneDepthTextureCoord = uint2(SvPosition.xy - 0.5) * TracingCoordToZbufferCoordScaleBias.xy + TracingCoordToZbufferCoordScaleBias.zw;
-	SceneDepthTextureCoord = clamp(SceneDepthTextureCoord, RenderVolumetricCloudParameters.SceneDepthTextureMinMaxCoord.xy, RenderVolumetricCloudParameters.SceneDepthTextureMinMaxCoord.zw);
-	{
-		float DeviceZ = RenderVolumetricCloudParameters.SceneDepthTexture.Load(uint3(SceneDepthTextureCoord, 0)).r;
-		float3 DepthBufferTranslatedWorldPos = SvPositionToTranslatedWorld(float4(SvPosition.xy, DeviceZ, 1.0));
-		const float TDepthBufferKm = min(length(DepthBufferTranslatedWorldPos * CENTIMETER_TO_KILOMETER - RayTranslatedWorldOriginKm), MaxHalfFloat);
-		OutDepth.xyzw = float4(TDepthBufferKm, TDepthBufferKm, DeviceZ, DeviceZ);
-	}
 
 	////////////////////////////////////////////////////////////////////////////////////////////////////
 	// Begin the cloud ray marching and setup the variable above
@@ -600,6 +591,10 @@ void MainCommon(in FMaterialPixelParameters MaterialParameters, in float4 SvPosi
 		// Sample the depth buffer and update tracing distance
 		//
 
+		uint4 TracingCoordToZbufferCoordScaleBias = RenderVolumetricCloudParameters.TracingCoordToZbufferCoordScaleBias;
+		uint2 SceneDepthTextureCoord = uint2(SvPosition.xy - 0.5) * TracingCoordToZbufferCoordScaleBias.xy + TracingCoordToZbufferCoordScaleBias.zw;
+		SceneDepthTextureCoord = clamp(SceneDepthTextureCoord, RenderVolumetricCloudParameters.SceneDepthTextureMinMaxCoord.xy, RenderVolumetricCloudParameters.SceneDepthTextureMinMaxCoord.zw);
+
 #if CLOUD_MIN_AND_MAX_DEPTH
 
 		float2 MinAndMaxDepth = RenderVolumetricCloudParameters.SceneDepthMinAndMaxTexture.Load(uint3(SceneDepthTextureCoord / 2, 0)).rg;
@@ -611,9 +606,14 @@ void MainCommon(in FMaterialPixelParameters MaterialParameters, in float4 SvPosi
 		float DeviceZ = MinAndMaxDepth.y;
 
 		float3 ClostestDepthBufferTranslatedWorldPos = SvPositionToTranslatedWorld(float4(SvPosition.xy, MinAndMaxDepth.x, 1.0));
-		const float TClosestDepthBufferCm = min(length(ClostestDepthBufferTranslatedWorldPos * CENTIMETER_TO_KILOMETER - RayTranslatedWorldOriginKm), MaxHalfFloat) * KILOMETER_TO_CENTIMETER;
+		const float TClosestDepthBufferKm = min(length(ClostestDepthBufferTranslatedWorldPos * CENTIMETER_TO_KILOMETER - RayTranslatedWorldOriginKm), MaxHalfFloat);
+		const float TClosestDepthBufferCm = TClosestDepthBufferKm * KILOMETER_TO_CENTIMETER;
+		float3 FurthestDepthBufferTranslatedWorldPos = SvPositionToTranslatedWorld(float4(SvPosition.xy, MinAndMaxDepth.y, 1.0));
+		const float TFurthestDepthBufferKm = min(length(FurthestDepthBufferTranslatedWorldPos * CENTIMETER_TO_KILOMETER - RayTranslatedWorldOriginKm), MaxHalfFloat);
 
-		OutDepth.zw = MinAndMaxDepth;
+		OutDepth.xyzw = float4(TFurthestDepthBufferKm, TFurthestDepthBufferKm, TClosestDepthBufferKm, TFurthestDepthBufferKm);
+
+		const float TDepthBufferCm = TFurthestDepthBufferKm * KILOMETER_TO_CENTIMETER;
 
 #else // CLOUD_MIN_AND_MAX_DEPTH
 
@@ -621,14 +621,13 @@ void MainCommon(in FMaterialPixelParameters MaterialParameters, in float4 SvPosi
 	#if HAS_INVERTED_Z_BUFFER
 		DeviceZ = max(0.000000000001, DeviceZ);
 	#endif
-
-#endif // CLOUD_MIN_AND_MAX_DEPTH
-
 		float3 DepthBufferTranslatedWorldPos = SvPositionToTranslatedWorld(float4(SvPosition.xy, DeviceZ, 1.0));
 		const float TDepthBufferKm = min(length(DepthBufferTranslatedWorldPos * CENTIMETER_TO_KILOMETER - RayTranslatedWorldOriginKm), MaxHalfFloat);
+		OutDepth.xyzw = float4(TDepthBufferKm, TDepthBufferKm, TDepthBufferKm, TDepthBufferKm);
+
 		const float TDepthBufferCm = TDepthBufferKm * KILOMETER_TO_CENTIMETER;
 
-		OutDepth.y = TDepthBufferKm;
+#endif // CLOUD_MIN_AND_MAX_DEPTH
 
 		//
 		// Skip tracing if the depth buffer is in front of the cloud layer front interface. If inside the layer, clamp the tracing to the depth buffer
diff --git a/Engine/Shaders/Private/VolumetricRenderTarget.usf b/Engine/Shaders/Private/VolumetricRenderTarget.usf
--- a/Engine/Shaders/Private/VolumetricRenderTarget.usf
+++ b/Engine/Shaders/Private/VolumetricRenderTarget.usf
@@ -33,14 +33,14 @@ struct FDepthData
 {
 	float  CloudFrontDepthFromViewKm;
 	float  SceneDepthFromViewKm;
-	float2 MinMaxDeviceZ;
+	float2 MinMaxViewDepthKm;
 };
 FDepthData GetDepthDataFromVector(float4 DepthVector)
 {
 	FDepthData DepthData;
 	DepthData.CloudFrontDepthFromViewKm = DepthVector.x;
 	DepthData.SceneDepthFromViewKm		= DepthVector.y;
-	DepthData.MinMaxDeviceZ				= DepthVector.zw;
+	DepthData.MinMaxViewDepthKm			= DepthVector.zw;
 	return DepthData;
 }
 
@@ -74,7 +74,8 @@ FDepthData SafeSampleTracingVolumetricDepthTexture(float2 UV)
 
 #if PERMUTATION_HISTORY_AVAILABLE
 Texture2D<float4> PreviousFrameVolumetricTexture; 
-Texture2D<float2> PreviousFrameVolumetricDepthTexture;
+Texture2D<float4> PreviousFrameVolumetricSecondaryTexture;
+Texture2D<float4> PreviousFrameVolumetricDepthTexture;
 float4 PreviousVolumetricTextureSizeAndInvSize;
 
 uint4 PreviousFrameVolumetricTextureValidCoordRect;
@@ -94,14 +95,27 @@ float4 SafeSamplePreviousFrameVolumetricTexture(float2 UV)
 	float4 Result = PreviousFrameVolumetricTexture.SampleLevel(LinearTextureSampler, clamp(UV, PreviousFrameVolumetricTextureValidUvRect.xy, PreviousFrameVolumetricTextureValidUvRect.zw), 0);
 	return float4(Result.rgb * HistoryPreExposureCorrection, Result.a);
 }
+
+float4 SafeLoadPreviousFrameVolumetricSecondaryTexture(uint2 Coord)
+{
+	float4 Result = PreviousFrameVolumetricSecondaryTexture.Load(uint3(clamp(Coord, PreviousFrameVolumetricTextureValidCoordRect.xy, PreviousFrameVolumetricTextureValidCoordRect.zw), 0));
+	return float4(Result.rgb * HistoryPreExposureCorrection, Result.a);
+}
+float4 SafeSamplePreviousFrameVolumetricSecondaryTexture(float2 UV)
+{
+	UV *= PreviousFrameVolumetricTextureUVScale.xy;
+	float4 Result = PreviousFrameVolumetricSecondaryTexture.SampleLevel(LinearTextureSampler, clamp(UV, PreviousFrameVolumetricTextureValidUvRect.xy, PreviousFrameVolumetricTextureValidUvRect.zw), 0);
+	return float4(Result.rgb * HistoryPreExposureCorrection, Result.a);
+}
+
 FDepthData SafeLoadPreviousFrameVolumetricDepthTexture(uint2 Coord)
 {
-	return GetDepthDataFromVector(float4(PreviousFrameVolumetricDepthTexture.Load(uint3(clamp(Coord, PreviousFrameVolumetricTextureValidCoordRect.xy, PreviousFrameVolumetricTextureValidCoordRect.zw), 0)).rg, 0.0, 0.0));
+	return GetDepthDataFromVector(PreviousFrameVolumetricDepthTexture.Load(uint3(clamp(Coord, PreviousFrameVolumetricTextureValidCoordRect.xy, PreviousFrameVolumetricTextureValidCoordRect.zw), 0)));
 }
 FDepthData SafeSamplePreviousFrameVolumetricDepthTexture(float2 UV)
 {
 	UV *= PreviousFrameVolumetricTextureUVScale.xy;
-	return GetDepthDataFromVector(float4(PreviousFrameVolumetricDepthTexture.SampleLevel(LinearTextureSampler, clamp(UV, PreviousFrameVolumetricTextureValidUvRect.xy, PreviousFrameVolumetricTextureValidUvRect.zw), 0).rg, 0.0, 0.0));
+	return GetDepthDataFromVector(PreviousFrameVolumetricDepthTexture.SampleLevel(LinearTextureSampler, clamp(UV, PreviousFrameVolumetricTextureValidUvRect.xy, PreviousFrameVolumetricTextureValidUvRect.zw), 0));
 }
 #endif // PERMUTATION_HISTORY_AVAILABLE
 
@@ -145,9 +159,16 @@ float BOX_NORM_LUMA(float3 Clip, float3 Min, float3 Max)
 }
 
 void ReconstructVolumetricRenderTargetPS(
-	in float4 SVPos : SV_POSITION,
+	in float4 SVPos : SV_POSITION, 
+#if PERMUTATION_CLOUD_MIN_AND_MAX_DEPTH
 	out float4 OutputRt0 : SV_Target0,
-	out float2 OutputRt1 : SV_Target1)
+	out float4 OutputRt1 : SV_Target1,
+	out float4 OutputRt2 : SV_Target2
+#else
+	out float4 OutputRt0 : SV_Target0,
+	out float2 OutputRt1 : SV_Target1
+#endif
+)
 {
 	float2 PixelPos = SVPos.xy;
 	float2 ScreenUV = SVPos.xy * DstVolumetricTextureSizeAndInvSize.zw; // UV in [0,1]
@@ -159,8 +180,9 @@ void ReconstructVolumetricRenderTargetPS(
 	const int XSub = int(IntPixelPos.x) - (IntPixelPosDownsample.x * DownSampleFactor);
 	const int YSub = int(IntPixelPos.y) - (IntPixelPosDownsample.y * DownSampleFactor);
 
-	const bool bUseNewSample = (XSub == CurrentTracingPixelOffset.x) && (YSub == CurrentTracingPixelOffset.y);
+	bool bUseNewSample = (XSub == CurrentTracingPixelOffset.x) && (YSub == CurrentTracingPixelOffset.y);
 	float4 RGBA = 0.0f;
+	float4 RGBA2 = 0.0f;
 	FDepthData DepthData = (FDepthData)0;
   
 	{
@@ -184,26 +206,39 @@ void ReconstructVolumetricRenderTargetPS(
 																							// especially visible when flying through the cloud layer. It is not perfect but will help in lots of cases. 
 																							// The problem when using this method: clouds will look noisier when closer to that distance. 
 
-		#if 0
+
+		
+
+	#if 0
+		// Debug: always use new sample else history result. Should be 
+		if (bUseNewSample)
 		{
-			// Debug: always use new sample else history result. Should be 
-			if (bUseNewSample)
-			{
-				RGBA = SafeLoadTracingVolumetricTexture(int2(SVPos.xy) / DownSampleFactor);
-				//RGBA  = SafeLoadSecondaryTracingVolumetricTexture(int2(SVPos.xy) / DownSampleFactor);
-				DepthData = SafeLoadTracingVolumetricDepthTexture(int2(SVPos.xy) / DownSampleFactor);
-			}
-			else
-			{
-				RGBA = SafeSamplePreviousFrameVolumetricTexture(PrevScreenUVs);
-				//RGBA = SafeSampleSecondaryTracingVolumetricTexture(PrevScreenUVs);
-				DepthData = SafeSamplePreviousFrameVolumetricDepthTexture(PrevScreenUVs);
-			}
-			OutputRt0 = RGBA;
-			OutputRt1 = float2(DepthData.CloudFrontDepthFromViewKm, DepthData.SceneDepthFromViewKm);
-			return;
+			RGBA		= SafeLoadTracingVolumetricTexture(int2(SVPos.xy) / DownSampleFactor);
+			RGBA2		= SafeLoadSecondaryTracingVolumetricTexture(int2(SVPos.xy) / DownSampleFactor);
+			DepthData	= SafeLoadTracingVolumetricDepthTexture(int2(SVPos.xy) / DownSampleFactor);
 		}
-		#endif
+		else
+		{
+			RGBA		= SafeSamplePreviousFrameVolumetricTexture(PrevScreenUVs);
+			RGBA2		= SafeSamplePreviousFrameVolumetricSecondaryTexture(PrevScreenUVs);
+			DepthData	= SafeSamplePreviousFrameVolumetricDepthTexture(PrevScreenUVs);
+		}
+		OutputRt0 = RGBA;
+	#if PERMUTATION_CLOUD_MIN_AND_MAX_DEPTH
+		OutputRt1 = RGBA2;
+		OutputRt2 = float4(
+			DepthData.CloudFrontDepthFromViewKm, 
+			DepthData.SceneDepthFromViewKm, 
+			DepthData.MinMaxViewDepthKm);
+	#else
+		OutputRt1 = float2(
+			DepthData.CloudFrontDepthFromViewKm,
+			DepthData.SceneDepthFromViewKm);
+	#endif
+		return;
+	#endif
+
+
 
 		if (VolumetricRenderTargetMode == 2)
 		{
@@ -232,113 +267,118 @@ void ReconstructVolumetricRenderTargetPS(
 				RGBA = SafeSampleTracingVolumetricTexture(TracingTextureUV);
 				DepthData = SafeSampleTracingVolumetricDepthTexture(TracingTextureUV);
 			}
+		#if PERMUTATION_CLOUD_MIN_AND_MAX_DEPTH
+			OutputRt0 = float4(0.0f, 0.0f, 0.0f, 1.0f);
+			OutputRt1 = float4(0.0f, 0.0f, 0.0f, 1.0f);
+			OutputRt2 = float4(63000.0f, 63000.0f, 63000.0f, 63000.0f);
+		#else
 			OutputRt0 = RGBA;
 			OutputRt1 = float2(DepthData.CloudFrontDepthFromViewKm, DepthData.SceneDepthFromViewKm);
+		#endif
 			return;
 		}
 		else
 		{
-			#define DEBUGCOLORS 0
 
 			int2 CenterSample = int2(SVPos.xy) / DownSampleFactor;
 			float4 NewRGBA = SafeLoadTracingVolumetricTexture(CenterSample);
+			float4 NewRGBA2 = SafeLoadSecondaryTracingVolumetricTexture(CenterSample);
 			FDepthData NewDepthData = SafeLoadTracingVolumetricDepthTexture(CenterSample);
 
+			FDepthData HistoryDepthData = SafeSamplePreviousFrameVolumetricDepthTexture(PrevScreenUVs);
+			const float HistorySceneDepthFromViewCm = HistoryDepthData.SceneDepthFromViewKm * KILOMETER_TO_CENTIMETER;
+
+		#if PERMUTATION_CLOUD_MIN_AND_MAX_DEPTH
+			
+			const bool bLowScreenVelocity = length(ScreenVelocity) < 0.001;
+
+			const float ThresholdToNewSampleKm = 2.0; // Arbitrary
+
+			if (abs(NewDepthData.MinMaxViewDepthKm.x - NewDepthData.MinMaxViewDepthKm.y) > ThresholdToNewSampleKm 
+				&& (abs(NewDepthData.MinMaxViewDepthKm.x - NewDepthData.SceneDepthFromViewKm) < abs(NewDepthData.MinMaxViewDepthKm.y - NewDepthData.SceneDepthFromViewKm)) )
+			{
+				// If there is a huge delta of depth for a pixel, use the new sample but only if we are on the closest depth part.
+				// This helps a lot removing cloud over trees and small details
+				bUseNewSample = true;
+			}
+			else 
+				if (HistoryDepthData.SceneDepthFromViewKm < (NewDepthData.SceneDepthFromViewKm - ThresholdToNewSampleKm))
+			{
+				// History is closer than the near cloud tracing this frame. This means an object had move and an new disocluded area is discovered.
+				// So we simply use the new data from this frame according to the new depth
+				bUseNewSample = true;
+			}
+			else 
+				if (HistoryDepthData.SceneDepthFromViewKm > (NewDepthData.SceneDepthFromViewKm + ThresholdToNewSampleKm))
+			{
+				// An area that just go covered (history is invalid because occluded)
+				bUseNewSample = true;
+			}
+		//	else
+		//		if (NewDepthData.SceneDepthFromViewKm > 1.5)
+		//	{
+		//		if (bUseNewSample || !bValidPreviousUVs)
+		//		{
+		//			RGBA		= SafeLoadTracingVolumetricTexture(int2(SVPos.xy) / DownSampleFactor);
+		//			RGBA2		= SafeLoadSecondaryTracingVolumetricTexture(int2(SVPos.xy) / DownSampleFactor);
+		//			DepthData	= SafeLoadTracingVolumetricDepthTexture(int2(SVPos.xy) / DownSampleFactor);
+		//		}
+		//		else
+		//		{
+		//			RGBA		= SafeSamplePreviousFrameVolumetricTexture(PrevScreenUVs);
+		//			RGBA2		= SafeSamplePreviousFrameVolumetricSecondaryTexture(PrevScreenUVs);
+		//			DepthData	= SafeSamplePreviousFrameVolumetricDepthTexture(PrevScreenUVs);
+		//		}
+		//		OutputRt0 = RGBA;// *float4(1, 0, 0, 1);
+		//		OutputRt1 = RGBA2;// *float4(1, 0, 0, 1);
+		//		OutputRt2 = float4(
+		//			DepthData.CloudFrontDepthFromViewKm, 
+		//			DepthData.SceneDepthFromViewKm, 
+		//			DepthData.MinMaxViewDepthKm);
+		//		return;	
+		//	}
+
 			if (bUseNewSample)
 			{
 				// Load the new sample for this pixel we have just traced
 				RGBA = NewRGBA;
+				RGBA2 = NewRGBA2;
 				DepthData = NewDepthData;
 			}
 			else if (bValidPreviousUVs)
 			{
+				RGBA = SafeSamplePreviousFrameVolumetricTexture(PrevScreenUVs);
+				RGBA2 = SafeSamplePreviousFrameVolumetricSecondaryTexture(PrevScreenUVs);
+				DepthData = SafeSamplePreviousFrameVolumetricDepthTexture(PrevScreenUVs);
 
-
-
-		#if PERMUTATION_CLOUD_MIN_AND_MAX_DEPTH
-
-
-
-				//
-				// NOTE: The best temporal upsampling path reserver for mode 0 when compute is available.
-				//
-
-				const float SceneHalfResDeviceZ = HalfResDepthTexture.Load(int3(SVPos.xy + ViewViewRectMin, 0)).r;
-				const float SceneHalfResDepth = ConvertFromDeviceZ(max(0.000000000001, SceneHalfResDeviceZ));
-
-				const float ThresholdToUpsamplingCm = 100000.0; // super arbitrary :/
-
-				const float NearCloudTracingDepth = ConvertFromDeviceZ(max(0.000000000001, NewDepthData.MinMaxDeviceZ.x));
-				const float FarCloudTracingDepth  = ConvertFromDeviceZ(max(0.000000000001, NewDepthData.MinMaxDeviceZ.y));
-				if (abs(NearCloudTracingDepth - FarCloudTracingDepth) > ThresholdToUpsamplingCm)
-				{
-					// if the min/max depth has a large difference, then always upsample with the min and max traced data
-					float4 NewRGBANear = SafeLoadSecondaryTracingVolumetricTexture(CenterSample);
-					
-					const float LerpFactor = saturate((SceneHalfResDepth - NearCloudTracingDepth) / max(0.00001, FarCloudTracingDepth - NearCloudTracingDepth));
-					RGBA = lerp(NewRGBANear, NewRGBA, LerpFactor);
-
-				//	DepthData = NewDepthData;
-					DepthData.SceneDepthFromViewKm = SceneHalfResDepth * CENTIMETER_TO_KILOMETER;	// TODO trasform from scene depth to view depth.
-					DepthData.CloudFrontDepthFromViewKm = DepthData.SceneDepthFromViewKm;			// Keer far cloud tracing front depth. That should be good.
-						
-			#if DEBUGCOLORS
-					RGBA = float4(1, 1, 0, 0.5);
-			#endif
-				}
-				else
+				if (!(all(IsFinite(RGBA)) && all(IsFinite(float2(DepthData.CloudFrontDepthFromViewKm, DepthData.SceneDepthFromViewKm)))))
 				{
-					// Otherwise we use history 
-					FDepthData HistoryDepthData = SafeSamplePreviousFrameVolumetricDepthTexture(PrevScreenUVs);
-
-					const float HistorySceneDepthFromViewCm = HistoryDepthData.SceneDepthFromViewKm * KILOMETER_TO_CENTIMETER;
-
-					//if (HistorySceneDepthFromViewCm < (NearCloudTracingDepth - ThresholdToUpsamplingCm))
-					if (HistorySceneDepthFromViewCm < (NewDepthData.SceneDepthFromViewKm * KILOMETER_TO_CENTIMETER - ThresholdToUpsamplingCm))
-					{
-						// History is closer than the near cloud tracing this frame. This means an object had move and an new disocluded area is discovered.
-						// So we simply use the new data from this frame according to the new depth
-						
-						RGBA = NewRGBA;
-						DepthData = NewDepthData;
-					//	DepthData.SceneDepthFromViewKm = SceneHalfResDepth * CENTIMETER_TO_KILOMETER;	// TODO trasform from scene depth to view depth.
-					//	DepthData.CloudFrontDepthFromViewKm = DepthData.SceneDepthFromViewKm;			// Keer far cloud tracing front depth. That should be good.
-							
-			#if DEBUGCOLORS
-						RGBA = float4(1, 0, 0, 0.5);
-			#endif
-					}
-					else if (HistorySceneDepthFromViewCm > (NewDepthData.SceneDepthFromViewKm * KILOMETER_TO_CENTIMETER + ThresholdToUpsamplingCm))
-					{
-						// An area that just go covered (history is invalid because occluded)
-						RGBA = NewRGBA;
-						DepthData = NewDepthData;
-					//	DepthData.SceneDepthFromViewKm = SceneHalfResDepth * CENTIMETER_TO_KILOMETER;	// TODO trasform from scene depth to view depth.
-					//	DepthData.CloudFrontDepthFromViewKm = DepthData.SceneDepthFromViewKm;			// Keer far cloud tracing front depth. That should be good.
-
-			#if DEBUGCOLORS
-						RGBA = float4(0, 1, 0, 0.5);
-			#endif
-					}
-					else
-					{
-						// Good reprojection
-						float4 HistoryRGBA = SafeSamplePreviousFrameVolumetricTexture(PrevScreenUVs);
-						RGBA = HistoryRGBA;
-						DepthData = HistoryDepthData;
-							
-			#if DEBUGCOLORS
-						RGBA = float4(0, 0, 1, 0.5);
-			#endif
-					}
+					RGBA = float4(0.0f, 0.0f, 0.0f, 1.0f);
+					RGBA2 = float4(0.0f, 0.0f, 0.0f, 1.0f);
+					DepthData.CloudFrontDepthFromViewKm = 1000.0f;
+					DepthData.SceneDepthFromViewKm = 1000.0f;
+					DepthData.MinMaxViewDepthKm = 1000.0f.xx;
 				}
-
-
+			}
+			else // !bValidPreviousUVs
+			{
+				// History is invalid so simply use this frame low resolution render with bilinear sampling.
+				// Single sampel of the far data seem sto always be enough
+				RGBA = SafeSampleTracingVolumetricTexture(ScreenUV);
+				RGBA2 = SafeSampleSecondaryTracingVolumetricTexture(ScreenUV);
+				DepthData = SafeSampleTracingVolumetricDepthTexture(ScreenUV);
+			}
 
 		#else // PERMUTATION_CLOUD_MIN_AND_MAX_DEPTH
 
-
-
+			if (bUseNewSample)
+			{
+				// Load the new sample for this pixel we have just traced
+				RGBA = NewRGBA;
+				DepthData = NewDepthData;
+			}
+			else if (bValidPreviousUVs)
+			{
 				//
 				// NOTE: This path is use when mode0 is not used or when compute is not available (min max depth permutation are only generated for compute path)
 				//
@@ -420,12 +460,6 @@ void ReconstructVolumetricRenderTargetPS(
 				#endif
 				}
 
-
-
-		#endif // PERMUTATION_CLOUD_MIN_AND_MAX_DEPTH
-
-
-
 				if (!(all(IsFinite(RGBA)) && all(IsFinite(float2(DepthData.CloudFrontDepthFromViewKm, DepthData.SceneDepthFromViewKm)))))
 				{
 					RGBA = float4(0.0f, 0.0f, 0.0f, 1.0f);
@@ -440,18 +474,35 @@ void ReconstructVolumetricRenderTargetPS(
 				RGBA = SafeSampleTracingVolumetricTexture(ScreenUV);
 				DepthData = SafeSampleTracingVolumetricDepthTexture(ScreenUV);
 			}
+
+		#endif // PERMUTATION_CLOUD_MIN_AND_MAX_DEPTH
+
+
 		}
 	}
 
 	OutputRt0 = RGBA;
+#if PERMUTATION_CLOUD_MIN_AND_MAX_DEPTH
+	OutputRt1 = RGBA2;
+	OutputRt2 = float4(DepthData.CloudFrontDepthFromViewKm, DepthData.SceneDepthFromViewKm, DepthData.MinMaxViewDepthKm);
+#else // PERMUTATION_CLOUD_MIN_AND_MAX_DEPTH
 	OutputRt1 = float2(DepthData.CloudFrontDepthFromViewKm, DepthData.SceneDepthFromViewKm);
+#endif // PERMUTATION_CLOUD_MIN_AND_MAX_DEPTH
 
 #else // PERMUTATION_HISTORY_AVAILABLE
 
 	// Simple bilinear upsample
 	OutputRt0 = SafeSampleTracingVolumetricTexture(ScreenUV);
+
+#if PERMUTATION_CLOUD_MIN_AND_MAX_DEPTH
+	OutputRt1 = SafeSampleSecondaryTracingVolumetricTexture(ScreenUV);
+
+	FDepthData DepthData = SafeSampleTracingVolumetricDepthTexture(ScreenUV);
+	OutputRt2 = float4(DepthData.CloudFrontDepthFromViewKm, DepthData.SceneDepthFromViewKm, DepthData.MinMaxViewDepthKm);
+#else // PERMUTATION_CLOUD_MIN_AND_MAX_DEPTH
 	FDepthData DepthData = SafeSampleTracingVolumetricDepthTexture(ScreenUV);
 	OutputRt1 = float2(DepthData.CloudFrontDepthFromViewKm, DepthData.SceneDepthFromViewKm);
+#endif // PERMUTATION_CLOUD_MIN_AND_MAX_DEPTH
 
 #endif // PERMUTATION_HISTORY_AVAILABLE
 }
@@ -477,11 +528,13 @@ void ReconstructVolumetricRenderTargetPS(
 SamplerState LinearTextureSampler;
 
 Texture2D<float4> VolumetricTexture;
-Texture2D<float2> VolumetricDepthTexture;
+Texture2D<float4> VolumetricSecondaryTexture;
+Texture2D<float4> VolumetricDepthTexture;
 uint4 VolumetricTextureValidCoordRect;
 float4 VolumetricTextureValidUvRect;
 uint ForwardShadingEnable;
 uint OutputAlphaHoldout;
+uint MinMaxDepthAvailable;
 
 #if PERMUTATION_RENDER_UNDERWATER_BUFFER || PERMUTATION_COMPOSE_WITH_WATER
 Texture2D WaterLinearDepthTexture;
@@ -532,15 +585,25 @@ float4 SafeSampleVolumetricTexture(float2 UV)
 	const float2 WrappedUV = WrapUVsForInstancedViews(UV, VolumetricTextureValidUvRect);
 	return VolumetricTexture.SampleLevel(LinearTextureSampler, clamp(WrappedUV, VolumetricTextureValidUvRect.xy, VolumetricTextureValidUvRect.zw), 0);
 }
-float2 SafeLoadVolumetricDepthTexture(uint2 Coord)
+float4 SafeLoadVolumetricSecondaryTexture(uint2 Coord)
 {
 	const uint2 WrappedCoords = WrapCoordsForInstancedViews(Coord, VolumetricTextureValidCoordRect);
-	return VolumetricDepthTexture.Load(uint3(clamp(WrappedCoords, VolumetricTextureValidCoordRect.xy, VolumetricTextureValidCoordRect.zw), 0)).rg;
+	return VolumetricSecondaryTexture.Load(uint3(clamp(WrappedCoords, VolumetricTextureValidCoordRect.xy, VolumetricTextureValidCoordRect.zw), 0));
 }
-float2 SafeSampleVolumetricDepthTexture(float2 UV)
+float4 SafeSampleVolumetricSecondaryTexture(float2 UV)
 {
 	const float2 WrappedUV = WrapUVsForInstancedViews(UV, VolumetricTextureValidUvRect);
-	return VolumetricDepthTexture.SampleLevel(LinearTextureSampler, clamp(WrappedUV, VolumetricTextureValidUvRect.xy, VolumetricTextureValidUvRect.zw), 0).rg;
+	return VolumetricSecondaryTexture.SampleLevel(LinearTextureSampler, clamp(WrappedUV, VolumetricTextureValidUvRect.xy, VolumetricTextureValidUvRect.zw), 0);
+}
+float4 SafeLoadVolumetricDepthTexture(uint2 Coord)
+{
+	const uint2 WrappedCoords = WrapCoordsForInstancedViews(Coord, VolumetricTextureValidCoordRect);
+	return VolumetricDepthTexture.Load(uint3(clamp(WrappedCoords, VolumetricTextureValidCoordRect.xy, VolumetricTextureValidCoordRect.zw), 0));
+}
+float4 SafeSampleVolumetricDepthTexture(float2 UV)
+{
+	const float2 WrappedUV = WrapUVsForInstancedViews(UV, VolumetricTextureValidUvRect);
+	return VolumetricDepthTexture.SampleLevel(LinearTextureSampler, clamp(WrappedUV, VolumetricTextureValidUvRect.xy, VolumetricTextureValidUvRect.zw), 0);
 }
 
 float4 VolumetricTextureSizeAndInvSize;
@@ -635,12 +698,12 @@ void ComposeVolumetricRTOverScenePS(
 
 
 	//Make the offset be independent of aspect ratio, resolution scale, downsampling
-	const float2 FullResOffsetUVScale =	float2(1.0f,View.BufferSizeAndInvSize.x * View.BufferSizeAndInvSize.w)	// Aspect ratio correction
-										* View.BufferSizeAndInvSize.zw											// Pixel size
-										* FullResolutionToVolumetricBufferResolutionScale.y;					// Volumetric buffer downsample factor
-	
-	float2 Offset0Sample = (float2(Rand3DPCG16(int3(CurResPixelCoord, View.StateFrameIndexMod8     )).xy) * rcp(65536.0)) * 2.0f - 1.0f;
-	float2 Offset1Sample = (float2(Rand3DPCG16(int3(CurResPixelCoord, View.StateFrameIndexMod8 + 8 )).xy) * rcp(65536.0)) * 2.0f - 1.0f;
+	const float2 FullResOffsetUVScale = float2(1.0f, View.BufferSizeAndInvSize.x * View.BufferSizeAndInvSize.w)	// Aspect ratio correction
+		* View.BufferSizeAndInvSize.zw											// Pixel size
+		* FullResolutionToVolumetricBufferResolutionScale.y;					// Volumetric buffer downsample factor
+
+	float2 Offset0Sample = (float2(Rand3DPCG16(int3(CurResPixelCoord, View.StateFrameIndexMod8)).xy) * rcp(65536.0)) * 2.0f - 1.0f;
+	float2 Offset1Sample = (float2(Rand3DPCG16(int3(CurResPixelCoord, View.StateFrameIndexMod8 + 8)).xy) * rcp(65536.0)) * 2.0f - 1.0f;
 	float2 Offset2Sample = (float2(Rand3DPCG16(int3(CurResPixelCoord, View.StateFrameIndexMod8 + 16)).xy) * rcp(65536.0)) * 2.0f - 1.0f;
 	float2 Offset3Sample = (float2(Rand3DPCG16(int3(CurResPixelCoord, View.StateFrameIndexMod8 + 32)).xy) * rcp(65536.0)) * 2.0f - 1.0f;
 	Offset0Sample = normalize(Offset0Sample);
@@ -661,7 +724,8 @@ void ComposeVolumetricRTOverScenePS(
 
 #if PERMUTATION_UPSAMPLINGMODE==0
 	// Single bilinear sample
-	OutputRt0 = SafeLoadVolumetricTexture(VolumeUV * VolumetricTextureSizeAndInvSize.xy); // SafeSampleVolumetricTexture(VolumeUV);
+//	OutputRt0 = SafeLoadVolumetricTexture(VolumeUV * VolumetricTextureSizeAndInvSize.xy);			// Closest
+	OutputRt0 = SafeLoadVolumetricSecondaryTexture(VolumeUV * VolumetricTextureSizeAndInvSize.xy);	// Furthest
 	float2 VolumeFrontDepth0 = SafeSampleVolumetricDepthTexture(VolumeUV);
 	OutputRt0 = ApplyFogToCloudSampleFromVPosAndDepth(CurResPixelCoord, VolumeFrontDepth0.x, OutputRt0);
 	return;
@@ -681,7 +745,7 @@ void ComposeVolumetricRTOverScenePS(
 	float4 Data3 = SafeSampleVolumetricTexture(VolumeUVOffset3);
 	Data3 = ApplyFogToCloudFromUVs(VolumeUVOffset3, Data3);
 
-	OutputRt0 = 0.25 * (Data0 + Data1 + Data2 + Data3); 
+	OutputRt0 = 0.25 * (Data0 + Data1 + Data2 + Data3);
 	return;
 
 #elif (PERMUTATION_UPSAMPLINGMODE==4 || PERMUTATION_UPSAMPLINGMODE==3 || PERMUTATION_UPSAMPLINGMODE==2)
@@ -698,11 +762,14 @@ void ComposeVolumetricRTOverScenePS(
 	float3 TranslatedWorldPosition = SvPositionToTranslatedWorld(float4(SVPos.xy, 0.5, 1.0));
 	TranslatedWorldPosition = normalize(TranslatedWorldPosition - PrimaryView.TranslatedWorldCameraOrigin) * PixelLinearDepth + PrimaryView.TranslatedWorldCameraOrigin;
 
-	float4 ClipPosition = mul(float4(TranslatedWorldPosition,1.0), PrimaryView.TranslatedWorldToClip);
+	float4 ClipPosition = mul(float4(TranslatedWorldPosition, 1.0), PrimaryView.TranslatedWorldToClip);
 	ClipPosition /= ClipPosition.w;
 	float PixelDeviceZ = ClipPosition.z;
+#if HAS_INVERTED_Z_BUFFER
+	PixelDeviceZ = max(0.000000000001, PixelDeviceZ);
+#endif
 
-	float3 ScreenTranslatedWorldPosition = SvPositionToTranslatedWorld(float4(SVPos.xy, ClipPosition.z, 1.0));
+	float3 ScreenTranslatedWorldPosition = SvPositionToTranslatedWorld(float4(SVPos.xy, PixelDeviceZ, 1.0));
 
 #else
 
@@ -710,11 +777,14 @@ void ComposeVolumetricRTOverScenePS(
 	float PixelDeviceZ = MSAADepthTexture.Load(int2(SVPos.xy), SampleIndex).x;
 #else
 	float PixelDeviceZ = SceneTexturesStruct.SceneDepthTexture.Load(uint3(SVPos.xy, 0)).r;
+#endif
+#if HAS_INVERTED_Z_BUFFER
+	PixelDeviceZ = max(0.000000000001, PixelDeviceZ);
 #endif
 	float3 ScreenTranslatedWorldPosition = SvPositionToTranslatedWorld(float4(SVPos.xy, PixelDeviceZ, 1.0));
 #endif 
 
-	float PixelDistanceFromView   = length(PrimaryView.TranslatedWorldCameraOrigin - ScreenTranslatedWorldPosition);
+	float PixelDistanceFromView = length(PrimaryView.TranslatedWorldCameraOrigin - ScreenTranslatedWorldPosition);
 	float PixelDistanceFromViewKm = PixelDistanceFromView * CENTIMETER_TO_KILOMETER;
 
 
@@ -802,7 +872,7 @@ void ComposeVolumetricRTOverScenePS(
 	// Bilateral upsampling
 
 	int2 PixelPos = SVPos.xy - View.ViewRectMin.xy;
-	int2 VolumeCoordUInt = PixelPos  / int(FullResolutionToVolumetricBufferResolutionScale.y);
+	int2 VolumeCoordUInt = PixelPos / int(FullResolutionToVolumetricBufferResolutionScale.y);
 	int OffsetX = (VolumeCoordUInt.x * int(FullResolutionToVolumetricBufferResolutionScale.y)) == PixelPos.x ? -1 : 1;
 	int OffsetY = (VolumeCoordUInt.y * int(FullResolutionToVolumetricBufferResolutionScale.y)) == PixelPos.y ? -1 : 1;
 
@@ -817,13 +887,13 @@ void ComposeVolumetricRTOverScenePS(
 		VolumeUVOffset3 = WaterVolumeUV + Offset3;
 
 		float4 VolumeRGBT0 = SafeLoadVolumetricTexture(VolumeUVOffset0 * VolumetricTextureSizeAndInvSize.xy);
-		float2 VolumeFrontDepth0 = SafeLoadVolumetricDepthTexture(VolumeUVOffset0 * VolumetricTextureSizeAndInvSize.xy);
+		float4 VolumeFrontDepth0 = SafeLoadVolumetricDepthTexture(VolumeUVOffset0 * VolumetricTextureSizeAndInvSize.xy);
 		float4 VolumeRGBT1 = SafeLoadVolumetricTexture(VolumeUVOffset1 * VolumetricTextureSizeAndInvSize.xy);
-		float2 VolumeFrontDepth1 = SafeLoadVolumetricDepthTexture(VolumeUVOffset1 * VolumetricTextureSizeAndInvSize.xy);
+		float4 VolumeFrontDepth1 = SafeLoadVolumetricDepthTexture(VolumeUVOffset1 * VolumetricTextureSizeAndInvSize.xy);
 		float4 VolumeRGBT2 = SafeLoadVolumetricTexture(VolumeUVOffset2 * VolumetricTextureSizeAndInvSize.xy);
-		float2 VolumeFrontDepth2 = SafeLoadVolumetricDepthTexture(VolumeUVOffset2 * VolumetricTextureSizeAndInvSize.xy);
+		float4 VolumeFrontDepth2 = SafeLoadVolumetricDepthTexture(VolumeUVOffset2 * VolumetricTextureSizeAndInvSize.xy);
 		float4 VolumeRGBT3 = SafeLoadVolumetricTexture(VolumeUVOffset3 * VolumetricTextureSizeAndInvSize.xy);
-		float2 VolumeFrontDepth3 = SafeLoadVolumetricDepthTexture(VolumeUVOffset3 * VolumetricTextureSizeAndInvSize.xy);
+		float4 VolumeFrontDepth3 = SafeLoadVolumetricDepthTexture(VolumeUVOffset3 * VolumetricTextureSizeAndInvSize.xy);
 
 		float ValidSampleCount = 0.0f;
 		float4 DataAcc = 0.0f;
@@ -875,12 +945,13 @@ void ComposeVolumetricRTOverScenePS(
 	return;
 #else // PERMUTATION_RENDER_UNDERWATER_BUFFER
 
+
 	// We only want to run the special "cloud over water" code path on actual water pixels, otherwise depth threshold can cause visual issues.
 	// However, with Forward Shading, we do not have access to the gbuffer: so we still run that simple water depth test to occlude cloud on forrground meshes. Not the perfect but it will work in most cases.
 #if SUBTRATE_GBUFFER_FORMAT==1
 	FSubstrateAddressing SubstrateAddressing = GetSubstratePixelDataByteOffset(PixelPos, uint2(View.BufferSizeAndInvSize.xy), Substrate.MaxBytesPerPixel);
 	FSubstratePixelHeader SubstratePixelHeader = UnpackSubstrateHeaderIn(Substrate.MaterialTextureArray, SubstrateAddressing, Substrate.TopLayerTexture);
-	const bool bIsWaterPixel = SubstratePixelHeader.IsSingleLayerWater(); 
+	const bool bIsWaterPixel = SubstratePixelHeader.IsSingleLayerWater();
 #else
 	FScreenSpaceData ScreenSpaceData = GetScreenSpaceData(ScreenUV);
 	const bool bIsWaterPixel = ForwardShadingEnable || ScreenSpaceData.GBuffer.ShadingModelID == SHADINGMODELID_SINGLELAYERWATER;
@@ -891,84 +962,215 @@ void ComposeVolumetricRTOverScenePS(
 	uint2 VolumeCoord2 = max(0, int2(VolumeCoordUInt) + int2(OffsetX, OffsetY));
 	uint2 VolumeCoord3 = max(0, int2(VolumeCoordUInt) + int2(0, OffsetY));
 
-	float2 VolumeFrontDepth0 = SafeLoadVolumetricDepthTexture(VolumeCoord0);
-	float2 VolumeFrontDepth1 = SafeLoadVolumetricDepthTexture(VolumeCoord1);
-	float2 VolumeFrontDepth2 = SafeLoadVolumetricDepthTexture(VolumeCoord2);
-	float2 VolumeFrontDepth3 = SafeLoadVolumetricDepthTexture(VolumeCoord3);
+	float4 VolumeFrontDepth0 = SafeLoadVolumetricDepthTexture(VolumeCoord0);
+	float4 VolumeFrontDepth1 = SafeLoadVolumetricDepthTexture(VolumeCoord1);
+	float4 VolumeFrontDepth2 = SafeLoadVolumetricDepthTexture(VolumeCoord2);
+	float4 VolumeFrontDepth3 = SafeLoadVolumetricDepthTexture(VolumeCoord3);
 
-	// Check that cloud medium sample are in front of the surface we are upsampling over water to make sure we will have valid sample.
-	// This is especially needed when upsampling cloud over water surface where depth is not matching the depth for which clouds have been traced for.
-	// If samples are behind, we assume clouds should not be visible (it will be a harsh transition for now)
-	const bool bAllCloudSamplesInFrontOfWater = all(float4(VolumeFrontDepth0.x, VolumeFrontDepth1.x, VolumeFrontDepth2.x, VolumeFrontDepth3.x) < PixelDistanceFromViewKm);
-	const bool bAnyCloudSamplesInFrontOfWater = any(float4(VolumeFrontDepth0.x, VolumeFrontDepth1.x, VolumeFrontDepth2.x, VolumeFrontDepth3.x) < PixelDistanceFromViewKm);
 
-	// clamp to the maximum of the read depth to avoid no depth matching
-	PixelDistanceFromViewKm = min(PixelDistanceFromViewKm, max(max(VolumeFrontDepth0.y, VolumeFrontDepth1.y),max(VolumeFrontDepth2.y, VolumeFrontDepth3.y)));
 
-	float4 Depth4 = float4(VolumeFrontDepth0.y, VolumeFrontDepth1.y, VolumeFrontDepth2.y, VolumeFrontDepth3.y);
-	float4 Depth4Diff = PixelDistanceFromViewKm - Depth4;
-	Depth4Diff = abs(Depth4Diff);
-	float MaxDepth4Diff = max(max(Depth4Diff.x, Depth4Diff.y), max(Depth4Diff.z, Depth4Diff.w));
+#if 0
+//	OutputRt0 = SafeLoadVolumetricSecondaryTexture(VolumeCoord0);
+	OutputRt0 = SafeLoadVolumetricTexture(VolumeUVOffset0 * VolumetricTextureSizeAndInvSize.xy);
+	return;
+#endif
 
-	float ValidSampleCount = 0;
-	float4 DataAcc = 0;
 
-	const float WeightMultiplier = 1000.0f;
-	const float ThresholdToBilinear = PixelDistanceFromViewKm * 0.1;
 
-	if (MaxDepth4Diff > ThresholdToBilinear)
-	{
-		float4 VolumeRGBT0 = SafeLoadVolumetricTexture(VolumeCoord0);
-		float4 VolumeRGBT1 = SafeLoadVolumetricTexture(VolumeCoord1);
-		float4 VolumeRGBT2 = SafeLoadVolumetricTexture(VolumeCoord2);
-		float4 VolumeRGBT3 = SafeLoadVolumetricTexture(VolumeCoord3);
+	if (MinMaxDepthAvailable == 0)
+	{	// ==> This path is when compute shader is disable and so no minmax depth texture is available. It is also used by VRT.Mode=1 for upsampling based on front depth only.
+		
+		// Check that cloud medium sample are in front of the surface we are upsampling over water to make sure we will have valid sample.
+		// This is especially needed when upsampling cloud over water surface where depth is not matching the depth for which clouds have been traced for.
+		// If samples are behind, we assume clouds should not be visible (it will be a harsh transition for now)
+		const bool bAllCloudSamplesInFrontOfWater = all(float4(VolumeFrontDepth0.x, VolumeFrontDepth1.x, VolumeFrontDepth2.x, VolumeFrontDepth3.x) < PixelDistanceFromViewKm);
+		const bool bAnyCloudSamplesInFrontOfWater = any(float4(VolumeFrontDepth0.x, VolumeFrontDepth1.x, VolumeFrontDepth2.x, VolumeFrontDepth3.x) < PixelDistanceFromViewKm);
+
+		// clamp to the maximum of the read depth to avoid no depth matching
+		PixelDistanceFromViewKm = min(PixelDistanceFromViewKm, max(max(VolumeFrontDepth0.y, VolumeFrontDepth1.y), max(VolumeFrontDepth2.y, VolumeFrontDepth3.y)));
 
-		VolumeRGBT0 = ApplyFogToCloudFromUVsAndDepth(VolumeCoord0 * VolumetricTextureSizeAndInvSize.zw, VolumeFrontDepth0.x, VolumeRGBT0);
-		VolumeRGBT1 = ApplyFogToCloudFromUVsAndDepth(VolumeCoord1 * VolumetricTextureSizeAndInvSize.zw, VolumeFrontDepth1.x, VolumeRGBT1);
-		VolumeRGBT2 = ApplyFogToCloudFromUVsAndDepth(VolumeCoord2 * VolumetricTextureSizeAndInvSize.zw, VolumeFrontDepth2.x, VolumeRGBT2);
-		VolumeRGBT3 = ApplyFogToCloudFromUVsAndDepth(VolumeCoord3 * VolumetricTextureSizeAndInvSize.zw, VolumeFrontDepth3.x, VolumeRGBT3);
+		float4 Depth4 = float4(VolumeFrontDepth0.y, VolumeFrontDepth1.y, VolumeFrontDepth2.y, VolumeFrontDepth3.y);
+		float4 Depth4Diff = PixelDistanceFromViewKm - Depth4;
+		Depth4Diff = abs(Depth4Diff);
+		float MaxDepth4Diff = max(max(Depth4Diff.x, Depth4Diff.y), max(Depth4Diff.z, Depth4Diff.w));
 
-		if (bIsWaterPixel) 
+		float ValidSampleCount = 0;
+		float4 DataAcc = 0;
+		const float WeightMultiplier = 1000.0f;
+		const float ThresholdToBilinear = PixelDistanceFromViewKm * 0.1;
+		if (MaxDepth4Diff > ThresholdToBilinear)
 		{
-			// Only do the following is the pixel is water
-			const float ConsideredFarWaterDistanceKm = 1.0f;				// We only apply the fix for pixel that are further than 1km
-			const float NotRenderedPixelDistanceKm = 500.0f;				// We assume pixel have a depth greater than that such distant depth pixel
-			if (any(Depth4 > NotRenderedPixelDistanceKm)					// Some pixels have non rendered pixel...
-				&& !all(Depth4 > NotRenderedPixelDistanceKm)				// ...but not all (in this case we want to fall back to bilinear filtering).
-				&& PixelDistanceFromViewKm > ConsideredFarWaterDistanceKm)	// Only treat this way pixel that are far enough to not mix up close object with far water.
+			float4 VolumeRGBT0 = SafeLoadVolumetricTexture(VolumeCoord0);
+			float4 VolumeRGBT1 = SafeLoadVolumetricTexture(VolumeCoord1);
+			float4 VolumeRGBT2 = SafeLoadVolumetricTexture(VolumeCoord2);
+			float4 VolumeRGBT3 = SafeLoadVolumetricTexture(VolumeCoord3);
+
+			VolumeRGBT0 = ApplyFogToCloudFromUVsAndDepth(VolumeCoord0 * VolumetricTextureSizeAndInvSize.zw, VolumeFrontDepth0.x, VolumeRGBT0);
+			VolumeRGBT1 = ApplyFogToCloudFromUVsAndDepth(VolumeCoord1 * VolumetricTextureSizeAndInvSize.zw, VolumeFrontDepth1.x, VolumeRGBT1);
+			VolumeRGBT2 = ApplyFogToCloudFromUVsAndDepth(VolumeCoord2 * VolumetricTextureSizeAndInvSize.zw, VolumeFrontDepth2.x, VolumeRGBT2);
+			VolumeRGBT3 = ApplyFogToCloudFromUVsAndDepth(VolumeCoord3 * VolumetricTextureSizeAndInvSize.zw, VolumeFrontDepth3.x, VolumeRGBT3);
+
+			if (bIsWaterPixel)
 			{
-				// This is a special case / hack added for to fix some visual issues encoutnered with water as of today.
-				// Water is render after checker boarded minmax depth is taken for cloud rendering (because clouds also needs to be rendered within the water pass).
-				// As such, the reconstructure fails a taking a matching color for a matching depth (cloud depth = 1000km but water is closer, less than 2 kilometers).
-				// This does not handle all the cases but fixes current issues with the water system.
-				// ==> this code can be removed when we automatically fix up edges of objects to have sharp silhouette under strong conflicting motion.
-				const float AcceptanceOffsetKm = 0.001f;
-				if (VolumeFrontDepth0.y > (PixelDistanceFromViewKm - AcceptanceOffsetKm))
-				{
-					DataAcc += VolumeRGBT0;
-					ValidSampleCount++;
-				}
-				if (VolumeFrontDepth1.y > (PixelDistanceFromViewKm - AcceptanceOffsetKm))
-				{
-					DataAcc += VolumeRGBT1;
-					ValidSampleCount++;
-				}
-				if (VolumeFrontDepth2.y > (PixelDistanceFromViewKm - AcceptanceOffsetKm))
-				{
-					DataAcc += VolumeRGBT2;
-					ValidSampleCount++;
-				}
-				if (VolumeFrontDepth3.y > (PixelDistanceFromViewKm - AcceptanceOffsetKm))
+				// Only do the following is the pixel is water
+				const float ConsideredFarWaterDistanceKm = 1.0f;				// We only apply the fix for pixel that are further than 1km
+				const float NotRenderedPixelDistanceKm = 500.0f;				// We assume pixel have a depth greater than that such distant depth pixel
+				if (any(Depth4 > NotRenderedPixelDistanceKm)					// Some pixels have non rendered pixel...
+					&& !all(Depth4 > NotRenderedPixelDistanceKm)				// ...but not all (in this case we want to fall back to bilinear filtering).
+					&& PixelDistanceFromViewKm > ConsideredFarWaterDistanceKm)	// Only treat this way pixel that are far enough to not mix up close object with far water.
 				{
-					DataAcc += VolumeRGBT3;
-					ValidSampleCount++;
+					// This is a special case / hack added for to fix some visual issues encoutnered with water as of today.
+					// Water is render after checker boarded minmax depth is taken for cloud rendering (because clouds also needs to be rendered within the water pass).
+					// As such, the reconstructure fails a taking a matching color for a matching depth (cloud depth = 1000km but water is closer, less than 2 kilometers).
+					// This does not handle all the cases but fixes current issues with the water system.
+					// ==> this code can be removed when we automatically fix up edges of objects to have sharp silhouette under strong conflicting motion.
+					const float AcceptanceOffsetKm = 0.001f;
+					if (VolumeFrontDepth0.y > (PixelDistanceFromViewKm - AcceptanceOffsetKm))
+					{
+						DataAcc += VolumeRGBT0;
+						ValidSampleCount++;
+					}
+					if (VolumeFrontDepth1.y > (PixelDistanceFromViewKm - AcceptanceOffsetKm))
+					{
+						DataAcc += VolumeRGBT1;
+						ValidSampleCount++;
+					}
+					if (VolumeFrontDepth2.y > (PixelDistanceFromViewKm - AcceptanceOffsetKm))
+					{
+						DataAcc += VolumeRGBT2;
+						ValidSampleCount++;
+					}
+					if (VolumeFrontDepth3.y > (PixelDistanceFromViewKm - AcceptanceOffsetKm))
+					{
+						DataAcc += VolumeRGBT3;
+						ValidSampleCount++;
+					}
 				}
-				//DataAcc = float4(0.0, 0.0, 1.0, 0); 
+			}
+
+			if (!bIsWaterPixel || (bIsWaterPixel && ValidSampleCount == 0 && bAnyCloudSamplesInFrontOfWater))
+			{
+				// Depth discontinuities edges
+				float4 weights = 1.0f / (Depth4Diff * WeightMultiplier + 1.0f);
+				const float weightsSum = dot(weights, float4(1.0f, 1.0f, 1.0f, 1.0f));
+				weights /= weightsSum;
+
+				ValidSampleCount = weightsSum > 0.0f ? 1.0 : 0.0f;
+				DataAcc = weights.x * VolumeRGBT0 + weights.y * VolumeRGBT1 + weights.z * VolumeRGBT2 + weights.w * VolumeRGBT3;
 			}
 		}
+		else
+		{
+			// Now do a bilinear sample to have a soft look region without depth edges.
+			ValidSampleCount = 1.0;
+			float2 SampleUVs = (float2(VolumeCoord0) + 0.25 + (float2(OffsetX, OffsetY) * 0.5 + 0.5) * 0.5) * VolumetricTextureSizeAndInvSize.zw;
+			DataAcc = SafeSampleVolumetricTexture(SampleUVs);
+			DataAcc = ApplyFogToCloudFromUVs(SampleUVs, DataAcc);
+		}
 
-		if(!bIsWaterPixel || (bIsWaterPixel && ValidSampleCount == 0 && bAnyCloudSamplesInFrontOfWater))
+		OutputRt0 = float4(0.0f, 0.0f, 0.0f, 1.0f);
+		if (ValidSampleCount > 0.0f)
 		{
+			OutputRt0 = DataAcc / ValidSampleCount;
+		}
+		else
+		{
+			clip(-1.0f);
+		}
+		return;
+	}
+
+	// ==> Now this is the higher quality upsampling using minmax depth/luminance/trabnsmittance per pixel from reconstruction.
+
+	if (bIsWaterPixel)
+	{
+		// Clamp far depth to water depth so that upsampling works and select appropriate surface (if depth=inf, the cloud on such depth composited by water later on will never be selected)
+		VolumeFrontDepth0.w = min(VolumeFrontDepth0.w, PixelDistanceFromViewKm);
+		VolumeFrontDepth1.w = min(VolumeFrontDepth1.w, PixelDistanceFromViewKm);
+		VolumeFrontDepth2.w = min(VolumeFrontDepth2.w, PixelDistanceFromViewKm);
+		VolumeFrontDepth3.w = min(VolumeFrontDepth3.w, PixelDistanceFromViewKm);
+	}
+	//PixelDistanceFromViewKm = min(PixelDistanceFromViewKm, max(max(VolumeFrontDepth0.w, VolumeFrontDepth1.w), max(VolumeFrontDepth2.w, VolumeFrontDepth3.w)));
+
+	// Check that cloud medium sample are in front of the surface we are upsampling over water to make sure we will have valid sample.
+	// This is especially needed when upsampling cloud over water surface where depth is not matching the depth for which clouds have been traced for.
+	// If samples are all in front, we can simply upsample using the further traced clouds when for water pixel.
+//	const bool bAllCloudSamplesInFrontOfWater = bIsWaterPixel && all(float4(VolumeFrontDepth0.x, VolumeFrontDepth1.x, VolumeFrontDepth2.x, VolumeFrontDepth3.x) < PixelDistanceFromViewKm);
+	const bool bAllCloudSamplesInFrontOfDepth = all(float4(VolumeFrontDepth0.x, VolumeFrontDepth1.x, VolumeFrontDepth2.x, VolumeFrontDepth3.x) < (PixelDistanceFromViewKm-0.01));
+	const bool bAnyCloudSamplesInFrontOfWater = any(float4(VolumeFrontDepth0.x, VolumeFrontDepth1.x, VolumeFrontDepth2.x, VolumeFrontDepth3.x) < PixelDistanceFromViewKm);
+
+	float ValidSampleCount = 0;
+	float4 DataAcc = 0;
+	if (!bIsWaterPixel || (bIsWaterPixel && ValidSampleCount == 0 && bAnyCloudSamplesInFrontOfWater))
+	{
+		float4 VolumeRGBT0, VolumeRGBT1, VolumeRGBT2, VolumeRGBT3;
+		float4 Depth4;
+		BRANCH
+		if (abs(PixelDistanceFromViewKm - VolumeFrontDepth0.z) < abs(PixelDistanceFromViewKm - VolumeFrontDepth0.w))// && PixelDistanceFromViewKm < VolumeFrontDepth0.x)
+		{
+			VolumeRGBT0 = SafeLoadVolumetricSecondaryTexture(VolumeCoord0);
+			Depth4.x = VolumeFrontDepth0.z;
+		}
+		else
+		{
+			VolumeRGBT0 = SafeLoadVolumetricTexture(VolumeCoord0);
+			Depth4.x = VolumeFrontDepth0.w;
+		}
+		BRANCH
+		if (abs(PixelDistanceFromViewKm - VolumeFrontDepth1.z) < abs(PixelDistanceFromViewKm - VolumeFrontDepth1.w))// && PixelDistanceFromViewKm < VolumeFrontDepth1.x)
+		{
+			VolumeRGBT1 = SafeLoadVolumetricSecondaryTexture(VolumeCoord1);
+			Depth4.y = VolumeFrontDepth1.z;
+		}
+		else
+		{
+			VolumeRGBT1 = SafeLoadVolumetricTexture(VolumeCoord1);
+			Depth4.y = VolumeFrontDepth1.w;
+		}
+		BRANCH
+		if (abs(PixelDistanceFromViewKm - VolumeFrontDepth2.z) < abs(PixelDistanceFromViewKm - VolumeFrontDepth2.w))// && PixelDistanceFromViewKm < VolumeFrontDepth2.x)
+		{
+			VolumeRGBT2 = SafeLoadVolumetricSecondaryTexture(VolumeCoord2);
+			Depth4.z = VolumeFrontDepth2.z;
+		}
+		else
+		{
+			VolumeRGBT2 = SafeLoadVolumetricTexture(VolumeCoord2);
+			Depth4.z = VolumeFrontDepth2.w;
+		}
+		BRANCH
+		if (abs(PixelDistanceFromViewKm - VolumeFrontDepth3.z) < abs(PixelDistanceFromViewKm - VolumeFrontDepth3.w))// && PixelDistanceFromViewKm < VolumeFrontDepth3.x)
+		{
+			VolumeRGBT3 = SafeLoadVolumetricSecondaryTexture(VolumeCoord3);
+			Depth4.w = VolumeFrontDepth3.z;
+		}
+		else
+		{
+			VolumeRGBT3 = SafeLoadVolumetricTexture(VolumeCoord3);
+			Depth4.w = VolumeFrontDepth3.w;
+		}
+
+
+		float4 Depth4Diff = PixelDistanceFromViewKm - Depth4;
+		Depth4Diff = abs(Depth4Diff);
+		float MaxDepth4Diff = max(max(Depth4Diff.x, Depth4Diff.y), max(Depth4Diff.z, Depth4Diff.w));
+
+		MaxDepth4Diff = max(max(abs(VolumeFrontDepth0.w- VolumeFrontDepth0.z), abs(VolumeFrontDepth1.w - VolumeFrontDepth1.z)), max(abs(VolumeFrontDepth2.w - VolumeFrontDepth2.z), abs(VolumeFrontDepth3.w - VolumeFrontDepth3.z)));
+
+		const float MaximumDistanceKmToAllowUpsample = 2.0f;
+		const bool bAllowUpsamplingOnDistantPixel = true;// any(Depth4 < MaximumDistanceKmToAllowUpsample); // Not used since it causes upsampling problems on such area
+
+		const float WeightMultiplier = 1.0f;
+		const float ThresholdToBilinear = 0.5;// Could be parameterised or based on distance from view, max(PixelDistanceFromViewKm, 2.0f) * 0.1;
+		if (MaxDepth4Diff > ThresholdToBilinear && !bIsWaterPixel && bAllowUpsamplingOnDistantPixel)// && !bAllCloudSamplesInFrontOfDepth)
+		{
+			// Bilateral upsampling using depth
+
+			VolumeRGBT0 = ApplyFogToCloudFromUVsAndDepth(VolumeCoord0 * VolumetricTextureSizeAndInvSize.zw, Depth4.x, VolumeRGBT0);
+			VolumeRGBT1 = ApplyFogToCloudFromUVsAndDepth(VolumeCoord1 * VolumetricTextureSizeAndInvSize.zw, Depth4.y, VolumeRGBT1);
+			VolumeRGBT2 = ApplyFogToCloudFromUVsAndDepth(VolumeCoord2 * VolumetricTextureSizeAndInvSize.zw, Depth4.z, VolumeRGBT2);
+			VolumeRGBT3 = ApplyFogToCloudFromUVsAndDepth(VolumeCoord3 * VolumetricTextureSizeAndInvSize.zw, Depth4.w, VolumeRGBT3);
+
 			// Depth discontinuities edges
 			float4 weights = 1.0f / (Depth4Diff * WeightMultiplier + 1.0f);
 			const float weightsSum = dot(weights, float4(1.0f, 1.0f, 1.0f, 1.0f));
@@ -976,19 +1178,19 @@ void ComposeVolumetricRTOverScenePS(
 
 			ValidSampleCount = weightsSum > 0.0f ? 1.0 : 0.0f;
 			DataAcc = weights.x * VolumeRGBT0 + weights.y * VolumeRGBT1 + weights.z * VolumeRGBT2 + weights.w * VolumeRGBT3;
-			//DataAcc = float4(1.0, 0, 0, 0); 
+		//	DataAcc.x += 0.2;
+		}
+		else
+		{
+			// Bilinear upsampling to have a soft look region without depth edges.
+			ValidSampleCount = 1.0;
+			float2 SampleUVs = (float2(VolumeCoord0) + 0.25 + (float2(OffsetX, OffsetY) * 0.5 + 0.5) * 0.5) * VolumetricTextureSizeAndInvSize.zw;
+			DataAcc = SafeSampleVolumetricTexture(SampleUVs);
+			DataAcc = ApplyFogToCloudFromUVs(SampleUVs, DataAcc);
+		//	DataAcc.y += 0.2;
 		}
 	}
-	else
-	{
-		// Now do a bilinear sample to have a soft look region without depth edges.
-		ValidSampleCount = 1.0;
-		float2 SampleUVs = (float2(VolumeCoord0) + 0.25 + (float2(OffsetX, OffsetY) * 0.5 + 0.5) * 0.5) * VolumetricTextureSizeAndInvSize.zw;
-		DataAcc = SafeSampleVolumetricTexture(SampleUVs);
-		DataAcc = ApplyFogToCloudFromUVs(SampleUVs, DataAcc);
 
-		//DataAcc = float4(0, 1.0, 0, 0);
-	}
 
 	OutputRt0 = float4(0.0f, 0.0f, 0.0f, 1.0f);
 	if (ValidSampleCount > 0.0f)
diff --git a/Engine/Source/Runtime/Renderer/Private/DeferredShadingRenderer.cpp b/Engine/Source/Runtime/Renderer/Private/DeferredShadingRenderer.cpp
--- a/Engine/Source/Runtime/Renderer/Private/DeferredShadingRenderer.cpp
+++ b/Engine/Source/Runtime/Renderer/Private/DeferredShadingRenderer.cpp
@@ -2478,18 +2478,23 @@ void FDeferredShadingSceneRenderer::Render(FRDGBuilder& GraphBuilder, const FSce
 		UpdateLumenScene(GraphBuilder, LumenFrameTemporaries);
 
 		FRDGTextureRef HalfResolutionDepthCheckerboardMinMaxTexture = nullptr;
+		FRDGTextureRef HalfResolutionDepthMinMaxTexture = nullptr;
 		FRDGTextureRef QuarterResolutionDepthMinMaxTexture = nullptr;
 		bool bQuarterResMinMaxDepthRequired = bShouldRenderVolumetricCloud && ShouldVolumetricCloudTraceWithMinMaxDepth(Views);
 
-		auto GenerateQuarterResDepthMinMaxTexture = [&](auto& GraphBuilder, auto& Views, auto& InputTexture)
+		auto GenerateQuarterResDepthMinMaxTexture = [&](auto& GraphBuilder, auto& Views, auto& SceneDepthTexture)
 		{
 			if (bQuarterResMinMaxDepthRequired)
 			{
-				check(InputTexture != nullptr);	// Must receive a valid texture
+				check(SceneDepthTexture != nullptr);					// Must receive a valid texture
+				check(HalfResolutionDepthMinMaxTexture == nullptr);		// Only generate it once
 				check(QuarterResolutionDepthMinMaxTexture == nullptr);	// Only generate it once
-				return CreateQuarterResolutionDepthMinAndMax(GraphBuilder, Views, InputTexture);
+				CreateQuarterResolutionDepthMinAndMaxFromDepthTexture(GraphBuilder, Views, SceneDepthTexture, HalfResolutionDepthMinMaxTexture, QuarterResolutionDepthMinMaxTexture);
+			}
+			else
+			{
+				HalfResolutionDepthCheckerboardMinMaxTexture = CreateHalfResolutionDepthCheckerboardMinMax(GraphBuilder, Views, SceneDepthTexture);
 			}
-			return FRDGTextureRef(nullptr);
 		};
 		
 		FRDGTextureRef ForwardScreenSpaceShadowMaskTexture = nullptr;
@@ -2752,8 +2757,7 @@ void FDeferredShadingSceneRenderer::Render(FRDGBuilder& GraphBuilder, const FSce
 			// Volumetric Clouds have a depedency on translucent, so should run second and overlap opaque work after Lumen async is done
 			if (bShouldRenderVolumetricCloud && bAsyncComputeVolumetricCloud)
 			{
-				HalfResolutionDepthCheckerboardMinMaxTexture = CreateHalfResolutionDepthCheckerboardMinMax(GraphBuilder, Views, SceneTextures.Depth.Resolve);
-				QuarterResolutionDepthMinMaxTexture = GenerateQuarterResDepthMinMaxTexture(GraphBuilder, Views, HalfResolutionDepthCheckerboardMinMaxTexture);
+				GenerateQuarterResDepthMinMaxTexture(GraphBuilder, Views, SceneTextures.Depth.Resolve);
 
 				bool bSkipVolumetricRenderTarget = false;
 				bool bSkipPerPixelTracing = true;
@@ -3024,8 +3028,7 @@ void FDeferredShadingSceneRenderer::Render(FRDGBuilder& GraphBuilder, const FSce
 
 		if (bShouldRenderVolumetricCloud && IsVolumetricRenderTargetEnabled() && HalfResolutionDepthCheckerboardMinMaxTexture==nullptr && !bHasRayTracedOverlay)
 		{
-			HalfResolutionDepthCheckerboardMinMaxTexture = CreateHalfResolutionDepthCheckerboardMinMax(GraphBuilder, Views, SceneTextures.Depth.Resolve);
-			QuarterResolutionDepthMinMaxTexture = GenerateQuarterResDepthMinMaxTexture(GraphBuilder, Views, HalfResolutionDepthCheckerboardMinMaxTexture);
+			GenerateQuarterResDepthMinMaxTexture(GraphBuilder, Views, SceneTextures.Depth.Resolve);
 		}
 
 		if (bShouldRenderVolumetricCloud && !bHasRayTracedOverlay)
diff --git a/Engine/Source/Runtime/Renderer/Private/SceneRendering.cpp b/Engine/Source/Runtime/Renderer/Private/SceneRendering.cpp
--- a/Engine/Source/Runtime/Renderer/Private/SceneRendering.cpp
+++ b/Engine/Source/Runtime/Renderer/Private/SceneRendering.cpp
@@ -6110,6 +6110,39 @@ FRDGTextureRef CreateQuarterResolutionDepthMinAndMax(FRDGBuilder& GraphBuilder,
 	return SmallTexture;
 }
 
+void CreateQuarterResolutionDepthMinAndMaxFromDepthTexture(FRDGBuilder& GraphBuilder, TArrayView<const FViewInfo> Views, FRDGTextureRef DepthTexture, FRDGTextureRef& OutHalfResMinMax, FRDGTextureRef& OutQuarterResMinMax)
+{
+	const uint32 DownscaleFactor = 2;
+	const ETextureCreateFlags NoFastClearFlags = (CVarHalfResDepthNoFastClear.GetValueOnAnyThread() != 0) ? TexCreate_NoFastClear : TexCreate_None;
+
+	const FIntPoint HalfResMinMaxDepthExtent	= GetDownscaledExtent(DepthTexture->Desc.Extent, DownscaleFactor);
+	const FIntPoint QuarterResMinMaxDepthExtent	= GetDownscaledExtent(HalfResMinMaxDepthExtent, DownscaleFactor);
+
+	const FRDGTextureDesc HalfMinMaxDepthDesc	= FRDGTextureDesc::Create2D(HalfResMinMaxDepthExtent, PF_G16R16F, FClearValueBinding::None, TexCreate_RenderTargetable | TexCreate_ShaderResource | NoFastClearFlags);
+	const FRDGTextureDesc QuarterMinMaxDepthDesc= FRDGTextureDesc::Create2D(QuarterResMinMaxDepthExtent, PF_G16R16F, FClearValueBinding::None, TexCreate_RenderTargetable | TexCreate_ShaderResource | NoFastClearFlags);
+
+	OutHalfResMinMax	= GraphBuilder.CreateTexture(HalfMinMaxDepthDesc, TEXT("HalfResMinMaxDepthTexture"));
+	OutQuarterResMinMax	= GraphBuilder.CreateTexture(QuarterMinMaxDepthDesc, TEXT("QuarterResMinMaxDepthTexture"));
+
+	for (const FViewInfo& View : Views)
+	{
+		RDG_GPU_MASK_SCOPE(GraphBuilder, View.GPUMask);
+		
+		const FScreenPassTexture DepthPassTexture(DepthTexture, View.ViewRect);
+		const FScreenPassRenderTarget HalfResDepthTexture(OutHalfResMinMax, GetDownscaledRect(View.ViewRect, DownscaleFactor), View.DecayLoadAction(ERenderTargetLoadAction::ENoAction));
+		AddDownsampleDepthPass(GraphBuilder, View, DepthPassTexture, HalfResDepthTexture, EDownsampleDepthFilter::MinAndMaxDepth);
+	}
+
+	for (const FViewInfo& View : Views)
+	{
+		RDG_GPU_MASK_SCOPE(GraphBuilder, View.GPUMask);
+
+		const FScreenPassTexture HalfResDepthPassTexture(OutHalfResMinMax, GetDownscaledRect(View.ViewRect, DownscaleFactor));
+		const FScreenPassRenderTarget QuarterResDepthTexture(OutQuarterResMinMax, GetDownscaledRect(GetDownscaledRect(View.ViewRect, DownscaleFactor), DownscaleFactor), View.DecayLoadAction(ERenderTargetLoadAction::ENoAction));
+		AddDownsampleDepthPass(GraphBuilder, View, HalfResDepthPassTexture, QuarterResDepthTexture, EDownsampleDepthFilter::MinAndMaxDepthFromMinAndMaxDepth);
+	}
+}
+
 bool IsPrimitiveAlphaHoldoutEnabled(const FViewInfo& View)
 {
 	const bool bSupportPrimitiveAlphaHoldout = CVarPrimitiveAlphaHoldoutSupport.GetValueOnRenderThread();
diff --git a/Engine/Source/Runtime/Renderer/Private/SceneRendering.h b/Engine/Source/Runtime/Renderer/Private/SceneRendering.h
--- a/Engine/Source/Runtime/Renderer/Private/SceneRendering.h
+++ b/Engine/Source/Runtime/Renderer/Private/SceneRendering.h
@@ -3018,6 +3018,9 @@ FRDGTextureRef CreateHalfResolutionDepthCheckerboardMinMax(FRDGBuilder& GraphBui
 /** Creates a half resolution depth buffer storing the min and max depth for each 2x2 pixel quad of a half resolution buffer. */
 FRDGTextureRef CreateQuarterResolutionDepthMinAndMax(FRDGBuilder& GraphBuilder, TArrayView<const FViewInfo> Views, FRDGTextureRef DepthTexture);
 
+/** Creates a quarter resolution depth buffer storing the min and max depth for each 4x4 pixel quad of the depth buffer. */
+void CreateQuarterResolutionDepthMinAndMaxFromDepthTexture(FRDGBuilder& GraphBuilder, TArrayView<const FViewInfo> Views, FRDGTextureRef DepthTexture, FRDGTextureRef& OutHalfResMinMax, FRDGTextureRef& OutQuarterResMinMax);
+
 inline const FSceneTexturesConfig& FViewInfo::GetSceneTexturesConfig() const
 {
 	// TODO:  We are refactoring away use of the FSceneTexturesConfig::Get() global singleton, but need this workaround for now to avoid crashes
diff --git a/Engine/Source/Runtime/Renderer/Private/ScreenPass.cpp b/Engine/Source/Runtime/Renderer/Private/ScreenPass.cpp
--- a/Engine/Source/Runtime/Renderer/Private/ScreenPass.cpp
+++ b/Engine/Source/Runtime/Renderer/Private/ScreenPass.cpp
@@ -397,12 +397,14 @@ public:
 	SHADER_USE_PARAMETER_STRUCT(FDownsampleDepthPS, FGlobalShader);
 
 	class FOutputMinAndMaxDepth : SHADER_PERMUTATION_BOOL("OUTPUT_MIN_AND_MAX_DEPTH");
+	class FOutputMinMaxDepthFromMinMaxDepth : SHADER_PERMUTATION_BOOL("OUTPUT_MINMAXDEPTH_FROM_MINMAXDEPTH");
 
-	using FPermutationDomain = TShaderPermutationDomain<FOutputMinAndMaxDepth>;
+	using FPermutationDomain = TShaderPermutationDomain<FOutputMinAndMaxDepth, FOutputMinMaxDepthFromMinMaxDepth>;
 
 	BEGIN_SHADER_PARAMETER_STRUCT(FParameters, )
 		SHADER_PARAMETER_STRUCT_INCLUDE(FViewShaderParameters, View)
 		SHADER_PARAMETER_RDG_TEXTURE(Texture2D, DepthTexture)
+		SHADER_PARAMETER_RDG_TEXTURE(Texture2D<float2>, MinMaxDepthTexture)
 		SHADER_PARAMETER(FVector2f, DstToSrcPixelScale)
 		SHADER_PARAMETER(FVector2f, SourceMaxUV)
 		SHADER_PARAMETER(FVector2f, DestinationResolution)
@@ -414,6 +416,11 @@ public:
 
 	static bool ShouldCompilePermutation(const FGlobalShaderPermutationParameters& Parameters)
 	{
+		FPermutationDomain PermutationVector(Parameters.PermutationId);
+		if (PermutationVector.Get<FOutputMinMaxDepthFromMinMaxDepth>() == true && PermutationVector.Get<FOutputMinAndMaxDepth>() == true)
+		{
+			return false;
+		}
 		return IsFeatureLevelSupported(Parameters.Platform, ERHIFeatureLevel::SM5);
 	}
 };
@@ -433,8 +440,10 @@ void AddDownsampleDepthPass(
 	TShaderMapRef<FScreenPassVS> VertexShader(View.ShaderMap);
 
 	const bool bIsMinAndMaxDepthFilter = DownsampleDepthFilter == EDownsampleDepthFilter::MinAndMaxDepth;
+	const bool bIsMinAndMaxDepthFromMinMaxFilter = DownsampleDepthFilter == EDownsampleDepthFilter::MinAndMaxDepthFromMinAndMaxDepth;
 	FDownsampleDepthPS::FPermutationDomain Permutation;
 	Permutation.Set<FDownsampleDepthPS::FOutputMinAndMaxDepth>(bIsMinAndMaxDepthFilter ? 1 : 0);
+	Permutation.Set<FDownsampleDepthPS::FOutputMinMaxDepthFromMinMaxDepth>(bIsMinAndMaxDepthFromMinMaxFilter ? 1 : 0);
 	TShaderMapRef<FDownsampleDepthPS> PixelShader(View.ShaderMap, Permutation);
 
 	// The lower right corner pixel whose coordinate is max considered excluded https://learn.microsoft.com/en-us/windows/win32/direct3d11/d3d11-rect
@@ -442,7 +451,8 @@ void AddDownsampleDepthPass(
 
 	FDownsampleDepthPS::FParameters* PassParameters = GraphBuilder.AllocParameters<FDownsampleDepthPS::FParameters>();
 	PassParameters->View = View.GetShaderParameters();
-	PassParameters->DepthTexture = Input.Texture;
+	PassParameters->DepthTexture		= bIsMinAndMaxDepthFromMinMaxFilter ? GSystemTextures.GetDepthDummy(GraphBuilder) : Input.Texture;
+	PassParameters->MinMaxDepthTexture	= bIsMinAndMaxDepthFromMinMaxFilter ? Input.Texture : GSystemTextures.GetBlackDummy(GraphBuilder);
 	PassParameters->DstToSrcPixelScale = FVector2f(float(InputViewport.Extent.X) / float(OutputViewport.Extent.X), float(InputViewport.Extent.Y) / float(OutputViewport.Extent.Y));
 	PassParameters->SourceMaxUV = FVector2f((float(View.ViewRect.Max.X) -1.0f - 0.51f) / InputViewport.Extent.X, (float(View.ViewRect.Max.Y) - 1.0f - 0.51f) / InputViewport.Extent.Y);
 	PassParameters->DownsampleDepthFilter = (uint32)DownsampleDepthFilter;
@@ -456,7 +466,7 @@ void AddDownsampleDepthPass(
 
 	FRHIDepthStencilState* DepthStencilState = TStaticDepthStencilState<true, CF_Always>::GetRHI();
 
-	if (bIsMinAndMaxDepthFilter)
+	if (bIsMinAndMaxDepthFilter || bIsMinAndMaxDepthFromMinMaxFilter)
 	{
 		DepthStencilState = TStaticDepthStencilState<false, CF_Always>::GetRHI();
 		PassParameters->RenderTargets[0] = FRenderTargetBinding(Output.Texture, Output.LoadAction);
@@ -471,6 +481,7 @@ void AddDownsampleDepthPass(
 		TEXT("Max"),
 		TEXT("CheckerMinMax"),
 		TEXT("MinAndMaxDepth"),
+		TEXT("MinMaxFromMinMaxDepth"),
 	};
 
 	AddDrawScreenPass(
diff --git a/Engine/Source/Runtime/Renderer/Private/VolumetricCloudRendering.cpp b/Engine/Source/Runtime/Renderer/Private/VolumetricCloudRendering.cpp
--- a/Engine/Source/Runtime/Renderer/Private/VolumetricCloudRendering.cpp
+++ b/Engine/Source/Runtime/Renderer/Private/VolumetricCloudRendering.cpp
@@ -2793,8 +2793,8 @@ bool FSceneRenderer::RenderVolumetricCloud(
 
 					CloudRC.NoiseFrameIndexModPattern = VRT.GetNoiseFrameIndexModPattern();
 
-					check(VRT.GetMode() != 0 || CloudRC.bCloudDebugViewModeEnabled || HalfResolutionDepthCheckerboardMinMaxTexture);
-					CloudRC.SceneDepthZ = (VRT.GetMode() == 0 || VRT.GetMode() == 1) ? HalfResolutionDepthCheckerboardMinMaxTexture : SceneDepthZ;
+					check(VRT.GetMode() != 0 || CloudRC.bCloudDebugViewModeEnabled || (bShouldVolumetricCloudTraceWithMinMaxDepth || (!bShouldVolumetricCloudTraceWithMinMaxDepth && HalfResolutionDepthCheckerboardMinMaxTexture)));
+					CloudRC.SceneDepthZ = ((VRT.GetMode() == 0 && !bShouldVolumetricCloudTraceWithMinMaxDepth) || VRT.GetMode() == 1) ? HalfResolutionDepthCheckerboardMinMaxTexture : SceneDepthZ;
 					CloudRC.SceneDepthMinAndMax = bShouldVolumetricCloudTraceWithMinMaxDepth ? QuarterResolutionDepthMinMaxTexture : SceneDepthZ;
 				}
 				else
diff --git a/Engine/Source/Runtime/Renderer/Private/VolumetricRenderTarget.cpp b/Engine/Source/Runtime/Renderer/Private/VolumetricRenderTarget.cpp
--- a/Engine/Source/Runtime/Renderer/Private/VolumetricRenderTarget.cpp
+++ b/Engine/Source/Runtime/Renderer/Private/VolumetricRenderTarget.cpp
@@ -272,6 +272,7 @@ void FVolumetricRenderTargetViewStateData::Initialise(
 		{
 			// Resolution does not match so release target we are going to render in
 			VolumetricReconstructRT[CurrentRT].SafeRelease();
+			VolumetricReconstructSecondaryRT[CurrentRT].SafeRelease();
 			VolumetricReconstructRTDepth[CurrentRT].SafeRelease();
 		}
 
@@ -336,6 +337,7 @@ void FVolumetricRenderTargetViewStateData::Reset()
 	for(uint32 i = 0 ; i < kRenderTargetCount; ++i)
 	{
 		VolumetricReconstructRT[i].SafeRelease();
+		VolumetricReconstructSecondaryRT[i].SafeRelease();
 		VolumetricReconstructRTDepth[i].SafeRelease();
 	}
 }
@@ -426,6 +428,11 @@ FRDGTextureRef FVolumetricRenderTargetViewStateData::GetDstVolumetricReconstruct
 	return VolumetricReconstructRT[CurrentRT].IsValid() ? GraphBuilder.RegisterExternalTexture(VolumetricReconstructRT[CurrentRT]) : nullptr;
 }
 
+FRDGTextureRef FVolumetricRenderTargetViewStateData::GetDstVolumetricReconstructSecondaryRT(FRDGBuilder& GraphBuilder)
+{
+	return VolumetricReconstructSecondaryRT[CurrentRT].IsValid() ? GraphBuilder.RegisterExternalTexture(VolumetricReconstructSecondaryRT[CurrentRT]) : nullptr;
+}
+
 FRDGTextureRef FVolumetricRenderTargetViewStateData::GetOrCreateDstVolumetricReconstructRT(FRDGBuilder& GraphBuilder)
 {
 	check(VolumetricReconstructRTResolution != FIntPoint::ZeroValue); // check that initialization has been done at least once
@@ -441,6 +448,21 @@ FRDGTextureRef FVolumetricRenderTargetViewStateData::GetOrCreateDstVolumetricRec
 	return GraphBuilder.RegisterExternalTexture(VolumetricReconstructRT[CurrentRT]);
 }
 
+FRDGTextureRef FVolumetricRenderTargetViewStateData::GetOrCreateDstVolumetricReconstructSecondaryRT(FRDGBuilder& GraphBuilder)
+{
+	check(VolumetricReconstructRTResolution != FIntPoint::ZeroValue); // check that initialization has been done at least once
+
+	if (!VolumetricReconstructSecondaryRT[CurrentRT].IsValid())
+	{
+		FPooledRenderTargetDesc Desc = FPooledRenderTargetDesc::Create2DDesc(
+			VolumetricReconstructRTResolution, PF_FloatRGBA, FClearValueBinding(FLinearColor(0.0f, 0.0f, 0.0f, 1.0f)),
+			TexCreate_None, TexCreate_ShaderResource | TexCreate_RenderTargetable, false);
+		GRenderTargetPool.FindFreeElement(GraphBuilder.RHICmdList, Desc, VolumetricReconstructSecondaryRT[CurrentRT], TEXT("VolumetricRenderTarget.Reconstruct"));
+	}
+
+	return GraphBuilder.RegisterExternalTexture(VolumetricReconstructSecondaryRT[CurrentRT]);
+}
+
 
 FRDGTextureRef FVolumetricRenderTargetViewStateData::GetOrCreateDstVolumetricReconstructRTDepth(FRDGBuilder& GraphBuilder)
 {
@@ -449,7 +471,7 @@ FRDGTextureRef FVolumetricRenderTargetViewStateData::GetOrCreateDstVolumetricRec
 	if (!VolumetricReconstructRTDepth[CurrentRT].IsValid())
 	{
 		FPooledRenderTargetDesc Desc = FPooledRenderTargetDesc::Create2DDesc(
-			VolumetricReconstructRTResolution, PF_G16R16F, FClearValueBinding(FLinearColor(63000.0f, 63000.0f, 63000.0f, 63000.0f)),
+			VolumetricReconstructRTResolution, PF_FloatRGBA, FClearValueBinding(FLinearColor(63000.0f, 63000.0f, 63000.0f, 63000.0f)),
 			TexCreate_None, TexCreate_ShaderResource | TexCreate_RenderTargetable, false);
 		GRenderTargetPool.FindFreeElement(GraphBuilder.RHICmdList, Desc, VolumetricReconstructRTDepth[CurrentRT], TEXT("VolumetricRenderTarget.ReconstructDepth"));
 	}
@@ -461,6 +483,10 @@ TRefCountPtr<IPooledRenderTarget> FVolumetricRenderTargetViewStateData::GetDstVo
 {
 	return VolumetricReconstructRT[CurrentRT];
 }
+TRefCountPtr<IPooledRenderTarget> FVolumetricRenderTargetViewStateData::GetDstVolumetricReconstructSecondaryRT()
+{
+	return VolumetricReconstructSecondaryRT[CurrentRT];
+}
 TRefCountPtr<IPooledRenderTarget> FVolumetricRenderTargetViewStateData::GetDstVolumetricReconstructRTDepth()
 {
 	return VolumetricReconstructRTDepth[CurrentRT];
@@ -498,6 +524,14 @@ FRDGTextureRef FVolumetricRenderTargetViewStateData::GetOrCreateSrcVolumetricRec
 	return GraphBuilder.RegisterExternalTexture(VolumetricReconstructRT[1u - CurrentRT]);
 }
 
+FRDGTextureRef FVolumetricRenderTargetViewStateData::GetOrCreateSrcVolumetricReconstructSecondaryRT(FRDGBuilder& GraphBuilder)
+{
+	check(VolumetricReconstructRTResolution != FIntPoint::ZeroValue); // check that initialization has been done at least once
+
+	check(VolumetricReconstructSecondaryRT[1u - CurrentRT].IsValid());
+	return GraphBuilder.RegisterExternalTexture(VolumetricReconstructSecondaryRT[1u - CurrentRT]);
+}
+
 FRDGTextureRef FVolumetricRenderTargetViewStateData::GetOrCreateSrcVolumetricReconstructRTDepth(FRDGBuilder& GraphBuilder)
 {
 	check(VolumetricReconstructRTResolution != FIntPoint::ZeroValue); // check that initialization has been done at least once
@@ -665,6 +699,7 @@ class FReconstructVolumetricRenderTargetPS : public FGlobalShader
 		SHADER_PARAMETER_RDG_TEXTURE(Texture2D, SecondaryTracingVolumetricTexture)
 		SHADER_PARAMETER_RDG_TEXTURE(Texture2D, TracingVolumetricDepthTexture)
 		SHADER_PARAMETER_RDG_TEXTURE(Texture2D, PreviousFrameVolumetricTexture)
+		SHADER_PARAMETER_RDG_TEXTURE(Texture2D, PreviousFrameVolumetricSecondaryTexture)
 		SHADER_PARAMETER_RDG_TEXTURE(Texture2D, PreviousFrameVolumetricDepthTexture)
 		SHADER_PARAMETER_RDG_TEXTURE(Texture2D, HalfResDepthTexture)
 		SHADER_PARAMETER_SAMPLER(SamplerState, LinearTextureSampler)
@@ -748,29 +783,42 @@ void ReconstructVolumetricRenderTarget(
 		}
 
 		FRDGTextureRef DstVolumetric = VolumetricCloudRT.GetOrCreateDstVolumetricReconstructRT(GraphBuilder);
+		FRDGTextureRef DstVolumetricSecondary = VolumetricCloudRT.GetOrCreateDstVolumetricReconstructSecondaryRT(GraphBuilder);
 		FRDGTextureRef DstVolumetricDepth = VolumetricCloudRT.GetOrCreateDstVolumetricReconstructRTDepth(GraphBuilder);
 		FRDGTextureRef SrcTracingVolumetric = VolumetricCloudRT.GetOrCreateVolumetricTracingRT(GraphBuilder);
 		FRDGTextureRef SrcSecondaryTracingVolumetric = VolumetricCloudRT.GetOrCreateVolumetricSecondaryTracingRT(GraphBuilder);
 		FRDGTextureRef SrcTracingVolumetricDepth = VolumetricCloudRT.GetOrCreateVolumetricTracingRTDepth(GraphBuilder);
 		FRDGTextureRef PreviousFrameVolumetricTexture = VolumetricCloudRT.GetHistoryValid() ? VolumetricCloudRT.GetOrCreateSrcVolumetricReconstructRT(GraphBuilder) : SystemTextures.Black;
+		FRDGTextureRef PreviousFrameVolumetricSecondaryTexture = VolumetricCloudRT.GetHistoryValid() ? VolumetricCloudRT.GetOrCreateSrcVolumetricReconstructSecondaryRT(GraphBuilder) : SystemTextures.Black;
 		FRDGTextureRef PreviousFrameVolumetricDepthTexture = VolumetricCloudRT.GetHistoryValid() ? VolumetricCloudRT.GetOrCreateSrcVolumetricReconstructRTDepth(GraphBuilder) : SystemTextures.Black;
 
 		const uint32 TracingVolumetricCloudRTDownSample = VolumetricCloudRT.GetVolumetricTracingRTDownsampleFactor();
 
+		const bool bMinMaxDepth = ShouldVolumetricCloudTraceWithMinMaxDepth(ViewInfo);
+
 		FReconstructVolumetricRenderTargetPS::FPermutationDomain PermutationVector;
 		PermutationVector.Set<FReconstructVolumetricRenderTargetPS::FHistoryAvailable>(VolumetricCloudRT.GetHistoryValid());
 		PermutationVector.Set<FReconstructVolumetricRenderTargetPS::FReprojectionBoxConstraint>(CVarVolumetricRenderTargetReprojectionBoxConstraint.GetValueOnAnyThread() > 0);
-		PermutationVector.Set<FReconstructVolumetricRenderTargetPS::FCloudMinAndMaxDepth>(ShouldVolumetricCloudTraceWithMinMaxDepth(ViewInfo));
+		PermutationVector.Set<FReconstructVolumetricRenderTargetPS::FCloudMinAndMaxDepth>(bMinMaxDepth);
 		TShaderMapRef<FReconstructVolumetricRenderTargetPS> PixelShader(ViewInfo.ShaderMap, PermutationVector);
 
 		FReconstructVolumetricRenderTargetPS::FParameters* PassParameters = GraphBuilder.AllocParameters<FReconstructVolumetricRenderTargetPS::FParameters>();
 		PassParameters->ViewUniformBuffer = ViewInfo.VolumetricRenderTargetViewUniformBuffer; // Using a special uniform buffer because the view has some special resolution and no split screen offset.
 		PassParameters->RenderTargets[0] = FRenderTargetBinding(DstVolumetric, ERenderTargetLoadAction::ENoAction);
-		PassParameters->RenderTargets[1] = FRenderTargetBinding(DstVolumetricDepth, ERenderTargetLoadAction::ENoAction);
+		if (bMinMaxDepth)
+		{
+			PassParameters->RenderTargets[1] = FRenderTargetBinding(DstVolumetricSecondary, ERenderTargetLoadAction::ENoAction);
+			PassParameters->RenderTargets[2] = FRenderTargetBinding(DstVolumetricDepth, ERenderTargetLoadAction::ENoAction);
+		}
+		else
+		{
+			PassParameters->RenderTargets[1] = FRenderTargetBinding(DstVolumetricDepth, ERenderTargetLoadAction::ENoAction);
+		}
 		PassParameters->TracingVolumetricTexture = SrcTracingVolumetric;
 		PassParameters->SecondaryTracingVolumetricTexture = SrcSecondaryTracingVolumetric;
 		PassParameters->TracingVolumetricDepthTexture = SrcTracingVolumetricDepth;
 		PassParameters->PreviousFrameVolumetricTexture = PreviousFrameVolumetricTexture;
+		PassParameters->PreviousFrameVolumetricSecondaryTexture = PreviousFrameVolumetricSecondaryTexture;
 		PassParameters->PreviousFrameVolumetricDepthTexture = PreviousFrameVolumetricDepthTexture;
 		PassParameters->LinearTextureSampler = TStaticSamplerState<SF_Bilinear>::GetRHI();
 		PassParameters->CurrentTracingPixelOffset = VolumetricCloudRT.GetCurrentTracingPixelOffset();
@@ -829,6 +877,7 @@ class FComposeVolumetricRTOverScenePS : public FGlobalShader
 		SHADER_PARAMETER_STRUCT(FLocalFogVolumeUniformParameters, LFV)
 		SHADER_PARAMETER_RDG_UNIFORM_BUFFER(FFogUniformParameters, FogStruct)
 		SHADER_PARAMETER_RDG_TEXTURE(Texture2D, VolumetricTexture)
+		SHADER_PARAMETER_RDG_TEXTURE(Texture2D, VolumetricSecondaryTexture)
 		SHADER_PARAMETER_RDG_TEXTURE(Texture2D, VolumetricDepthTexture)
 		SHADER_PARAMETER_RDG_TEXTURE(Texture2D, WaterLinearDepthTexture)
 		SHADER_PARAMETER_RDG_TEXTURE(Texture2DMS<float>, MSAADepthTexture)
@@ -837,6 +886,7 @@ class FComposeVolumetricRTOverScenePS : public FGlobalShader
 		RENDER_TARGET_BINDING_SLOTS()
 		SHADER_PARAMETER(uint32, ForwardShadingEnable)
 		SHADER_PARAMETER(uint32, OutputAlphaHoldout)
+		SHADER_PARAMETER(uint32, MinMaxDepthAvailable)
 		SHADER_PARAMETER(float, UvOffsetSampleAcceptanceWeight)
 		SHADER_PARAMETER(FVector4f, VolumetricTextureSizeAndInvSize)
 		SHADER_PARAMETER(FVector2f, FullResolutionToVolumetricBufferResolutionScale)
@@ -886,7 +936,7 @@ static FVector2f GetCompositionFullResolutionToVolumetricBufferResolutionScale(u
 	return FVector2f(1.0f / float(GetMainDownsampleFactor(VRTMode)), float(GetMainDownsampleFactor(VRTMode)));
 }
 
-static void GetCompositionCloudTextures(uint32 VRTMode, FVolumetricRenderTargetViewStateData& VolumetricCloudRT, FRDGBuilder& GraphBuilder, FRDGTextureRef& VolumetricTexture, FRDGTextureRef& VolumetricDepthTexture,
+static void GetCompositionCloudTextures(uint32 VRTMode, FVolumetricRenderTargetViewStateData& VolumetricCloudRT, FRDGBuilder& GraphBuilder, FRDGTextureRef& VolumetricTexture, FRDGTextureRef& VolumetricSecondaryTexture, FRDGTextureRef& VolumetricDepthTexture,
 										FIntPoint& VolumetricViewRect)
 {
 	if (VRTMode == 1 || VRTMode == 3)
@@ -895,12 +945,14 @@ static void GetCompositionCloudTextures(uint32 VRTMode, FVolumetricRenderTargetV
 		// We will then directly up sample on screen from half resolution to full resolution. 
 		// No reconstruction needed.
 		VolumetricTexture = VolumetricCloudRT.GetOrCreateVolumetricTracingRT(GraphBuilder);
+		VolumetricSecondaryTexture = GSystemTextures.GetBlackAlphaOneDummy(GraphBuilder);
 		VolumetricDepthTexture = VolumetricCloudRT.GetOrCreateVolumetricTracingRTDepth(GraphBuilder);
 		VolumetricViewRect = VolumetricCloudRT.GetCurrentVolumetricTracingViewRect();
 	}
 	else
 	{
 		VolumetricTexture = VolumetricCloudRT.GetOrCreateDstVolumetricReconstructRT(GraphBuilder);
+		VolumetricSecondaryTexture = VRTMode == 0 ? VolumetricCloudRT.GetOrCreateDstVolumetricReconstructSecondaryRT(GraphBuilder) : GSystemTextures.GetBlackAlphaOneDummy(GraphBuilder);
 		VolumetricDepthTexture = VolumetricCloudRT.GetOrCreateDstVolumetricReconstructRTDepth(GraphBuilder);
 		VolumetricViewRect = VolumetricCloudRT.GetDstVolumetricReconstructViewRect();
 	}
@@ -964,11 +1016,13 @@ void ComposeVolumetricRenderTargetOverScene(
 		// When reconstructed and back buffer resolution matches, force using a pixel perfect upsampling.
 		const uint32 VRTMode = VolumetricCloudRT.GetMode();
 		int32 UpsamplingMode = GetCompositionUpsamplingMode(VRTMode, VolumetricCloudRT.GetUpsamplingMode());
+		const bool bMinMaxDepth = ShouldVolumetricCloudTraceWithMinMaxDepth(ViewInfo);
 
 		FRDGTextureRef VolumetricTexture = nullptr;
+		FRDGTextureRef VolumetricSecondaryTexture = nullptr;
 		FRDGTextureRef VolumetricDepthTexture = nullptr;
 		FIntPoint VolumetricViewRect(FIntPoint::ZeroValue);
-		GetCompositionCloudTextures(VRTMode, VolumetricCloudRT, GraphBuilder, VolumetricTexture, VolumetricDepthTexture, VolumetricViewRect);
+		GetCompositionCloudTextures(VRTMode, VolumetricCloudRT, GraphBuilder, VolumetricTexture, VolumetricSecondaryTexture, VolumetricDepthTexture, VolumetricViewRect);
 
 		// We only support MSAA up to 8 sample and in forward
 		check(SceneDepthTexture->Desc.NumSamples <= 8);
@@ -992,6 +1046,7 @@ void ComposeVolumetricRenderTargetOverScene(
 		PassParameters->RenderTargets[0] = FRenderTargetBinding(SceneColorTexture, ERenderTargetLoadAction::ELoad);
 		PassParameters->MSAADepthTexture = SceneDepthTexture;
 		PassParameters->VolumetricTexture = VolumetricTexture;
+		PassParameters->VolumetricSecondaryTexture = bMinMaxDepth ? VolumetricSecondaryTexture : VolumetricTexture;
 		PassParameters->VolumetricDepthTexture = VolumetricDepthTexture;
 		PassParameters->LinearTextureSampler = TStaticSamplerState<SF_Bilinear>::GetRHI();
 		PassParameters->UvOffsetSampleAcceptanceWeight = GetUvNoiseSampleAcceptanceWeight();
@@ -1000,6 +1055,7 @@ void ComposeVolumetricRenderTargetOverScene(
 		PassParameters->Substrate = Substrate::BindSubstrateGlobalUniformParameters(ViewInfo);
 		PassParameters->OutputAlphaHoldout = 0;
 		PassParameters->ForwardShadingEnable = bForwardShading ? 1 : 0;
+		PassParameters->MinMaxDepthAvailable = bMinMaxDepth ? 1 : 0;
 		FVector2f DummyUVScale;
 		GetTextureSafeUvCoordBound(VolumetricViewRect, PassParameters->VolumetricTexture, PassParameters->VolumetricTextureValidCoordRect, PassParameters->VolumetricTextureValidUvRect, DummyUVScale);
 
@@ -1084,11 +1140,13 @@ void ComposeVolumetricRenderTargetOverSceneUnderWater(
 		// When reconstructed and back buffer resolution matches, force using a pixel perfect upsampling.
 		const uint32 VRTMode = VolumetricCloudRT.GetMode();
 		int32 UpsamplingMode = GetCompositionUpsamplingMode(VRTMode, VolumetricCloudRT.GetUpsamplingMode());
+		const bool bMinMaxDepth = ShouldVolumetricCloudTraceWithMinMaxDepth(ViewInfo);
 
 		FRDGTextureRef VolumetricTexture = nullptr;
+		FRDGTextureRef VolumetricSecondaryTexture = nullptr;
 		FRDGTextureRef VolumetricDepthTexture = nullptr;
 		FIntPoint VolumetricViewRect(FIntPoint::ZeroValue);
-		GetCompositionCloudTextures(VRTMode, VolumetricCloudRT, GraphBuilder, VolumetricTexture, VolumetricDepthTexture, VolumetricViewRect);
+		GetCompositionCloudTextures(VRTMode, VolumetricCloudRT, GraphBuilder, VolumetricTexture, VolumetricSecondaryTexture, VolumetricDepthTexture, VolumetricViewRect);
 
 		FComposeVolumetricRTOverScenePS::FPermutationDomain PermutationVector;
 		PermutationVector.Set<FComposeVolumetricRTOverScenePS::FUpsamplingMode>(UpsamplingMode);
@@ -1101,6 +1159,7 @@ void ComposeVolumetricRenderTargetOverSceneUnderWater(
 		PassParameters->ViewUniformBuffer = ViewInfo.ViewUniformBuffer;
 		PassParameters->RenderTargets[0] = FRenderTargetBinding(WaterPassData.ColorTexture, ERenderTargetLoadAction::ELoad);
 		PassParameters->VolumetricTexture = VolumetricTexture;
+		PassParameters->VolumetricSecondaryTexture = bMinMaxDepth ? VolumetricSecondaryTexture : VolumetricTexture;
 		PassParameters->VolumetricDepthTexture = VolumetricDepthTexture;
 		PassParameters->WaterLinearDepthTexture = WaterPassData.DepthTexture;
 		PassParameters->LinearTextureSampler = TStaticSamplerState<SF_Bilinear>::GetRHI();
@@ -1113,6 +1172,8 @@ void ComposeVolumetricRenderTargetOverSceneUnderWater(
 		PassParameters->SceneTextures = SceneTextures.UniformBuffer;
 		PassParameters->Substrate = Substrate::BindSubstrateGlobalUniformParameters(ViewInfo);
 		PassParameters->OutputAlphaHoldout = 0;
+		PassParameters->ForwardShadingEnable = 0;
+		PassParameters->MinMaxDepthAvailable = bMinMaxDepth ? 1 : 0;
 		FVector2f DummyUVScale;
 		GetTextureSafeUvCoordBound(VolumetricViewRect, PassParameters->VolumetricTexture, PassParameters->VolumetricTextureValidCoordRect, PassParameters->VolumetricTextureValidUvRect, DummyUVScale);
 
@@ -1159,11 +1220,13 @@ void ComposeVolumetricRenderTargetOverSceneForVisualization(
 		// When reconstructed and back buffer resolution matches, force using a pixel perfect upsampling.
 		const uint32 VRTMode = VolumetricCloudRT.GetMode();
 		int32 UpsamplingMode = GetCompositionUpsamplingMode(VRTMode, VolumetricCloudRT.GetUpsamplingMode());
+		const bool bMinMaxDepth = ShouldVolumetricCloudTraceWithMinMaxDepth(ViewInfo);
 
 		FRDGTextureRef VolumetricTexture = nullptr;
+		FRDGTextureRef VolumetricSecondaryTexture = nullptr;
 		FRDGTextureRef VolumetricDepthTexture = nullptr;
 		FIntPoint VolumetricViewRect(FIntPoint::ZeroValue);
-		GetCompositionCloudTextures(VRTMode, VolumetricCloudRT, GraphBuilder, VolumetricTexture, VolumetricDepthTexture, VolumetricViewRect);
+		GetCompositionCloudTextures(VRTMode, VolumetricCloudRT, GraphBuilder, VolumetricTexture, VolumetricSecondaryTexture, VolumetricDepthTexture, VolumetricViewRect);
 
 		FComposeVolumetricRTOverScenePS::FPermutationDomain PermutationVector;
 		PermutationVector.Set<FComposeVolumetricRTOverScenePS::FUpsamplingMode>(0);
@@ -1176,6 +1239,7 @@ void ComposeVolumetricRenderTargetOverSceneForVisualization(
 		PassParameters->ViewUniformBuffer = ViewInfo.ViewUniformBuffer;
 		PassParameters->RenderTargets[0] = FRenderTargetBinding(SceneColorTexture, ERenderTargetLoadAction::ELoad);
 		PassParameters->VolumetricTexture = VolumetricTexture;
+		PassParameters->VolumetricSecondaryTexture = bMinMaxDepth ? VolumetricSecondaryTexture : VolumetricTexture;
 		PassParameters->VolumetricDepthTexture = VolumetricDepthTexture;
 		PassParameters->LinearTextureSampler = TStaticSamplerState<SF_Bilinear>::GetRHI();
 		PassParameters->UvOffsetSampleAcceptanceWeight = GetUvNoiseSampleAcceptanceWeight();
@@ -1183,6 +1247,8 @@ void ComposeVolumetricRenderTargetOverSceneForVisualization(
 		PassParameters->SceneTextures = SceneTextures.UniformBuffer;
 		PassParameters->Substrate = Substrate::BindSubstrateGlobalUniformParameters(ViewInfo);
 		PassParameters->OutputAlphaHoldout = 0;
+		PassParameters->ForwardShadingEnable = 0;
+		PassParameters->MinMaxDepthAvailable = bMinMaxDepth ? 1 : 0;
 		FVector2f DummyUVScale;
 		GetTextureSafeUvCoordBound(VolumetricViewRect, PassParameters->VolumetricTexture, PassParameters->VolumetricTextureValidCoordRect, PassParameters->VolumetricTextureValidUvRect, DummyUVScale);
 
diff --git a/Engine/Source/Runtime/Renderer/Private/VolumetricRenderTargetViewStateData.h b/Engine/Source/Runtime/Renderer/Private/VolumetricRenderTargetViewStateData.h
--- a/Engine/Source/Runtime/Renderer/Private/VolumetricRenderTargetViewStateData.h
+++ b/Engine/Source/Runtime/Renderer/Private/VolumetricRenderTargetViewStateData.h
@@ -48,16 +48,20 @@ public:
 	FVector2f GetVolumetricTracingUVMax() const;
 
 	FRDGTextureRef GetDstVolumetricReconstructRT(FRDGBuilder& GraphBuilder);
+	FRDGTextureRef GetDstVolumetricReconstructSecondaryRT(FRDGBuilder& GraphBuilder);
 	FRDGTextureRef GetOrCreateDstVolumetricReconstructRT(FRDGBuilder& GraphBuilder);
+	FRDGTextureRef GetOrCreateDstVolumetricReconstructSecondaryRT(FRDGBuilder& GraphBuilder);
 	FRDGTextureRef GetOrCreateDstVolumetricReconstructRTDepth(FRDGBuilder& GraphBuilder);
 	const FIntPoint& GetDstVolumetricReconstructViewRect() const;
 	FVector2f GetDstVolumetricReconstructUVScale() const;
 	FVector2f GetDstVolumetricReconstructUVMax() const;
 
 	TRefCountPtr<IPooledRenderTarget> GetDstVolumetricReconstructRT();
+	TRefCountPtr<IPooledRenderTarget> GetDstVolumetricReconstructSecondaryRT();
 	TRefCountPtr<IPooledRenderTarget> GetDstVolumetricReconstructRTDepth();
 
 	FRDGTextureRef GetOrCreateSrcVolumetricReconstructRT(FRDGBuilder& GraphBuilder);
+	FRDGTextureRef GetOrCreateSrcVolumetricReconstructSecondaryRT(FRDGBuilder& GraphBuilder);
 	FRDGTextureRef GetOrCreateSrcVolumetricReconstructRTDepth(FRDGBuilder& GraphBuilder);
 	const FIntPoint& GetSrcVolumetricReconstructViewRect() const;
 
@@ -106,6 +110,7 @@ private:
 
 	static constexpr uint32 kRenderTargetCount = 2;
 	TRefCountPtr<IPooledRenderTarget> VolumetricReconstructRT[kRenderTargetCount];
+	TRefCountPtr<IPooledRenderTarget> VolumetricReconstructSecondaryRT[kRenderTargetCount];
 	TRefCountPtr<IPooledRenderTarget> VolumetricReconstructRTDepth[kRenderTargetCount];
 	FIntPoint VolumetricReconstructViewRect[kRenderTargetCount];
 
diff --git a/Engine/Source/Runtime/Renderer/Public/ScreenPass.h b/Engine/Source/Runtime/Renderer/Public/ScreenPass.h
--- a/Engine/Source/Runtime/Renderer/Public/ScreenPass.h
+++ b/Engine/Source/Runtime/Renderer/Public/ScreenPass.h
@@ -753,8 +753,11 @@ enum class EDownsampleDepthFilter
 	// Produces a checkerboarded selection of min and max depth values
 	Checkerboard,
 
-	// Procuce a color texture where R=min and G=max DeviceZ
-	MinAndMaxDepth
+	// Produces a color texture where R=min and G=max DeviceZ
+	MinAndMaxDepth,
+
+	// Produces a color texture where R=min and G=max DeviceZ from a color texture containing R=min and G=max DeviceZ
+	MinAndMaxDepthFromMinAndMaxDepth
 };
 
 RENDERER_API void AddDownsampleDepthPass(
