Volumetric Cloud - fixed conservative density view.

#tests pc editor dx12
#rnx

diff --git a/Engine/Shaders/Private/VolumetricCloud.usf b/Engine/Shaders/Private/VolumetricCloud.usf
--- a/Engine/Shaders/Private/VolumetricCloud.usf
+++ b/Engine/Shaders/Private/VolumetricCloud.usf
@@ -51,10 +51,6 @@
 #define MATERIAL_VOLUMETRIC_CLOUD_EMPTY_SPACE_SKIPPING_OUTPUT 0
 #endif
 
-#ifndef CLOUD_CONSERVATIVE_DENSITY_DEBUG
-#define CLOUD_CONSERVATIVE_DENSITY_DEBUG 0
-#endif
-
 #ifndef CLOUD_DEBUG_SAMPLES
 #define CLOUD_DEBUG_SAMPLES 0
 #endif
@@ -1719,8 +1715,12 @@ void MainCommon(in FMaterialPixelParameters MaterialParameters, in float4 SvPosi
 #endif
 
 #if CLOUD_MIN_AND_MAX_DEPTH
+#if CLOUD_DEBUG_VIEW_MODE
+	OutColor1 = OutColor0;
+#else
 	float ClosestGrayScaleTransmittance = ClosestMeanTransmittance < RenderVolumetricCloudParameters.StopTracingTransmittanceThreshold ? 0.0f : ClosestMeanTransmittance;
 	OutColor1 = float4(ClosestLuminance * OutputPreExposure, ClosestGrayScaleTransmittance);
+#endif
 #endif // CLOUD_MIN_AND_MAX
 
 	OutDepth.x = MaxHalfFloat; // Default to far away depth to be flat and not intersect with any geometry.
