Fix curvature to roughness on Nanite due to missing derivates.

The CL add normal geo analytical forwarding to the function on Nanite meshes.

diff --git a/Engine/Shaders/Private/MaterialTemplate.ush b/Engine/Shaders/Private/MaterialTemplate.ush
--- a/Engine/Shaders/Private/MaterialTemplate.ush
+++ b/Engine/Shaders/Private/MaterialTemplate.ush
@@ -4097,25 +4097,36 @@ float3 CalculateAnisotropyTangent(in out FMaterialPixelParameters Parameters, FP
 	return Tangent;
 }
 
-float NormalCurvatureToRoughness(float3 WorldNormal)
+float NormalCurvatureToRoughness(float3 WorldNormal, float3 dNdx, float3 dNdy)
 {
-	float3 dNdx = ddx(WorldNormal);
-	float3 dNdy = ddy(WorldNormal);
 	float x = dot(dNdx, dNdx);
 	float y = dot(dNdy, dNdy);
 	float CurvatureApprox = pow(max(x, y), View.NormalCurvatureToRoughnessScaleBias.z);
 	return saturate(CurvatureApprox * View.NormalCurvatureToRoughnessScaleBias.x + View.NormalCurvatureToRoughnessScaleBias.y);
 }
 
+float NormalCurvatureToRoughness(float3 WorldNormal)
+{
+	return NormalCurvatureToRoughness(WorldNormal, ddx(WorldNormal), ddy(WorldNormal));
+}
+
 float GetRoughnessFromNormalCurvature(FMaterialPixelParameters InMaterialParameters)
 {
 #if MATERIAL_NORMAL_CURVATURE_TO_ROUGHNESS
 	// Curvature-to-roughness uses derivatives of the WorldVertexNormal, which is incompatible with centroid interpolation because
 	// the samples are not uniformly distributed. Therefore we use WorldVertexNormal_Center which is guaranteed to be center interpolated.
 	#if USE_WORLDVERTEXNORMAL_CENTER_INTERPOLATION
-	return NormalCurvatureToRoughness(InMaterialParameters.WorldVertexNormal_Center);
+		#if USE_ANALYTIC_DERIVATIVES
+		return NormalCurvatureToRoughness(InMaterialParameters.WorldVertexNormal_Center, InMaterialParameters.WorldGeoNormal_DDX, InMaterialParameters.WorldGeoNormal_DDY);
+		#else
+		return NormalCurvatureToRoughness(InMaterialParameters.WorldVertexNormal_Center);
+		#endif
 	#else
-	return NormalCurvatureToRoughness(InMaterialParameters.TangentToWorld[2].xyz);
+		#if USE_ANALYTIC_DERIVATIVES
+		return NormalCurvatureToRoughness(InMaterialParameters.TangentToWorld[2].xyz, InMaterialParameters.WorldGeoNormal_DDX, InMaterialParameters.WorldGeoNormal_DDY);
+		#else
+		return NormalCurvatureToRoughness(InMaterialParameters.TangentToWorld[2].xyz);
+		#endif
 	#endif
 #else
 	return 0.f;
