VolumetricCloud - Fixed shader compilation warnings.

#tests pc editor dx12
#rnx

diff --git a/Engine/Shaders/Private/VolumetricRenderTarget.usf b/Engine/Shaders/Private/VolumetricRenderTarget.usf
--- a/Engine/Shaders/Private/VolumetricRenderTarget.usf
+++ b/Engine/Shaders/Private/VolumetricRenderTarget.usf
@@ -726,7 +726,7 @@ void ComposeVolumetricRTOverScenePS(
 	// Single bilinear sample
 //	OutputRt0 = SafeLoadVolumetricTexture(VolumeUV * VolumetricTextureSizeAndInvSize.xy);			// Closest
 	OutputRt0 = SafeLoadVolumetricSecondaryTexture(VolumeUV * VolumetricTextureSizeAndInvSize.xy);	// Furthest
-	float2 VolumeFrontDepth0 = SafeSampleVolumetricDepthTexture(VolumeUV);
+	float4 VolumeFrontDepth0 = SafeSampleVolumetricDepthTexture(VolumeUV);
 	OutputRt0 = ApplyFogToCloudSampleFromVPosAndDepth(CurResPixelCoord, VolumeFrontDepth0.x, OutputRt0);
 	return;
 
@@ -822,22 +822,22 @@ void ComposeVolumetricRTOverScenePS(
 
 #if 1
 	float4 VolumeRGBT0 = SafeSampleVolumetricTexture(VolumeUVOffset0);
-	float2 VolumeFrontDepth0 = SafeSampleVolumetricDepthTexture(VolumeUVOffset0);
+	float4 VolumeFrontDepth0 = SafeSampleVolumetricDepthTexture(VolumeUVOffset0);
 	float4 VolumeRGBT1 = SafeSampleVolumetricTexture(VolumeUVOffset1);
-	float2 VolumeFrontDepth1 = SafeSampleVolumetricDepthTexture(VolumeUVOffset1);
+	float4 VolumeFrontDepth1 = SafeSampleVolumetricDepthTexture(VolumeUVOffset1);
 	float4 VolumeRGBT2 = SafeSampleVolumetricTexture(VolumeUVOffset2);
-	float2 VolumeFrontDepth2 = SafeSampleVolumetricDepthTexture(VolumeUVOffset2);
+	float4 VolumeFrontDepth2 = SafeSampleVolumetricDepthTexture(VolumeUVOffset2);
 	float4 VolumeRGBT3 = SafeSampleVolumetricTexture(VolumeUVOffset3);
-	float2 VolumeFrontDepth3 = SafeSampleVolumetricDepthTexture(VolumeUVOffset3);
+	float4 VolumeFrontDepth3 = SafeSampleVolumetricDepthTexture(VolumeUVOffset3);
 #else
 	float4 VolumeRGBT0 = SafeLoadVolumetricTexture(VolumeUVOffset0 * VolumetricTextureSizeAndInvSize.xy);
-	float2 VolumeFrontDepth0 = SafeLoadVolumetricDepthTexture(VolumeUVOffset0 * VolumetricTextureSizeAndInvSize.xy);
+	float4 VolumeFrontDepth0 = SafeLoadVolumetricDepthTexture(VolumeUVOffset0 * VolumetricTextureSizeAndInvSize.xy);
 	float4 VolumeRGBT1 = SafeLoadVolumetricTexture(VolumeUVOffset1 * VolumetricTextureSizeAndInvSize.xy);
-	float2 VolumeFrontDepth1 = SafeLoadVolumetricDepthTexture(VolumeUVOffset1 * VolumetricTextureSizeAndInvSize.xy);
+	float4 VolumeFrontDepth1 = SafeLoadVolumetricDepthTexture(VolumeUVOffset1 * VolumetricTextureSizeAndInvSize.xy);
 	float4 VolumeRGBT2 = SafeLoadVolumetricTexture(VolumeUVOffset2 * VolumetricTextureSizeAndInvSize.xy);
-	float2 VolumeFrontDepth2 = SafeLoadVolumetricDepthTexture(VolumeUVOffset2 * VolumetricTextureSizeAndInvSize.xy);
+	float4 VolumeFrontDepth2 = SafeLoadVolumetricDepthTexture(VolumeUVOffset2 * VolumetricTextureSizeAndInvSize.xy);
 	float4 VolumeRGBT3 = SafeLoadVolumetricTexture(VolumeUVOffset3 * VolumetricTextureSizeAndInvSize.xy);
-	float2 VolumeFrontDepth3 = SafeLoadVolumetricDepthTexture(VolumeUVOffset3 * VolumetricTextureSizeAndInvSize.xy);
+	float4 VolumeFrontDepth3 = SafeLoadVolumetricDepthTexture(VolumeUVOffset3 * VolumetricTextureSizeAndInvSize.xy);
 #endif
 
 
