Fixed for PSO creation issue on metal

#tests pc dx12 editor, mac metal
#rnx

diff --git a/Engine/Shaders/Private/VolumetricRenderTarget.usf b/Engine/Shaders/Private/VolumetricRenderTarget.usf
--- a/Engine/Shaders/Private/VolumetricRenderTarget.usf
+++ b/Engine/Shaders/Private/VolumetricRenderTarget.usf
@@ -166,7 +166,7 @@ void ReconstructVolumetricRenderTargetPS(
 	out float4 OutputRt2 : SV_Target2
 #else
 	out float4 OutputRt0 : SV_Target0,
-	out float2 OutputRt1 : SV_Target1
+	out float4 OutputRt1 : SV_Target1
 #endif
 )
 {
@@ -231,9 +231,10 @@ void ReconstructVolumetricRenderTargetPS(
 			DepthData.SceneDepthFromViewKm, 
 			DepthData.MinMaxViewDepthKm);
 	#else
-		OutputRt1 = float2(
+		OutputRt1 = float4(
 			DepthData.CloudFrontDepthFromViewKm,
-			DepthData.SceneDepthFromViewKm);
+			DepthData.SceneDepthFromViewKm,
+			0.0, 0.0);	// unused components
 	#endif
 		return;
 	#endif
@@ -273,7 +274,10 @@ void ReconstructVolumetricRenderTargetPS(
 			OutputRt2 = float4(63000.0f, 63000.0f, 63000.0f, 63000.0f);
 		#else
 			OutputRt0 = RGBA;
-			OutputRt1 = float2(DepthData.CloudFrontDepthFromViewKm, DepthData.SceneDepthFromViewKm);
+			OutputRt1 = float4(
+				DepthData.CloudFrontDepthFromViewKm, 
+				DepthData.SceneDepthFromViewKm,
+				0.0, 0.0);	// unused components
 		#endif
 			return;
 		}
@@ -486,7 +490,10 @@ void ReconstructVolumetricRenderTargetPS(
 	OutputRt1 = RGBA2;
 	OutputRt2 = float4(DepthData.CloudFrontDepthFromViewKm, DepthData.SceneDepthFromViewKm, DepthData.MinMaxViewDepthKm);
 #else // PERMUTATION_CLOUD_MIN_AND_MAX_DEPTH
-	OutputRt1 = float2(DepthData.CloudFrontDepthFromViewKm, DepthData.SceneDepthFromViewKm);
+	OutputRt1 = float4(
+		DepthData.CloudFrontDepthFromViewKm, 
+		DepthData.SceneDepthFromViewKm,
+		0.0, 0.0);	// unused components
 #endif // PERMUTATION_CLOUD_MIN_AND_MAX_DEPTH
 
 #else // PERMUTATION_HISTORY_AVAILABLE
@@ -501,7 +508,10 @@ void ReconstructVolumetricRenderTargetPS(
 	OutputRt2 = float4(DepthData.CloudFrontDepthFromViewKm, DepthData.SceneDepthFromViewKm, DepthData.MinMaxViewDepthKm);
 #else // PERMUTATION_CLOUD_MIN_AND_MAX_DEPTH
 	FDepthData DepthData = SafeSampleTracingVolumetricDepthTexture(ScreenUV);
-	OutputRt1 = float2(DepthData.CloudFrontDepthFromViewKm, DepthData.SceneDepthFromViewKm);
+	OutputRt1 = float4(
+		DepthData.CloudFrontDepthFromViewKm,
+		DepthData.SceneDepthFromViewKm,
+		0.0, 0.0);	// unused components
 #endif // PERMUTATION_CLOUD_MIN_AND_MAX_DEPTH
 
 #endif // PERMUTATION_HISTORY_AVAILABLE
