ASTC runtime compression

diff --git a/Engine/Shaders/Private/ASTCCompressionCommon.ush b/Engine/Shaders/Private/ASTCCompressionCommon.ush
new file mode 100644
--- /dev/null
+++ b/Engine/Shaders/Private/ASTCCompressionCommon.ush
@@ -0,0 +1,744 @@
+// Copyright Epic Games, Inc. All Rights Reserved.
+
+/*=============================================================================
+	ASTCompressionCommon.ush:
+	Helpers for compute shader ASTC texture compression
+=============================================================================*/
+
+#pragma once
+
+#include "BlockCompressionCommon.ush"
+#include "BCCompressionCommon.ush"
+#include "GammaCorrectionCommon.ush"
+#include "IntegerSequenceEncoding.ush"
+
+
+float MaxNorm3(float3 v)
+{
+	v = abs(v);
+	return max(max(v.x, v.y), v.z);
+}
+
+float MaxNorm4(float4 v)
+{
+	v = abs(v);
+	return max(max(max(v.x, v.y), v.z), v.w);
+}
+
+void PCAOptimizeEndpoints3f(in float4 BlockRGBA[16], inout float4 BlockMin, inout float4 BlockMax)
+{
+	// Compute mean of all texels (3D).
+	float3 Mean = 0.0f;
+	float m00 = 0.0f, m01 = 0.0f, m02 = 0.0f;
+	float m11 = 0.0f, m12 = 0.0f, m22 = 0.0f;
+	for (uint i = 0; i < 16; i++)
+	{
+		float3 c = BlockRGBA[i].rgb;
+		Mean += c;
+		m00 += c.x * c.x;
+		m01 += c.x * c.y;
+		m02 += c.x * c.z;
+		m11 += c.y * c.y;
+		m12 += c.y * c.z;
+		m22 += c.z * c.z;
+	}
+	Mean /= 16.0f;
+
+	m00 -= 16.0f * Mean.x * Mean.x;
+	m01 -= 16.0f * Mean.x * Mean.y;
+	m02 -= 16.0f * Mean.x * Mean.z;
+	m11 -= 16.0f * Mean.y * Mean.y;
+	m12 -= 16.0f * Mean.y * Mean.z;
+	m22 -= 16.0f * Mean.z * Mean.z;
+
+	// Build candidate vectors from the "rows" of the covariance matrix.
+	float3 CovCol0 = float3(m00, m01, m02);
+	float3 CovCol1 = float3(m01, m11, m12);
+	float3 CovCol2 = float3(m02, m12, m22);
+
+	// Choose the candidate with the largest norm.
+	float3 CurrentAxis = CovCol0;
+	float CurrentNorm = MaxNorm3(CovCol0);
+
+	float CandidateNorm = MaxNorm3(CovCol1);
+	if (CandidateNorm > CurrentNorm)
+	{
+		CurrentAxis = CovCol1;
+		CurrentNorm = CandidateNorm;
+	}
+	CandidateNorm = MaxNorm3(CovCol2);
+	if (CandidateNorm > CurrentNorm)
+	{
+		CurrentAxis = CovCol2;
+		CurrentNorm = CandidateNorm;
+	}
+	
+	if (CurrentNorm != 0.0f)
+	{
+		CurrentAxis = normalize(CurrentAxis);
+		
+		// Power-Method-style iteration to refine the principal axis.
+		for (uint iteration = 0; iteration < 3; iteration++)
+		{
+			CurrentAxis = normalize(CovCol0 * CurrentAxis.x + CovCol1 * CurrentAxis.y + CovCol2 * CurrentAxis.z);
+		}
+	}
+	else
+	{
+		CurrentAxis = float3(1.0f, 0.0f, 0.0f);
+	}
+
+	// Project each texel onto the principal axis to find minimum / maximum projections.
+	float MinProj = 999999.0f;
+	float MaxProj = -999999.0f;
+	for (uint i = 0; i < 16; i++)
+	{
+		float Projection = dot(BlockRGBA[i].rgb - Mean, CurrentAxis);
+		MinProj = min(MinProj, Projection);
+		MaxProj = max(MaxProj, Projection);
+	}
+
+	// Reconstruct endpoints from the min / max projections, clamped to [0..1].
+	BlockMin.rgb = saturate(Mean + CurrentAxis * MinProj);
+	BlockMax.rgb = saturate(Mean + CurrentAxis * MaxProj);
+}
+
+void PCAOptimizeEndpoints3f_6x6(in float4 BlockRGBA[36], inout float3 BlockMin, inout float3 BlockMax)
+{
+	// Compute the mean and covariance of all texels (only RGB channels).
+	float3 Mean = 0.0f;
+	float m00 = 0.0f, m01 = 0.0f, m02 = 0.0f;
+	float m11 = 0.0f, m12 = 0.0f, m22 = 0.0f;
+	for (uint i = 0; i < 36; i++)
+	{
+		float3 c = BlockRGBA[i].rgb;
+		Mean += c;
+		m00 += c.x * c.x;
+		m01 += c.x * c.y;
+		m02 += c.x * c.z;
+		m11 += c.y * c.y;
+		m12 += c.y * c.z;
+		m22 += c.z * c.z;
+	}
+	Mean /= 36.0f;
+	
+	// Compute the covariance components.
+	m00 -= 36.0f * Mean.x * Mean.x;
+	m01 -= 36.0f * Mean.x * Mean.y;
+	m02 -= 36.0f * Mean.x * Mean.z;
+	m11 -= 36.0f * Mean.y * Mean.y;
+	m12 -= 36.0f * Mean.y * Mean.z;
+	m22 -= 36.0f * Mean.z * Mean.z;
+	
+	// Build the candidate axis vectors from the covariance matrix.
+	float3 CovCol0 = float3(m00, m01, m02);
+	float3 CovCol1 = float3(m01, m11, m12);
+	float3 CovCol2 = float3(m02, m12, m22);
+	
+	// Choose the candidate with the largest norm.
+	float3 CurrentAxis = CovCol0;
+	float CurrentNorm = MaxNorm3(CovCol0);
+	float CandidateNorm = MaxNorm3(CovCol1);
+	if (CandidateNorm > CurrentNorm)
+	{
+		CurrentAxis = CovCol1;
+		CurrentNorm = CandidateNorm;
+	}
+	CandidateNorm = MaxNorm3(CovCol2);
+	if (CandidateNorm > CurrentNorm)
+	{
+		CurrentAxis = CovCol2;
+		CurrentNorm = CandidateNorm;
+	}
+	
+	// Refine the principal axis using power iterations.
+	if (CurrentNorm != 0.0f)
+	{
+		CurrentAxis = normalize(CurrentAxis);
+		for (uint iteration = 0; iteration < 3; iteration++)
+		{
+			CurrentAxis = normalize(CovCol0 * CurrentAxis.x + CovCol1 * CurrentAxis.y + CovCol2 * CurrentAxis.z);
+		}
+	}
+	else
+	{
+		CurrentAxis = float3(1.0f, 0.0f, 0.0f);
+	}
+	
+	// Project each texel onto the principal axis to find the minimum and maximum projections.
+	float MinProj = 999999.0f;
+	float MaxProj = -999999.0f;
+	for (uint i = 0; i < 36; i++)
+	{
+		float Projection = dot(BlockRGBA[i].rgb - Mean, CurrentAxis);
+		MinProj = min(MinProj, Projection);
+		MaxProj = max(MaxProj, Projection);
+	}
+	
+	// Reconstruct the endpoints from the min and max projections, clamping results to [0,1].
+	BlockMin = saturate(Mean + CurrentAxis * MinProj);
+	BlockMax = saturate(Mean + CurrentAxis * MaxProj);
+}
+
+void PCAOptimizeEndpoints4f(in float4 BlockRGBA[16], inout float4 BlockMin, inout float4 BlockMax)
+{
+	// Compute mean over all texels
+	float4 Mean = float4(0.0, 0.0, 0.0, 0.0);
+	float m00 = 0.0, m01 = 0.0, m02 = 0.0, m03 = 0.0;
+	float m11 = 0.0, m12 = 0.0, m13 = 0.0;
+	float m22 = 0.0, m23 = 0.0;
+	float m33 = 0.0;
+	for (uint i = 0; i < 16; i++)
+	{
+		float4 c = BlockRGBA[i];
+		Mean += c;
+		m00 += c.x * c.x;
+		m01 += c.x * c.y;
+		m02 += c.x * c.z;
+		m03 += c.x * c.w;
+		m11 += c.y * c.y;
+		m12 += c.y * c.z;
+		m13 += c.y * c.w;
+		m22 += c.z * c.z;
+		m23 += c.z * c.w;
+		m33 += c.w * c.w;
+	}
+	Mean /= 16.0;
+
+	m00 -= 16.0 * Mean.x * Mean.x;
+	m01 -= 16.0 * Mean.x * Mean.y;
+	m02 -= 16.0 * Mean.x * Mean.z;
+	m03 -= 16.0 * Mean.x * Mean.w;
+	m11 -= 16.0 * Mean.y * Mean.y;
+	m12 -= 16.0 * Mean.y * Mean.z;
+	m13 -= 16.0 * Mean.y * Mean.w;
+	m22 -= 16.0 * Mean.z * Mean.z;
+	m23 -= 16.0 * Mean.z * Mean.w;
+	m33 -= 16.0 * Mean.w * Mean.w;
+
+	// Build candidate vectors from the symmetric 4x4 covariance matrix.
+	float4 CovCol0 = float4(m00, m01, m02, m03);
+	float4 CovCol1 = float4(m01, m11, m12, m13);
+	float4 CovCol2 = float4(m02, m12, m22, m23);
+	float4 CovCol3 = float4(m03, m13, m23, m33);
+
+	// Evaluate max norms of the candidates.
+	float4 CurrentAxis = CovCol0;
+	float CurrentNorm = MaxNorm4(CovCol0);
+
+	float CandidateNorm = MaxNorm4(CovCol1);
+	if (CandidateNorm > CurrentNorm)
+	{
+		CurrentAxis = CovCol1;
+		CurrentNorm = CandidateNorm;
+	}
+	CandidateNorm = MaxNorm4(CovCol2);
+	if (CandidateNorm > CurrentNorm)
+	{
+		CurrentAxis = CovCol2;
+		CurrentNorm = CandidateNorm;
+	}
+	CandidateNorm = MaxNorm4(CovCol3);
+	if (CandidateNorm > CurrentNorm)
+	{
+		CurrentAxis = CovCol3;
+		CurrentNorm = CandidateNorm;
+	}
+	if (CurrentNorm > 0.0)
+	{
+		CurrentAxis = normalize(CurrentAxis);
+		// Refine with a few iterations of the power method.
+		for (uint iteration = 0; iteration < 5; iteration++)
+		{
+			CurrentAxis = normalize(CovCol0 * CurrentAxis.x + CovCol1 * CurrentAxis.y + CovCol2 * CurrentAxis.z + CovCol3 * CurrentAxis.w);
+		}
+	}
+	else
+	{
+		CurrentAxis = float4(1.0, 0.0, 0.0, 0.0);
+	}
+	
+	// Project each texel (in 4D) onto Axis to locate end‐points.
+	float minProj = 1e10;
+	float maxProj = -1e10;
+	for (uint i = 0; i < 16; i++)
+	{
+		float proj = dot(BlockRGBA[i] - Mean, CurrentAxis);
+		minProj = min(minProj, proj);
+		maxProj = max(maxProj, proj);
+	}
+
+	BlockMin = saturate(Mean + CurrentAxis * minProj);
+	BlockMax = saturate(Mean + CurrentAxis * maxProj);
+}
+
+void LLSOptimizeEndpoints4f16(float4 Texels[16], uint WeightRange, inout float4 BlockMin, inout float4 BlockMax)
+{
+	float4 BlockVector = BlockMax - BlockMin;
+	float EndPoint0Pos = dot(BlockMin, BlockVector);
+	float EndPoint1Pos = dot(BlockMax, BlockVector);
+
+	float4 TexelSum = 0.0f;
+	float4 BetaTexelSum = 0.0f;
+	float BetaSum = 0.0f;
+	float BetaSqSum = 0.0f;
+
+	for (uint i = 0; i < 16; ++i)
+	{
+		float Pos = dot(Texels[i], BlockVector);
+		float NormalizedPos = saturate((Pos - EndPoint0Pos) / (EndPoint1Pos - EndPoint0Pos));
+		uint Index = (uint) (NormalizedPos * WeightRange + 0.5f);
+		float Beta = Index / (float) WeightRange;
+
+		TexelSum += Texels[i];
+		BetaTexelSum += Texels[i] * Beta;
+		BetaSum += Beta;
+		BetaSqSum += Beta * Beta;
+	}
+
+	float4 AlphaTexelSum = TexelSum - BetaTexelSum;
+	float AlphaSqSum = 16.0f - 2 * BetaSum + BetaSqSum;
+	float AlphaBetaSum = BetaSum - BetaSqSum;
+
+	float Det = AlphaSqSum * BetaSqSum - AlphaBetaSum * AlphaBetaSum;
+	if (abs(Det) > 0.1f)
+	{
+		float RcpDet = 1.0f / Det;
+		BlockMin = saturate((AlphaTexelSum * BetaSqSum - BetaTexelSum * AlphaBetaSum) * RcpDet);
+		BlockMax = saturate((BetaTexelSum * AlphaSqSum - AlphaTexelSum * AlphaBetaSum) * RcpDet);
+	}
+}
+
+void LLSOptimizeEndpoints3f16(float4 Texels[16], uint WeightRange, inout float3 BlockMin, inout float3 BlockMax)
+{
+	float3 BlockVector = BlockMax - BlockMin;
+	float EndPoint0Pos = dot(BlockMin, BlockVector);
+	float EndPoint1Pos = dot(BlockMax, BlockVector);
+
+	float3 TexelSum     = float3(0.0f, 0.0f, 0.0f);
+	float3 BetaTexelSum = float3(0.0f, 0.0f, 0.0f);
+	float BetaSum  = 0.0f;
+	float BetaSqSum = 0.0f;
+
+	for (uint i = 0; i < 16; ++i)
+	{
+		float Pos = dot(Texels[i].rgb, BlockVector);
+		float NormalizedPos = saturate((Pos - EndPoint0Pos) / (EndPoint1Pos - EndPoint0Pos));
+		uint Index = (uint)(NormalizedPos * WeightRange + 0.5f);
+		float Beta = Index / (float)WeightRange;
+
+		TexelSum     += Texels[i].rgb;
+		BetaTexelSum += Texels[i].rgb * Beta;
+		BetaSum    += Beta;
+		BetaSqSum  += Beta * Beta;
+	}
+
+	float3 AlphaTexelSum = TexelSum - BetaTexelSum;
+	float AlphaSqSum = 16.0f - 2.0f * BetaSum + BetaSqSum;
+	float AlphaBetaSum = BetaSum - BetaSqSum;
+
+	float Det = AlphaSqSum * BetaSqSum - AlphaBetaSum * AlphaBetaSum;
+	if (abs(Det) > 0.1f)
+	{
+		float RcpDet = 1.0f / Det;
+		BlockMin = saturate((AlphaTexelSum * BetaSqSum - BetaTexelSum * AlphaBetaSum) * RcpDet);
+		BlockMax = saturate((BetaTexelSum * AlphaSqSum - AlphaTexelSum * AlphaBetaSum) * RcpDet);
+	}
+}
+
+// https://registry.khronos.org/DataFormat/specs/1.3/dataformat.1.3.html See 23.10 Block Mode for more details
+uint BuildBlockMode(uint WeightQuantMethod, uint BlockSize)
+{
+	uint H = (BlockSize - 2) & 0x3;
+	uint W = (BlockSize - 4) & 0x3;
+
+	uint DualPlane = 0;
+	uint Precision = (WeightQuantMethod < 6) ? 0 : 1;
+	uint Range = (WeightQuantMethod % 6) + 2;
+
+	uint BlockMode = (Range >> 1) & 0x3;
+	BlockMode |= (Range & 0x1) << 4;
+	BlockMode |= H << 5;
+	BlockMode |= W << 7;
+	BlockMode |= Precision << 9;
+	BlockMode |= DualPlane << 10;
+	return BlockMode;
+}
+
+void EncodeWeights3f(in float4 BlockRGBA[16], uint WeightRange, uint WeightQuantMethod, float4 BlockMin, float4 BlockMax, inout uint4 BlockBuffer)
+{
+	float3 BlockVector = BlockMax.rgb - BlockMin.rgb;
+
+	float EndPoint0Pos = dot(BlockMin.rgb, BlockVector);
+	float EndPoint1Pos = dot(BlockMax.rgb, BlockVector);
+
+	uint WeightedBlocks[16];
+	// Quantize the weights using the endpoings and the Weight Range
+	for (uint ColorIndex = 0; ColorIndex < 16; ++ColorIndex)
+	{
+		float Pos = dot(BlockRGBA[ColorIndex].rgb, BlockVector);
+		float NormalizedPos = saturate((Pos - EndPoint0Pos) / (EndPoint1Pos - EndPoint0Pos));
+		uint QuantizedWeight = (uint) (NormalizedPos * WeightRange + 0.5f);
+
+		uint WeightIndex = WeightQuantMethod * WEIGHT_QUANTIZE_NUM + QuantizedWeight;
+		WeightedBlocks[ColorIndex] = AstcParameters.ScrambleTable[WeightIndex];
+	}
+
+	uint4 QuantWeights = uint4(0, 0, 0, 0);
+	uint QuantOffset = 0;
+	EncodeWeightsQuant16(WeightQuantMethod, WeightedBlocks, QuantWeights, QuantOffset);
+	
+	BlockBuffer.w = reversebits(QuantWeights.x);
+	BlockBuffer.z = reversebits(QuantWeights.y);
+	BlockBuffer.y = reversebits(QuantWeights.z);
+}
+
+void EncodeWeights4f(in float4 BlockRGBA[16], uint WeightRange, uint WeightQuantMethod, float4 BlockMin, float4 BlockMax, inout uint4 BlockBuffer)
+{
+	float4 BlockVector = BlockMax - BlockMin;
+
+	float EndPoint0Pos = dot(BlockMin, BlockVector);
+	float EndPoint1Pos = dot(BlockMax, BlockVector);
+
+	uint WeightedBlocks[16];
+	// Quantize the weights using the endpoings and the Weight Range
+	for (uint ColorIndex = 0; ColorIndex < 16; ++ColorIndex)
+	{
+		float Pos = dot(BlockRGBA[ColorIndex], BlockVector);
+		float NormalizedPos = saturate((Pos - EndPoint0Pos) / (EndPoint1Pos - EndPoint0Pos));
+		uint QuantizedWeight = (uint) (NormalizedPos * WeightRange + 0.5f);
+
+		uint WeightIndex = WeightQuantMethod * WEIGHT_QUANTIZE_NUM + QuantizedWeight;
+		WeightedBlocks[ColorIndex] = AstcParameters.ScrambleTable[WeightIndex];
+	}
+
+	uint4 QuantWeights = uint4(0, 0, 0, 0);
+	uint QuantOffset = 0;
+	EncodeWeightsQuant16(WeightQuantMethod, WeightedBlocks, QuantWeights, QuantOffset);
+	
+	BlockBuffer.w = reversebits(QuantWeights.x);
+	BlockBuffer.z = reversebits(QuantWeights.y);
+	BlockBuffer.y = reversebits(QuantWeights.z);
+}
+
+uint4 Compress_ASTC_RGBA_4x4(in float4 BlockRGBA[16])
+{
+	// Compute initial endpoints
+	float4 BlockMin = BlockRGBA[0];
+	float4 BlockMax = BlockRGBA[0];
+	for (uint TexelIndex = 1; TexelIndex < 16; ++TexelIndex)
+	{
+		BlockMin = min(BlockMin, BlockRGBA[TexelIndex]);
+		BlockMax = max(BlockMax, BlockRGBA[TexelIndex]);
+	}
+
+	uint BlockSize = 4;
+	
+	// Note: This is just a way to encode a void-extent like block - doesn't actually encode void-extent.
+	if (all(BlockMin == BlockMax))
+	{
+		uint ColorEndpointMode = 12; // LDR RGBA, direct
+		uint WeightQuantMethod = QUANT_4;
+		uint EndpointQuantMethod = QUANT_256;
+		uint WeightRange = 4u - 1u;
+		uint EndpointRange = 256u - 1u;
+
+		uint PartitionCount = 1;
+
+		uint BlockMode = BuildBlockMode(WeightQuantMethod, BlockSize);
+		BlockMode |= (PartitionCount - 1) << 11;
+		BlockMode |= ColorEndpointMode << 13;
+		
+		uint4 BlockBuffer = uint4(BlockMode, 0, 0, 0);
+		uint BlockOffset = 11 + 2 + 4;
+
+		uint4 QuantizeEndpoint = uint4(BlockMin * EndpointRange + 0.5f);
+
+		WriteBits(BlockBuffer, BlockOffset, QuantizeEndpoint.x, 8);
+		WriteBits(BlockBuffer, BlockOffset, 0xFF, 8);
+		WriteBits(BlockBuffer, BlockOffset, QuantizeEndpoint.y, 8);
+		WriteBits(BlockBuffer, BlockOffset, 0xFF, 8);
+		WriteBits(BlockBuffer, BlockOffset, QuantizeEndpoint.z, 8);
+		WriteBits(BlockBuffer, BlockOffset, 0xFF, 8);
+		WriteBits(BlockBuffer, BlockOffset, QuantizeEndpoint.w, 8);
+		WriteBits(BlockBuffer, BlockOffset, 0xFF, 8);
+		return BlockBuffer;
+	}
+
+	// Calculate the endpoint bit range
+	bool NoAlpha = BlockMin.a == 1.0f;
+
+	uint EndpointQuantMethod = QUANT_256;
+	uint EndpointRange = 256u - 1u;
+
+	uint4 BlockBuffer = uint4(0, 0, 0, 0);
+	uint BlockOffset = 11 + 2 + 4; // BlockMode + (PartitionCount-1) + CEM
+
+	if (NoAlpha)
+	{
+		uint ColorEndpointMode = 8; // LDR RGB, direct
+		uint WeightQuantMethod = QUANT_12;
+		uint WeightRange = 12u - 1u;
+		uint PartitionCount = 1;
+
+		uint BlockMode = BuildBlockMode(WeightQuantMethod, BlockSize);
+		BlockMode |= (PartitionCount - 1) << 11;
+		BlockMode |= ColorEndpointMode << 13;
+		BlockBuffer.x = BlockMode;
+	
+#ifdef ASTC_HIGH_PROFILE
+		PCAOptimizeEndpoints3f(BlockRGBA, BlockMin, BlockMax);
+#endif
+		//LLSOptimizeEndpoints3f16(BlockRGBA, WeightRange, BlockMin, BlockMax);
+	
+		// Inset the min/max
+		// See https://developer.download.nvidia.com/whitepapers/2007/Real-Time-YCoCg-DXT-Compression/Real-Time%20YCoCg-DXT%20Compression.pdf
+		{
+			float3 offset = (1.0f / ((WeightRange + 1) * 4.f)) * (BlockMax.rgb - BlockMin.rgb);
+			BlockMax.rgb = ceil((BlockMax.rgb - offset) * 255.0f) / 255.0f;
+			BlockMin.rgb = floor((BlockMin.rgb + offset) * 255.0f) / 255.0f;
+		}
+
+		float SumMin = BlockMin.x + BlockMin.y + BlockMin.z;
+		float SumMax = BlockMax.x + BlockMax.y + BlockMax.z;
+		if (SumMin > SumMax)
+		{
+			float4 tmp = BlockMin;
+			BlockMin = BlockMax;
+			BlockMax = tmp;
+		}
+
+		EncodeWeights3f(BlockRGBA, WeightRange, WeightQuantMethod, BlockMin, BlockMax, BlockBuffer);
+
+		uint4 QuantizeEndpoint0 = uint4(BlockMin * EndpointRange + 0.5f);
+		uint4 QuantizeEndpoint1 = uint4(BlockMax * EndpointRange + 0.5f);
+
+		WriteBits(BlockBuffer, BlockOffset, QuantizeEndpoint0.x, 8);
+		WriteBits(BlockBuffer, BlockOffset, QuantizeEndpoint1.x, 8);
+		WriteBits(BlockBuffer, BlockOffset, QuantizeEndpoint0.y, 8);
+		WriteBits(BlockBuffer, BlockOffset, QuantizeEndpoint1.y, 8);
+		WriteBits(BlockBuffer, BlockOffset, QuantizeEndpoint0.z, 8);
+		WriteBits(BlockBuffer, BlockOffset, QuantizeEndpoint1.z, 8);
+	}
+	else
+	{
+		uint ColorEndpointMode = 12; // LDR RGBA, direct
+		uint WeightQuantMethod = QUANT_6;
+		uint WeightRange = 6u - 1u;
+		uint PartitionCount = 1;
+
+		uint BlockMode = BuildBlockMode(WeightQuantMethod, BlockSize);
+		BlockMode |= (PartitionCount - 1) << 11;
+		BlockMode |= ColorEndpointMode << 13;
+		BlockBuffer.x = BlockMode;
+
+#ifdef ASTC_HIGH_PROFILE
+		PCAOptimizeEndpoints4f(BlockRGBA, BlockMin, BlockMax);
+#endif
+		//LLSOptimizeEndpoints4f16(BlockRGBA, WeightRange, BlockMin, BlockMax);
+	
+		// Inset the min/max
+		// See https://developer.download.nvidia.com/whitepapers/2007/Real-Time-YCoCg-DXT-Compression/Real-Time%20YCoCg-DXT%20Compression.pdf
+		{
+			float3 offset = (1.0f / ((WeightRange + 1) * 4.f)) * (BlockMax.rgb - BlockMin.rgb);
+			BlockMax.rgb = ceil((BlockMax.rgb - offset) * 255.0f) / 255.0f;
+			BlockMin.rgb = floor((BlockMin.rgb + offset) * 255.0f) / 255.0f;
+		}
+
+		float SumMin = BlockMin.x + BlockMin.y + BlockMin.z;
+		float SumMax = BlockMax.x + BlockMax.y + BlockMax.z;
+		if (SumMin > SumMax)
+		{
+			float4 tmp = BlockMin;
+			BlockMin = BlockMax;
+			BlockMax = tmp;
+		}
+
+		EncodeWeights4f(BlockRGBA, WeightRange, WeightQuantMethod, BlockMin, BlockMax, BlockBuffer);
+		uint4 QuantizeEndpoint0 = uint4(BlockMin * EndpointRange + 0.5f);
+		uint4 QuantizeEndpoint1 = uint4(BlockMax * EndpointRange + 0.5f);
+
+		WriteBits(BlockBuffer, BlockOffset, QuantizeEndpoint0.x, 8);
+		WriteBits(BlockBuffer, BlockOffset, QuantizeEndpoint1.x, 8);
+		WriteBits(BlockBuffer, BlockOffset, QuantizeEndpoint0.y, 8);
+		WriteBits(BlockBuffer, BlockOffset, QuantizeEndpoint1.y, 8);
+		WriteBits(BlockBuffer, BlockOffset, QuantizeEndpoint0.z, 8);
+		WriteBits(BlockBuffer, BlockOffset, QuantizeEndpoint1.z, 8);
+		WriteBits(BlockBuffer, BlockOffset, QuantizeEndpoint0.w, 8);
+		WriteBits(BlockBuffer, BlockOffset, QuantizeEndpoint1.w, 8);
+	}
+		
+	return BlockBuffer;
+}
+
+uint BuildBlockMode2(uint WeightQuantMethod, uint BlockX, uint BlockY)
+{
+	uint H = (BlockX - 6) & 0x3;
+	uint W = (BlockY - 6) & 0x3;
+
+	uint Range = (WeightQuantMethod % 6) + 2;
+	uint BlockMode = 0x100;
+	BlockMode |= (Range & 0x1) << 4;
+	BlockMode |= ((Range >> 1) & 0x3) << 2;
+	BlockMode |= H << 9;
+	BlockMode |= W << 5;
+	return BlockMode;
+}
+
+void EncodeWeights3f_6x6(in float4 BlockRGBA[36], uint WeightRange, uint WeightQuantMethod, float3 BlockMin, float3 BlockMax, inout uint4 BlockBuffer)
+{
+	float3 BlockVector = BlockMax - BlockMin;
+
+	float EndPoint0Pos = dot(BlockMin, BlockVector);
+	float EndPoint1Pos = dot(BlockMax, BlockVector);
+
+	uint WeightedBlocks[36];
+	// Quantize the weights using the endpoints and the Weight Range
+	for (uint ColorIndex = 0; ColorIndex < 36; ++ColorIndex)
+	{
+		float Pos = dot(BlockRGBA[ColorIndex].rgb, BlockVector);
+		float NormalizedPos = saturate((Pos - EndPoint0Pos) / (EndPoint1Pos - EndPoint0Pos));
+		uint QuantizedWeight = (uint) (NormalizedPos * WeightRange + 0.5f);
+
+		uint WeightIndex = WeightQuantMethod * WEIGHT_QUANTIZE_NUM + QuantizedWeight;
+		WeightedBlocks[ColorIndex] = AstcParameters.ScrambleTable[WeightIndex];
+	}
+
+	uint4 QuantWeights = uint4(0, 0, 0, 0);
+	uint QuantOffset = 0;
+	EncodeWeightsQuant36(WeightQuantMethod, WeightedBlocks, QuantWeights, QuantOffset);
+	
+	BlockBuffer.w |= reversebits(QuantWeights.x);
+	BlockBuffer.z |= reversebits(QuantWeights.y);
+	BlockBuffer.y |= reversebits(QuantWeights.z);
+}
+
+uint4 Compress_ASTC_RGB_6x6(in float4 BlockRGBA[36])
+{
+	// Compute initial endpoints
+	float3 BlockMin = BlockRGBA[0].rgb;
+	float3 BlockMax = BlockRGBA[0].rgb;
+	for (uint TexelIndex = 1; TexelIndex < 36; ++TexelIndex)
+	{
+		BlockMin = min(BlockMin, BlockRGBA[TexelIndex].rgb);
+		BlockMax = max(BlockMax, BlockRGBA[TexelIndex].rgb);
+	}
+
+	// Note: This is just a way to encode a void-extent like block - doesn't actually encode void-extent.
+	if (all(BlockMin == BlockMax))
+	{
+		uint ColorEndpointMode = 8; // LDR RGBA, direct
+		uint WeightQuantMethod = QUANT_4;
+		uint EndpointQuantMethod = QUANT_256;
+		uint WeightRange = 4u - 1u;
+		uint EndpointRange = 256u - 1u;
+
+		uint PartitionCount = 1;
+
+		uint BlockMode = BuildBlockMode2(WeightQuantMethod, 4, 4);
+		BlockMode |= (PartitionCount - 1) << 11;
+		BlockMode |= ColorEndpointMode << 13;
+		
+		uint4 BlockBuffer = uint4(BlockMode, 0, 0, 0);
+		uint BlockOffset = 11 + 2 + 4;
+
+		uint3 QuantizeEndpoint = uint3(BlockMin * EndpointRange + 0.5f);
+
+		WriteBits(BlockBuffer, BlockOffset, QuantizeEndpoint.x, 8);
+		WriteBits(BlockBuffer, BlockOffset, 0xFF, 8);
+		WriteBits(BlockBuffer, BlockOffset, QuantizeEndpoint.y, 8);
+		WriteBits(BlockBuffer, BlockOffset, 0xFF, 8);
+		WriteBits(BlockBuffer, BlockOffset, QuantizeEndpoint.z, 8);
+		WriteBits(BlockBuffer, BlockOffset, 0xFF, 8);
+		return BlockBuffer;
+	}
+	
+
+	uint4 BlockBuffer = uint4(0, 0, 0, 0);
+	uint BlockOffset = 11 + 2 + 4;
+
+	uint EndpointQuantMethod = QUANT_80;
+	uint EndpointRange = 80u - 1u;
+
+	uint ColorEndpointMode = 8; // LDR RGB, direct
+	uint WeightQuantMethod = QUANT_4;
+	uint WeightRange = 4u - 1u;
+	uint PartitionCount = 1;
+
+	uint BlockMode = BuildBlockMode2(WeightQuantMethod, 6, 6);
+	BlockMode |= (PartitionCount - 1) << 11;
+	BlockMode |= ColorEndpointMode << 13;
+	BlockBuffer.x = BlockMode;
+
+
+#ifdef ASTC_HIGH_PROFILE
+	PCAOptimizeEndpoints3f_6x6(BlockRGBA, BlockMin, BlockMax);
+#endif
+
+	// Inset the min/max
+	// See https://developer.download.nvidia.com/whitepapers/2007/Real-Time-YCoCg-DXT-Compression/Real-Time%20YCoCg-DXT%20Compression.pdf
+	{
+		float3 offset = (1.0f / ((WeightRange + 1) * 4.f)) * (BlockMax - BlockMin);
+		BlockMax = ceil((BlockMax - offset) * (float) EndpointRange) / (float) EndpointRange;
+		BlockMin = floor((BlockMin + offset) * (float) EndpointRange) / (float) EndpointRange;
+	}
+
+	float SumMin = BlockMin.x + BlockMin.y + BlockMin.z;
+	float SumMax = BlockMax.x + BlockMax.y + BlockMax.z;
+	if (SumMin > SumMax)
+	{
+		float3 tmp = BlockMin;
+		BlockMin = BlockMax;
+		BlockMax = tmp;
+	}
+
+	EncodeWeights3f_6x6(BlockRGBA, WeightRange, WeightQuantMethod, BlockMin, BlockMax, BlockBuffer);
+
+	uint3 QuantizeEndpoint0 = uint3(BlockMin * 255.0f + 0.5f);
+	uint3 QuantizeEndpoint1 = uint3(BlockMax * 255.0f + 0.5f);
+
+	uint EndpointNumbers[6] =
+	{
+		AstcParameters.ColorScrambleTable80[QuantizeEndpoint0.x], AstcParameters.ColorScrambleTable80[QuantizeEndpoint1.x],
+		AstcParameters.ColorScrambleTable80[QuantizeEndpoint0.y], AstcParameters.ColorScrambleTable80[QuantizeEndpoint1.y],
+		AstcParameters.ColorScrambleTable80[QuantizeEndpoint0.z], AstcParameters.ColorScrambleTable80[QuantizeEndpoint1.z]
+	};
+
+	EncodeEndpointsQuant6(EndpointQuantMethod, EndpointNumbers, BlockBuffer, BlockOffset);
+	return BlockBuffer;
+}
+
+uint4 CompressBlock_ASTC_RGBA(in float3 BlockRGB[16], in float BlockA[16])
+{
+	float4 B[16];
+	[unroll]
+	for (int i = 0; i < 16; ++i)
+	{
+		B[i] = float4(BlockRGB[i], BlockA[i]);
+	}
+	return Compress_ASTC_RGBA_4x4(B);
+}
+
+uint4 CompressBlock_ASTC_SRGBA(in float3 BlockRGB[16], in float BlockA[16])
+{
+	float4 B[16];
+	[unroll]
+	for (int i = 0; i < 16; ++i)
+	{
+		B[i] = float4(LinearToSrgb(BlockRGB[i]), BlockA[i]);
+	}
+	return Compress_ASTC_RGBA_4x4(B);
+}
+
+uint4 CompressBlock_ASTC_YCoCg(in float3 BlockRGB[16])
+{
+	float4 B[16];
+	[unroll]
+	for (int i = 0; i < 16; ++i)
+	{
+		float3 YCoCg = RGB2YCoCg(BlockRGB[i]);
+		B[i] = float4(YCoCg.yz, 0.0, YCoCg.x);
+	}
+	return Compress_ASTC_RGBA_4x4(B);
+}
+
diff --git a/Engine/Shaders/Private/IntegerSequenceEncoding.ush b/Engine/Shaders/Private/IntegerSequenceEncoding.ush
new file mode 100644
--- /dev/null
+++ b/Engine/Shaders/Private/IntegerSequenceEncoding.ush
@@ -0,0 +1,366 @@
+
+
+// astc specification Integer Sequence Encoding
+// https://registry.khronos.org/DataFormat/specs/1.3/dataformat.1.3.html#astc-integer-sequence-encoding
+
+#include "BitPacking.ush"
+
+// Quant  Levels
+#define QUANT_2		0
+#define QUANT_3		1
+#define QUANT_4		2
+#define QUANT_5		3
+#define QUANT_6		4
+#define QUANT_8		5
+#define QUANT_10	6
+#define QUANT_12	7
+#define QUANT_16	8
+#define QUANT_20	9
+#define QUANT_24	10
+#define QUANT_32	11
+#define QUANT_40	12
+#define QUANT_48	13
+#define QUANT_64	14
+#define QUANT_80	15
+#define QUANT_96	16
+#define QUANT_128	17
+#define QUANT_160	18
+#define QUANT_192	19
+#define QUANT_256	20
+
+// Bits Trits and Quints needed for quant encode
+static const uint BitTritQuantCounts[21][3] =
+{
+  //Bits Trits Quants Level
+	{ 1, 0, 0 }, //QUANT_2
+	{ 0, 1, 0 }, //QUANT_3
+	{ 2, 0, 0 }, //QUANT_4
+	{ 0, 0, 1 }, //QUANT_5
+	{ 1, 1, 0 }, //QUANT_6
+	{ 3, 0, 0 }, //QUANT_8
+	{ 1, 0, 1 }, //QUANT_10
+	{ 2, 1, 0 }, //QUANT_12
+	{ 4, 0, 0 }, //QUANT_16
+	{ 2, 0, 1 }, //QUANT_20
+	{ 3, 1, 0 }, //QUANT_24
+	{ 5, 0, 0 }, //QUANT_32
+	{ 3, 0, 1 }, //QUANT_40
+	{ 4, 1, 0 }, //QUANT_48
+	{ 6, 0, 0 }, //QUANT_64
+	{ 4, 0, 1 }, //QUANT_80
+	{ 5, 1, 0 }, //QUANT_96
+	{ 7, 0, 0 }, //QUANT_128
+	{ 5, 0, 1 }, //QUANT_160
+	{ 6, 1, 0 }, //QUANT_192
+	{ 8, 0, 0 }, //QUANT_256
+	};
+
+#define WEIGHT_QUANTIZE_NUM 32
+
+uint IntegerFromTrits(uint Trits0, uint Trits1, uint Trits2, uint Trits3, uint Trits4)
+{
+	uint Index = Trits4 * 81u + Trits3 * 27u + Trits2 * 9u + Trits1 * 3u + Trits0;
+	return AstcParameters.TritsToInteger[Index];
+}
+
+uint IntegerFromQuints(uint Quint0, uint Quint1, uint Quint2)
+{
+	uint Index = Quint2 * 25u + Quint1 * 5u + Quint0;
+	return AstcParameters.QuintsToInteger[Index];
+}
+
+void GetBitTritQuantCounts(uint Range, out uint BitCount, out uint TritCount, out uint QuintCount)
+{
+	BitCount = BitTritQuantCounts[Range][0];
+	TritCount = BitTritQuantCounts[Range][1];
+	QuintCount = BitTritQuantCounts[Range][2];
+}
+
+void ExtractLeastSignificantBits(uint NumBits, uint Num, out uint RemainingBits, out uint OutLSBits)
+{
+	uint BitsMask = (1u << NumBits) - 1u;
+	OutLSBits = Num & BitsMask;
+	RemainingBits = (Num >> NumBits) & 0xFFu;
+}
+
+// Encodes 5 quantized numbers into trit representation, and writes to the output buffeer
+void EncodeTrits(uint NumBits, uint Num0, uint Num1, uint Num2, uint Num3, uint Num4, inout uint4 OutBuffer, inout uint BufferOffset)
+{
+	uint TritBits0, TritBits1, TritBits2, TritBits3, TritBits4;
+	uint LSBits0, LSBits1, LSBits2, LSBits3, LSBits4;
+
+	ExtractLeastSignificantBits(NumBits, Num0, TritBits0, LSBits0);
+	ExtractLeastSignificantBits(NumBits, Num1, TritBits1, LSBits1);
+	ExtractLeastSignificantBits(NumBits, Num2, TritBits2, LSBits2);
+	ExtractLeastSignificantBits(NumBits, Num3, TritBits3, LSBits3);
+	ExtractLeastSignificantBits(NumBits, Num4, TritBits4, LSBits4);
+
+	// Packed 8 bit representation of the 5 trits
+	uint TritPack = IntegerFromTrits(TritBits0, TritBits1, TritBits2, TritBits3, TritBits4);
+
+	const uint TritLayout[5][3] =
+	{
+	// Offset Mask, NumBits
+		{ 0u, 3u, 2u },
+		{ 2u, 3u, 2u },
+		{ 4u, 1u, 1u },
+		{ 5u, 3u, 2u },
+		{ 7u, 1u, 1u },
+	};
+
+	uint TritPacks[5] = 
+	{
+		(TritPack >> TritLayout[0][0]) & TritLayout[0][1],
+		(TritPack >> TritLayout[1][0]) & TritLayout[1][1],
+		(TritPack >> TritLayout[2][0]) & TritLayout[2][1],
+		(TritPack >> TritLayout[3][0]) & TritLayout[3][1],
+		(TritPack >> TritLayout[4][0]) & TritLayout[4][1]
+	};
+
+	// Interleave the least significant bits with the trit pack value
+	WriteBits(OutBuffer, BufferOffset, LSBits0, NumBits);
+	WriteBits(OutBuffer, BufferOffset, TritPacks[0], TritLayout[0][2]);
+
+	WriteBits(OutBuffer, BufferOffset, LSBits1, NumBits);
+	WriteBits(OutBuffer, BufferOffset, TritPacks[1], TritLayout[1][2]);
+
+	WriteBits(OutBuffer, BufferOffset, LSBits2, NumBits);
+	WriteBits(OutBuffer, BufferOffset, TritPacks[2], TritLayout[2][2]);
+
+	WriteBits(OutBuffer, BufferOffset, LSBits3, NumBits);
+	WriteBits(OutBuffer, BufferOffset, TritPacks[3], TritLayout[3][2]);
+
+	WriteBits(OutBuffer, BufferOffset, LSBits4, NumBits);
+	WriteBits(OutBuffer, BufferOffset, TritPacks[4], TritLayout[4][2]);
+}
+
+// Encodes 3 quantized numbers into quints representation, and writes to the output buffeer
+void EncodeQuints(uint NumBits, uint Num0, uint Num1, uint Num2, inout uint4 OutBuffer, inout uint BufferOffset)
+{
+	uint QuintBits0, QuintBits1, QuintBits2;
+	uint LSBits0, LSBits1, LSBits2;
+
+	ExtractLeastSignificantBits(NumBits, Num0, QuintBits0, LSBits0);
+	ExtractLeastSignificantBits(NumBits, Num1, QuintBits1, LSBits1);
+	ExtractLeastSignificantBits(NumBits, Num2, QuintBits2, LSBits2);
+
+	// Packed 8 bit representation of the 3 quints
+	uint QuintPack = IntegerFromQuints(QuintBits0, QuintBits1, QuintBits2);
+
+	const uint QuintLayout[3][3] =
+	{
+	// Offset Mask, NumBits
+		{ 0u, 7u, 3u },
+		{ 3u, 3u, 2u },
+		{ 5u, 3u, 2u },
+	};
+
+	uint QuintPacks[3] = 
+	{
+		(QuintPack >> QuintLayout[0][0]) & QuintLayout[0][1],
+		(QuintPack >> QuintLayout[1][0]) & QuintLayout[1][1],
+		(QuintPack >> QuintLayout[2][0]) & QuintLayout[2][1]
+	};
+
+	// Interleave the least significant bits with the quint pack value
+	WriteBits(OutBuffer, BufferOffset, LSBits0, NumBits);
+	WriteBits(OutBuffer, BufferOffset, QuintPacks[0], QuintLayout[0][2]);
+
+	WriteBits(OutBuffer, BufferOffset, LSBits1, NumBits);
+	WriteBits(OutBuffer, BufferOffset, QuintPacks[1], QuintLayout[1][2]);
+
+	WriteBits(OutBuffer, BufferOffset, LSBits2, NumBits);
+	WriteBits(OutBuffer, BufferOffset, QuintPacks[2], QuintLayout[2][2]);
+}
+
+
+// Encodes 8 Numbers  
+// TODO - This could be repurposed for other size counts by passing EncodeCount as a value.
+//        This could also potentially be templated for any number count, but 8 and 16 work for basic u8 RGBA textures
+void EncodeEndpointsQuant(uint Method, uint Numbers[8], inout uint4 OutBuffer, inout uint BufferOffset)
+{
+	const uint EncodeCount = 8u;
+
+	uint BitCount; uint TritCount; uint QuintCount;
+	GetBitTritQuantCounts(Method, BitCount, TritCount, QuintCount);
+
+	uint InitialBufferOffset = BufferOffset;
+	if(TritCount == 1u)
+	{
+		EncodeTrits(BitCount, Numbers[0], Numbers[1], Numbers[2], Numbers[3], Numbers[4], OutBuffer, BufferOffset);
+		EncodeTrits(BitCount, Numbers[5], Numbers[6], Numbers[7], 0, 0, OutBuffer, BufferOffset);
+
+		// The OutBuffer size needs to be truncated for partial trits
+		uint FullTritBlockSize = 8u + 5u * BitCount;
+		uint FullBlockSize = FullTritBlockSize * EncodeCount;
+		uint TruncatedOffset = (FullBlockSize + 4u) / 5u;
+		BufferOffset = InitialBufferOffset + TruncatedOffset;
+		
+	}
+	else if(QuintCount == 1u)
+	{
+		EncodeQuints(BitCount, Numbers[0], Numbers[1], Numbers[2], OutBuffer, BufferOffset);
+		EncodeQuints(BitCount, Numbers[3], Numbers[4], Numbers[5], OutBuffer, BufferOffset);
+		EncodeQuints(BitCount, Numbers[6], Numbers[7], 0, OutBuffer, BufferOffset);
+
+		// The OutBuffer size needs to be truncated for partial quints
+		uint FullQuintBlockSize = 7u + 3u * BitCount;
+		uint FullBlockSize = FullQuintBlockSize * EncodeCount;
+		uint TruncatedOffset = (FullBlockSize + 2u) / 3u;
+		BufferOffset = InitialBufferOffset + TruncatedOffset;
+	}
+	else
+	{
+		[unroll]
+		for(uint Index = 0; Index < EncodeCount; ++Index)
+		{
+			WriteBits(OutBuffer, BufferOffset, Numbers[Index], BitCount);
+		}
+	}
+}
+
+void EncodeEndpointsQuant6(uint Method, uint Numbers[6], inout uint4 OutBuffer, inout uint BufferOffset)
+{
+	const uint EncodeCount = 6u;
+
+	uint BitCount; uint TritCount; uint QuintCount;
+	GetBitTritQuantCounts(Method, BitCount, TritCount, QuintCount);
+
+	uint InitialBufferOffset = BufferOffset;
+	if(TritCount == 1u)
+	{
+		EncodeTrits(BitCount, Numbers[0], Numbers[1], Numbers[2], Numbers[3], Numbers[4], OutBuffer, BufferOffset);
+		EncodeTrits(BitCount, Numbers[5], 0, 0, 0, 0, OutBuffer, BufferOffset);
+
+		// The OutBuffer size needs to be truncated for partial trits
+		uint FullTritBlockSize = 8u + 5u * BitCount;
+		uint FullBlockSize = FullTritBlockSize * EncodeCount;
+		uint TruncatedOffset = (FullBlockSize + 4u) / 5u;
+		BufferOffset = InitialBufferOffset + TruncatedOffset;
+		
+	}
+	else if(QuintCount == 1u)
+	{
+		EncodeQuints(BitCount, Numbers[0], Numbers[1], Numbers[2], OutBuffer, BufferOffset);
+		EncodeQuints(BitCount, Numbers[3], Numbers[4], Numbers[5], OutBuffer, BufferOffset);
+
+		// The OutBuffer size needs to be truncated for partial quints
+		uint FullQuintBlockSize = 7u + 3u * BitCount;
+		uint FullBlockSize = FullQuintBlockSize * EncodeCount;
+		uint TruncatedOffset = (FullBlockSize + 2u) / 3u;
+		BufferOffset = InitialBufferOffset + TruncatedOffset;
+	}
+	else
+	{
+		[unroll]
+		for(uint Index = 0; Index < EncodeCount; ++Index)
+		{
+			WriteBits(OutBuffer, BufferOffset, Numbers[Index], BitCount);
+		}
+	}
+}
+
+// Encodes 16 Numbers
+// TODO - This could be repurposed for other size counts by passing EncodeCount as a value.
+//        This could also potentially be templated for any number count, but 8 and 16 work for basic u8 RGBA textures
+void EncodeWeightsQuant16(uint Method, uint Numbers[16], inout uint4 OutBuffer, inout uint BufferOffset)
+{
+	const uint EncodeCount = 16u;
+
+	uint BitCount; uint TritCount; uint QuintCount;
+	GetBitTritQuantCounts(Method, BitCount, TritCount, QuintCount);
+
+	uint InitialBufferOffset = BufferOffset;
+	if(TritCount == 1u)
+	{
+		EncodeTrits(BitCount, Numbers[0], Numbers[1], Numbers[2], Numbers[3], Numbers[4], OutBuffer, BufferOffset);
+		EncodeTrits(BitCount, Numbers[5], Numbers[6], Numbers[7], Numbers[8], Numbers[9], OutBuffer, BufferOffset);
+		EncodeTrits(BitCount, Numbers[10], Numbers[11], Numbers[12], Numbers[13], Numbers[14], OutBuffer, BufferOffset);
+		EncodeTrits(BitCount, Numbers[15], 0, 0, 0, 0, OutBuffer, BufferOffset);
+
+		// The OutBuffer size needs to be truncated for partial trits
+		uint FullTritBlockSize = 8u + 5u * BitCount;
+		uint FullBlockSize = FullTritBlockSize * EncodeCount;
+		uint TruncatedOffset = (FullBlockSize + 4u) / 5u;
+		BufferOffset = InitialBufferOffset + TruncatedOffset;
+		
+	}
+	else if(QuintCount == 1u)
+	{
+		EncodeQuints(BitCount, Numbers[0], Numbers[1], Numbers[2], OutBuffer, BufferOffset);
+		EncodeQuints(BitCount, Numbers[3], Numbers[4], Numbers[5], OutBuffer, BufferOffset);
+		EncodeQuints(BitCount, Numbers[6], Numbers[7], 0, OutBuffer, BufferOffset);
+
+		// The OutBuffer size needs to be truncated for partial quints
+		uint FullQuintBlockSize = 7u + 3u * BitCount;
+		uint FullBlockSize = FullQuintBlockSize * EncodeCount;
+		uint TruncatedOffset = (FullBlockSize + 2u) / 3u;
+		BufferOffset = InitialBufferOffset + TruncatedOffset;
+	}
+	else
+	{
+		[unroll]
+		for(uint Index = 0; Index < EncodeCount; ++Index)
+		{
+			WriteBits(OutBuffer, BufferOffset, Numbers[Index], BitCount);
+		}
+	}
+}
+
+void EncodeWeightsQuant36(uint Method, uint Numbers[36], inout uint4 OutBuffer, inout uint BufferOffset)
+{
+	const uint EncodeCount = 36u;
+
+	uint BitCount; uint TritCount; uint QuintCount;
+	GetBitTritQuantCounts(Method, BitCount, TritCount, QuintCount);
+
+	uint InitialBufferOffset = BufferOffset;
+	if(TritCount == 1u)
+	{
+        EncodeTrits(BitCount, Numbers[0], Numbers[1], Numbers[2], Numbers[3], Numbers[4], OutBuffer, BufferOffset);
+        EncodeTrits(BitCount, Numbers[5], Numbers[6], Numbers[7], Numbers[8], Numbers[9], OutBuffer, BufferOffset);
+        EncodeTrits(BitCount, Numbers[10], Numbers[11], Numbers[12], Numbers[13], Numbers[14], OutBuffer, BufferOffset);
+        EncodeTrits(BitCount, Numbers[15], Numbers[16], Numbers[17], Numbers[18], Numbers[19], OutBuffer, BufferOffset);
+        EncodeTrits(BitCount, Numbers[20], Numbers[21], Numbers[22], Numbers[23], Numbers[24], OutBuffer, BufferOffset);
+        EncodeTrits(BitCount, Numbers[25], Numbers[26], Numbers[27], Numbers[28], Numbers[29], OutBuffer, BufferOffset);
+        EncodeTrits(BitCount, Numbers[30], Numbers[31], Numbers[32], Numbers[33], Numbers[34], OutBuffer, BufferOffset);
+        EncodeTrits(BitCount, Numbers[35], 0, 0, 0, 0, OutBuffer, BufferOffset);
+
+		// The OutBuffer size needs to be truncated for partial trits
+		uint FullTritBlockSize = 8u + 5u * BitCount;
+		uint FullBlockSize = FullTritBlockSize * EncodeCount;
+		uint TruncatedOffset = (FullBlockSize + 4u) / 5u;
+		BufferOffset = InitialBufferOffset + TruncatedOffset;
+	}
+	else if(QuintCount == 1u)
+	{
+        EncodeQuints(BitCount, Numbers[0], Numbers[1], Numbers[2], OutBuffer, BufferOffset);
+        EncodeQuints(BitCount, Numbers[3], Numbers[4], Numbers[5], OutBuffer, BufferOffset);
+        EncodeQuints(BitCount, Numbers[6], Numbers[7], Numbers[8], OutBuffer, BufferOffset);
+        EncodeQuints(BitCount, Numbers[9], Numbers[10], Numbers[11], OutBuffer, BufferOffset);
+        EncodeQuints(BitCount, Numbers[12], Numbers[13], Numbers[14], OutBuffer, BufferOffset);
+        EncodeQuints(BitCount, Numbers[15], Numbers[16], Numbers[17], OutBuffer, BufferOffset);
+        EncodeQuints(BitCount, Numbers[18], Numbers[19], Numbers[20], OutBuffer, BufferOffset);
+        EncodeQuints(BitCount, Numbers[21], Numbers[22], Numbers[23], OutBuffer, BufferOffset);
+        EncodeQuints(BitCount, Numbers[24], Numbers[25], Numbers[26], OutBuffer, BufferOffset);
+        EncodeQuints(BitCount, Numbers[27], Numbers[28], Numbers[29], OutBuffer, BufferOffset);
+        EncodeQuints(BitCount, Numbers[30], Numbers[31], Numbers[32], OutBuffer, BufferOffset);
+        EncodeQuints(BitCount, Numbers[33], Numbers[34], Numbers[35], OutBuffer, BufferOffset);
+
+		// The OutBuffer size needs to be truncated for partial quints
+		uint FullQuintBlockSize = 7u + 3u * BitCount;
+		uint FullBlockSize = FullQuintBlockSize * EncodeCount;
+		uint TruncatedOffset = (FullBlockSize + 2u) / 3u;
+		BufferOffset = InitialBufferOffset + TruncatedOffset;
+	}
+	else
+	{
+		[unroll]
+		for(uint Index = 0; Index < EncodeCount; ++Index)
+		{
+			WriteBits(OutBuffer, BufferOffset, Numbers[Index], BitCount);
+		}
+	}
+}
+
+
diff --git a/Engine/Shaders/Private/VirtualTextureCompress.usf b/Engine/Shaders/Private/VirtualTextureCompress.usf
--- a/Engine/Shaders/Private/VirtualTextureCompress.usf
+++ b/Engine/Shaders/Private/VirtualTextureCompress.usf
@@ -9,12 +9,17 @@
 #include "Common.ush"
 #include "BlockCompressionCommon.ush"
 #include "BCCompressionCommon.ush"
+#include "ASTCCompressionCommon.ush"
 #include "ETCCompressionCommon.ush"
 
 #ifndef ETC_PROFILE
 #define ETC_PROFILE 0
 #endif
 
+#ifndef ASTC_PROFILE
+#define ASTC_PROFILE 0
+#endif
+
 #ifndef PACK_RG32_RGBA16
 #define PACK_RG32_RGBA16 0
 #endif
@@ -87,7 +92,9 @@ void CompressBaseColorCS(
 
 	float3 BlockBaseColor[16];
 	ReadBlockRGB(RenderTexture0, TextureSampler0, SampleUV, TexelUVSize, BlockBaseColor);
-#if ETC_PROFILE
+#if ASTC_PROFIlE
+	OutCompressTexture0_64bit[ThreadId.xy + DestPos0] = PackRG32(CompressBlock_ETC2_SRGB(BlockBaseColor));
+#elif ETC_PROFILE
 	OutCompressTexture0_64bit[ThreadId.xy + DestPos0] = PackRG32(CompressBlock_ETC2_SRGB(BlockBaseColor));
 #else
 	OutCompressTexture0_64bit[ThreadId.xy + DestPos0] = PackRG32(CompressBC1Block_SRGB(BlockBaseColor));
@@ -133,7 +140,10 @@ void CompressBaseColorNormalSpecularCS(
 	{
 		BlockSpecular[i].z = round(BlockSpecular[i].z);
 	}
-#if ETC_PROFILE
+#if ASTC_PROFILE
+	OutCompressTexture0_128bit[ThreadId.xy + DestPos0] = CompressBlock_ASTC_SRGBA(BlockBaseColor, BlockNormalX);
+	OutCompressTexture1_128bit[ThreadId.xy + DestPos1] = CompressBlock_ASTC_RGBA(BlockSpecular, BlockNormalY);
+#elif ETC_PROFILE
 	OutCompressTexture0_128bit[ThreadId.xy + DestPos0] = CompressBlock_ETC2_SRGBA(BlockBaseColor, BlockNormalX);
 	OutCompressTexture1_128bit[ThreadId.xy + DestPos1] = CompressBlock_ETC2_RGBA(BlockSpecular, BlockNormalY);
 #else
@@ -166,7 +176,10 @@ void CompressBaseColorNormalRoughnessCS(
 	const float BlockNormalZ[16]={ 1.0, 1.0, 1.0, 1.0, 1.0, 1.0, 1.0, 1.0, 
 								   1.0, 1.0, 1.0, 1.0, 1.0, 1.0, 1.0, 1.0};
 	ReadBlockRGB(RenderTexture1, TextureSampler2, SampleUV, TexelUVSize, BlockNormalXRougnnessY);
-#if ETC_PROFILE
+#if ASTC_PROFILE
+	OutCompressTexture0_64bit[ThreadId.xy + DestPos0] = PackRG32(CompressBlock_ETC2_RGB(BlockBaseColor));
+	OutCompressTexture1_128bit[ThreadId.xy + DestPos1] = CompressBlock_ASTC_RGBA(BlockNormalXRougnnessY, BlockNormalZ);
+#elif ETC_PROFILE
 	OutCompressTexture0_64bit[ThreadId.xy + DestPos0] = PackRG32(CompressBlock_ETC2_RGB(BlockBaseColor));
 	OutCompressTexture1_128bit[ThreadId.xy + DestPos1] = CompressBlock_ETC2_RGBA(BlockNormalXRougnnessY, BlockNormalZ);
 #else
@@ -222,7 +235,11 @@ void CompressBaseColorNormalSpecularYCoCgCS(
 		BlockSpecular[i].z = round(BlockSpecular[i].z);
 	}
 
-#if ETC_PROFILE
+#if ASTC_PROFILE
+	OutCompressTexture0_128bit[ThreadId.xy + DestPos0] = CompressBlock_ASTC_YCoCg(BlockBaseColor);
+	OutCompressTexture1_128bit[ThreadId.xy + DestPos1] = CompressBlock_ETC2_RG(BlockNormalX, BlockNormalY);
+	OutCompressTexture2_64bit[ThreadId.xy + DestPos2] = PackRG32(CompressBlock_ETC2_RGB(BlockSpecular));
+#elif ETC_PROFILE
 	OutCompressTexture0_128bit[ThreadId.xy + DestPos0] = CompressBlock_ETC2_YCoCg(BlockBaseColor);
 	OutCompressTexture1_128bit[ThreadId.xy + DestPos1] = CompressBlock_ETC2_RG(BlockNormalX, BlockNormalY);
 	OutCompressTexture2_64bit[ThreadId.xy + DestPos2] = PackRG32(CompressBlock_ETC2_RGB(BlockSpecular));
@@ -280,7 +297,11 @@ void CompressBaseColorNormalSpecularMaskYCoCgCS(
 		BlockSpecular[i].z = round(BlockSpecular[i].z);
 	}
 
-#if ETC_PROFILE
+#if ASTC_PROFILE
+	OutCompressTexture0_128bit[ThreadId.xy + DestPos0] = CompressBlock_ASTC_YCoCg(BlockBaseColor);
+	OutCompressTexture1_128bit[ThreadId.xy + DestPos1] = CompressBlock_ETC2_RG(BlockNormalX, BlockNormalY);
+	OutCompressTexture2_128bit[ThreadId.xy + DestPos2] = CompressBlock_ASTC_RGBA(BlockSpecular, BlockMask);
+#elif ETC_PROFILE
 	OutCompressTexture0_128bit[ThreadId.xy + DestPos0] = CompressBlock_ETC2_YCoCg(BlockBaseColor);
 	OutCompressTexture1_128bit[ThreadId.xy + DestPos1] = CompressBlock_ETC2_RG(BlockNormalX, BlockNormalY);
 	OutCompressTexture2_128bit[ThreadId.xy + DestPos2] = CompressBlock_ETC2_RGBA(BlockSpecular, BlockMask);
@@ -330,7 +351,9 @@ void CompressMask4CS(
 	float BlockW[16];
 	ReadBlockX(RenderTexture1, TextureSampler0, SampleUV, TexelUVSize, BlockW);
 
-#if ETC_PROFILE
+#if ASTC_PROFILE
+	OutCompressTexture0_128bit[ThreadId.xy + DestPos0] = CompressBlock_ASTC_RGBA(BlockXYZ, BlockW);
+#elif ETC_PROFILE
 	OutCompressTexture0_128bit[ThreadId.xy + DestPos0] = CompressBlock_ETC2_RGBA(BlockXYZ, BlockW);
 #else
 	OutCompressTexture0_128bit[ThreadId.xy + DestPos0] = CompressBC3Block(BlockXYZ, BlockW);
@@ -365,7 +388,9 @@ void CompressDisplacementCS(
 
 	float Block[16];
 	ReadBlockX(RenderTexture0, TextureSampler0, SampleUV, TexelUVSize, Block);
-#if ETC_PROFILE
+#if ASTC_PROFILE
+	OutCompressTexture0_64bit[ThreadId.xy + DestPos0] = PackRG32(CompressBlock_ETC2_Alpha(Block));
+#elif ETC_PROFILE
 	OutCompressTexture0_64bit[ThreadId.xy + DestPos0] = PackRG32(CompressBlock_ETC2_Alpha(Block));
 #else
 	OutCompressTexture0_64bit[ThreadId.xy + DestPos0] = PackRG32(CompressBC4Block(Block));
diff --git a/Engine/Source/Runtime/Renderer/Private/VT/RuntimeVirtualTextureRender.cpp b/Engine/Source/Runtime/Renderer/Private/VT/RuntimeVirtualTextureRender.cpp
--- a/Engine/Source/Runtime/Renderer/Private/VT/RuntimeVirtualTextureRender.cpp
+++ b/Engine/Source/Runtime/Renderer/Private/VT/RuntimeVirtualTextureRender.cpp
@@ -26,6 +26,7 @@
 #include "SimpleMeshDrawCommandPass.h"
 #include "StaticMeshBatch.h"
 #include "SceneRendering.h"
+#include "RHIResourceUtils.h"
 
 CSV_DECLARE_CATEGORY_EXTERN(VirtualTexturing);
 
@@ -58,6 +59,18 @@ namespace RuntimeVirtualTexture
 	    RenderCaptureNextRVTPagesDraws,
 	    TEXT("Trigger a render capture during the next RVT RenderPages draw calls."));
 
+	static TAutoConsoleVariable<bool> CVarRVTAstc(
+		TEXT("r.VT.RVT.ASTC"),
+		0,
+		TEXT("Use ASTC compression instead of ETC2 when the hardware supports it."),
+		ECVF_ReadOnly);
+
+	static TAutoConsoleVariable<bool> CVarRVTAstcHigh(
+		TEXT("r.VT.RVT.ASTC.High"),
+		0,
+		TEXT("When using ASTC compression, produce higher quality output at roughly 2x the time spent encoding."),
+		ECVF_Default);
+
 	BEGIN_GLOBAL_SHADER_PARAMETER_STRUCT(FEtcParameters, )
 		SHADER_PARAMETER_ARRAY(FVector4f, ALPHA_DISTANCE_TABLES, [16])
 		SHADER_PARAMETER_ARRAY(FVector4f, RGB_DISTANCE_TABLES, [8])
@@ -108,6 +121,369 @@ namespace RuntimeVirtualTexture
 		return EtcParametersUniformBuffer.GetUniformBufferRef();
 	}
 
+	static const uint8 TritsToInteger[243] = 
+	{
+		0, 1, 2,
+		4, 5, 6,
+		8, 9, 10,
+
+		16, 17, 18,
+		20, 21, 22,
+		24, 25, 26,
+
+		3, 7, 15,
+		19, 23, 27,
+		12, 13, 14,
+
+		32, 33, 34,
+		36, 37, 38,
+		40, 41, 42,
+
+		48, 49, 50,
+		52, 53, 54,
+		56, 57, 58,
+
+		35, 39, 47,
+		51, 55, 59,
+		44, 45, 46,
+
+		64, 65, 66,
+		68, 69, 70,
+		72, 73, 74,
+
+		80, 81, 82,
+		84, 85, 86,
+		88, 89, 90,
+
+		67, 71, 79,
+		83, 87, 91,
+		76, 77, 78,
+
+		128, 129, 130,
+		132, 133, 134,
+		136, 137, 138,
+
+		144, 145, 146,
+		148, 149, 150,
+		152, 153, 154,
+
+		131, 135, 143,
+		147, 151, 155,
+		140, 141, 142,
+
+		160, 161, 162,
+		164, 165, 166,
+		168, 169, 170,
+
+		176, 177, 178,
+		180, 181, 182,
+		184, 185, 186,
+
+		163, 167, 175,
+		179, 183, 187,
+		172, 173, 174,
+
+		192, 193, 194,
+		196, 197, 198,
+		200, 201, 202,
+
+		208, 209, 210,
+		212, 213, 214,
+		216, 217, 218,
+
+		195, 199, 207,
+		211, 215, 219,
+		204, 205, 206,
+
+		96, 97, 98,
+		100, 101, 102,
+		104, 105, 106,
+
+		112, 113, 114,
+		116, 117, 118,
+		120, 121, 122,
+
+		99, 103, 111,
+		115, 119, 123,
+		108, 109, 110,
+
+		224, 225, 226,
+		228, 229, 230,
+		232, 233, 234,
+
+		240, 241, 242,
+		244, 245, 246,
+		248, 249, 250,
+
+		227, 231, 239,
+		243, 247, 251,
+		236, 237, 238,
+
+		28, 29, 30,
+		60, 61, 62,
+		92, 93, 94,
+
+		156, 157, 158,
+		188, 189, 190,
+		220, 221, 222,
+
+		31, 63, 127,
+		159, 191, 255,
+		252, 253, 254
+	};
+
+	static const uint8 QuintsToInteger[125] =
+	{
+		0, 1, 2, 3, 4,
+		8, 9, 10, 11, 12,
+		16, 17, 18, 19, 20,
+		24, 25, 26, 27, 28,
+		5, 13, 21, 29, 6,
+
+		32, 33, 34, 35, 36,
+		40, 41, 42, 43, 44,
+		48, 49, 50, 51, 52,
+		56, 57, 58, 59, 60,
+		37, 45, 53, 61, 14,
+
+		64, 65, 66, 67, 68,
+		72, 73, 74, 75, 76,
+		80, 81, 82, 83, 84,
+		88, 89, 90, 91, 92,
+		69, 77, 85, 93, 22,
+
+		96, 97, 98, 99, 100,
+		104, 105, 106, 107, 108,
+		112, 113, 114, 115, 116,
+		120, 121, 122, 123, 124,
+		101, 109, 117, 125, 30,
+
+		102, 103, 70, 71, 38,
+		110, 111, 78, 79, 46,
+		118, 119, 86, 87, 54,
+		126, 127, 94, 95, 62,
+		39, 47, 55, 63, 31
+	};
+
+	// from [ARM:astc-encoder] quantization_and_transfer_table quant_and_xfer_tables
+	#define WEIGHT_QUANTIZE_NUM 32
+	static const uint8 ScrambleTable[12 * WEIGHT_QUANTIZE_NUM] = {
+		// quantization method 0, range 0..1
+		//{
+			0, 1,
+			0, 0, 0, 0, 0, 0, 0, 0, 0, 0, 0, 0, 0, 0, 0, 0, 0, 0, 0, 0, 0, 0, 0, 0, 0, 0, 0, 0, 0, 0,
+		//},
+		// quantization method 1, range 0..2
+		//{
+			0, 1, 2,
+			0, 0, 0, 0, 0, 0, 0, 0, 0, 0, 0, 0, 0, 0, 0, 0, 0, 0, 0, 0, 0, 0, 0, 0, 0, 0, 0, 0, 0,
+		//},
+		// quantization method 2, range 0..3
+		//{
+			0, 1, 2, 3,
+			0, 0, 0, 0, 0, 0, 0, 0, 0, 0, 0, 0, 0, 0, 0, 0, 0, 0, 0, 0, 0, 0, 0, 0, 0, 0, 0, 0,
+		//},
+		// quantization method 3, range 0..4
+		//{
+			0, 1, 2, 3, 4,
+			0, 0, 0, 0, 0, 0, 0, 0, 0, 0, 0, 0, 0, 0, 0, 0, 0, 0, 0, 0, 0, 0, 0, 0, 0, 0, 0,
+		//},
+		// quantization method 4, range 0..5
+		//{
+			0, 2, 4, 5, 3, 1,
+			0, 0, 0, 0, 0, 0, 0, 0, 0, 0, 0, 0, 0, 0, 0, 0, 0, 0, 0, 0, 0, 0, 0, 0, 0, 0,
+		//},
+		// quantization method 5, range 0..7
+		//{
+			0, 1, 2, 3, 4, 5, 6, 7,
+			0, 0, 0, 0, 0, 0, 0, 0, 0, 0, 0, 0, 0, 0, 0, 0, 0, 0, 0, 0, 0, 0, 0, 0,
+		//},
+		// quantization method 6, range 0..9
+		//{
+			0, 2, 4, 6, 8, 9, 7, 5, 3, 1,
+			0, 0, 0, 0, 0, 0, 0, 0, 0, 0, 0, 0, 0, 0, 0, 0, 0, 0, 0, 0, 0, 0,
+		//},
+		// quantization method 7, range 0..11
+		//{
+			0, 4, 8, 2, 6, 10, 11, 7, 3, 9, 5, 1,
+			0, 0, 0, 0, 0, 0, 0, 0, 0, 0, 0, 0, 0, 0, 0, 0, 0, 0, 0, 0,
+		//},
+		// quantization method 8, range 0..15
+		//{
+			0, 1, 2, 3, 4, 5, 6, 7, 8, 9, 10, 11, 12, 13, 14, 15,
+			0, 0, 0, 0, 0, 0, 0, 0, 0, 0, 0, 0, 0, 0, 0, 0,
+		//},
+		// quantization method 9, range 0..19
+		//{
+			0, 4, 8, 12, 16, 2, 6, 10, 14, 18, 19, 15, 11, 7, 3, 17, 13, 9, 5, 1,
+			0, 0, 0, 0, 0, 0, 0, 0, 0, 0, 0, 0,
+		//},
+		// quantization method 10, range 0..23
+		//{
+			0, 8, 16, 2, 10, 18, 4, 12, 20, 6, 14, 22, 23, 15, 7, 21, 13, 5, 19,
+			11, 3, 17, 9, 1, 0, 0, 0, 0, 0, 0, 0, 0,
+		//},
+		// quantization method 11, range 0..31
+		//{
+			0, 1, 2, 3, 4, 5, 6, 7, 8, 9, 10, 11, 12, 13, 14, 15, 16, 17, 18, 19,
+			20, 21, 22, 23, 24, 25, 26, 27, 28, 29, 30, 31,
+		//}
+	};
+
+	static const uint8 ColorScrambleTable48[256] = {
+		  0,   0,   0,  16,  16,  16,  16,  16,  32,  32,  32,  32,  32,  32,   2,   2,
+		  2,   2,   2,  18,  18,  18,  18,  18,  34,  34,  34,  34,  34,  34,   4,   4,
+		  4,   4,   4,  20,  20,  20,  20,  20,  20,  36,  36,  36,  36,  36,   6,   6,
+		  6,   6,   6,  22,  22,  22,  22,  22,  22,  38,  38,  38,  38,  38,  38,   8,
+		  8,   8,   8,   8,  24,  24,  24,  24,  24,  40,  40,  40,  40,  40,  40,  10,
+		 10,  10,  10,  10,  26,  26,  26,  26,  26,  42,  42,  42,  42,  42,  42,  12,
+		 12,  12,  12,  12,  28,  28,  28,  28,  28,  28,  44,  44,  44,  44,  44,  14,
+		 14,  14,  14,  14,  30,  30,  30,  30,  30,  30,  46,  46,  46,  46,  46,  46,
+		 47,  47,  47,  47,  47,  47,  31,  31,  31,  31,  31,  31,  15,  15,  15,  15,
+		 15,  45,  45,  45,  45,  45,  29,  29,  29,  29,  29,  29,  13,  13,  13,  13,
+		 13,  43,  43,  43,  43,  43,  43,  27,  27,  27,  27,  27,  11,  11,  11,  11,
+		 11,  41,  41,  41,  41,  41,  41,  25,  25,  25,  25,  25,   9,   9,   9,   9,
+		  9,  39,  39,  39,  39,  39,  39,  23,  23,  23,  23,  23,  23,   7,   7,   7,
+		  7,   7,  37,  37,  37,  37,  37,  21,  21,  21,  21,  21,  21,   5,   5,   5,
+		  5,   5,  35,  35,  35,  35,  35,  35,  19,  19,  19,  19,  19,   3,   3,   3,
+		  3,   3,  33,  33,  33,  33,  33,  33,  17,  17,  17,  17,  17,   1,   1,   1 
+	};
+
+	static const uint8 ColorScrambleTable80[256] = {
+		0,   0,  16,  16,  16,  32,  32,  32,  48,  48,  48,  64,  64,  64,  64,   2,
+		2,   2,  18,  18,  18,  34,  34,  34,  50,  50,  50,  66,  66,  66,  66,   4,
+		4,   4,  20,  20,  20,  36,  36,  36,  52,  52,  52,  52,  68,  68,  68,   6,
+		6,   6,  22,  22,  22,  38,  38,  38,  54,  54,  54,  54,  70,  70,  70,   8,
+		8,   8,  24,  24,  24,  40,  40,  40,  40,  56,  56,  56,  72,  72,  72,  10,
+		10,  10,  26,  26,  26,  42,  42,  42,  42,  58,  58,  58,  74,  74,  74,  12,
+		12,  12,  28,  28,  28,  28,  44,  44,  44,  60,  60,  60,  76,  76,  76,  14,
+		14,  14,  30,  30,  30,  30,  46,  46,  46,  62,  62,  62,  78,  78,  78,  78,
+		79,  79,  79,  79,  63,  63,  63,  47,  47,  47,  31,  31,  31,  31,  15,  15,
+		15,  77,  77,  77,  61,  61,  61,  45,  45,  45,  29,  29,  29,  29,  13,  13,
+		13,  75,  75,  75,  59,  59,  59,  43,  43,  43,  43,  27,  27,  27,  11,  11,
+		11,  73,  73,  73,  57,  57,  57,  41,  41,  41,  41,  25,  25,  25,   9,   9,
+		9,  71,  71,  71,  55,  55,  55,  55,  39,  39,  39,  23,  23,  23,   7,   7,
+		7,  69,  69,  69,  53,  53,  53,  53,  37,  37,  37,  21,  21,  21,   5,   5,
+		5,  67,  67,  67,  67,  51,  51,  51,  35,  35,  35,  19,  19,  19,   3,   3,
+		3,  65,  65,  65,  65,  49,  49,  49,  33,  33,  33,  17,  17,  17,   1,   1
+	};
+
+	static const uint8 ColorScrambleTable192[256] = {
+		  0,  64, 128, 128,   2,  66, 130, 130,   4,  68, 132, 132,   6,  70, 134, 134,
+		  8,  72, 136, 136,  10,  74, 138, 138,  12,  76, 140, 140,  14,  78, 142, 142,
+		 16,  80, 144, 144,  18,  82, 146, 146,  20,  84, 148, 148,  22,  86, 150, 150,
+		 24,  88, 152, 152,  26,  90, 154, 154,  28,  92, 156, 156,  30,  94, 158, 158,
+		 32,  96, 160, 160,  34,  98, 162, 162,  36, 100, 164, 164,  38, 102, 166, 166,
+		 40, 104, 168, 168,  42, 106, 170, 170,  44, 108, 172, 172,  46, 110, 174, 174,
+		 48, 112, 176, 176,  50, 114, 178, 178,  52, 116, 180, 180,  54, 118, 182, 182,
+		 56, 120, 184, 184,  58, 122, 186, 186,  60, 124, 188, 188,  62, 126, 190, 190,
+		191, 191, 127,  63, 189, 189, 125,  61, 187, 187, 123,  59, 185, 185, 121,  57,
+		183, 183, 119,  55, 181, 181, 117,  53, 179, 179, 115,  51, 177, 177, 113,  49,
+		175, 175, 111,  47, 173, 173, 109,  45, 171, 171, 107,  43, 169, 169, 105,  41,
+		167, 167, 103,  39, 165, 165, 101,  37, 163, 163,  99,  35, 161, 161,  97,  33,
+		159, 159,  95,  31, 157, 157,  93,  29, 155, 155,  91,  27, 153, 153,  89,  25,
+		151, 151,  87,  23, 149, 149,  85,  21, 147, 147,  83,  19, 145, 145,  81,  17,
+		143, 143,  79,  15, 141, 141,  77,  13, 139, 139,  75,  11, 137, 137,  73,   9,
+		135, 135,  71,   7, 133, 133,  69,   5, 131, 131,  67,   3, 129, 129,  65,   1
+	};
+
+	BEGIN_GLOBAL_SHADER_PARAMETER_STRUCT(FAstcParameters, )
+		SHADER_PARAMETER_SRV(Buffer<uint>, TritsToInteger)
+		SHADER_PARAMETER_SRV(Buffer<uint>, QuintsToInteger)
+		SHADER_PARAMETER_SRV(Buffer<uint>, ScrambleTable)
+		SHADER_PARAMETER_SRV(Buffer<uint>, ColorScrambleTable48)
+		SHADER_PARAMETER_SRV(Buffer<uint>, ColorScrambleTable80)
+		SHADER_PARAMETER_SRV(Buffer<uint>, ColorScrambleTable192)
+	END_GLOBAL_SHADER_PARAMETER_STRUCT()
+
+	IMPLEMENT_GLOBAL_SHADER_PARAMETER_STRUCT(FAstcParameters, "AstcParameters");
+
+	class FAstcParametersUniformBuffer : public TUniformBuffer<FAstcParameters>
+	{
+		typedef TUniformBuffer<FAstcParameters> Super;
+	public:
+		virtual void InitRHI(FRHICommandListBase& RHICmdList) override
+		{
+			FAstcParameters Parameters;
+
+			{
+				const FRHIBufferCreateDesc CreateDesc =
+					FRHIBufferCreateDesc::CreateVertex<uint8>(TEXT("TritsToInteger"), 243)
+					.AddUsage(EBufferUsageFlags::Static | EBufferUsageFlags::ShaderResource)
+					.SetInitialState(ERHIAccess::VertexOrIndexBuffer | ERHIAccess::SRVMask);
+				Parameters.TritsToInteger = RHICmdList.CreateShaderResourceView(
+					UE::RHIResourceUtils::CreateBufferWithArray(RHICmdList, CreateDesc, TritsToInteger),
+					FRHIViewDesc::CreateBufferSRV()
+					.SetType(FRHIViewDesc::EBufferType::Typed)
+					.SetFormat(PF_R8_UINT));
+			}
+			{
+				const FRHIBufferCreateDesc CreateDesc =
+					FRHIBufferCreateDesc::CreateVertex<uint8>(TEXT("QuintsToInteger"), 125)
+					.AddUsage(EBufferUsageFlags::Static | EBufferUsageFlags::ShaderResource)
+					.SetInitialState(ERHIAccess::VertexOrIndexBuffer | ERHIAccess::SRVMask);
+				Parameters.QuintsToInteger = RHICmdList.CreateShaderResourceView(
+					UE::RHIResourceUtils::CreateBufferWithArray(RHICmdList, CreateDesc, QuintsToInteger),
+					FRHIViewDesc::CreateBufferSRV()
+					.SetType(FRHIViewDesc::EBufferType::Typed)
+					.SetFormat(PF_R8_UINT));
+			}
+			{
+				const FRHIBufferCreateDesc CreateDesc =
+					FRHIBufferCreateDesc::CreateVertex<uint8>(TEXT("ScrambleTable"), 12 * WEIGHT_QUANTIZE_NUM)
+					.AddUsage(EBufferUsageFlags::Static | EBufferUsageFlags::ShaderResource)
+					.SetInitialState(ERHIAccess::VertexOrIndexBuffer | ERHIAccess::SRVMask);
+				Parameters.ScrambleTable = RHICmdList.CreateShaderResourceView(
+					UE::RHIResourceUtils::CreateBufferWithArray(RHICmdList, CreateDesc, ScrambleTable),
+					FRHIViewDesc::CreateBufferSRV()
+					.SetType(FRHIViewDesc::EBufferType::Typed)
+					.SetFormat(PF_R8_UINT));
+			}
+			{
+				const FRHIBufferCreateDesc CreateDesc =
+					FRHIBufferCreateDesc::CreateVertex<uint8>(TEXT("ColorScrambleTable48"), 256)
+					.AddUsage(EBufferUsageFlags::Static | EBufferUsageFlags::ShaderResource)
+					.SetInitialState(ERHIAccess::VertexOrIndexBuffer | ERHIAccess::SRVMask);
+				Parameters.ColorScrambleTable48 = RHICmdList.CreateShaderResourceView(
+					UE::RHIResourceUtils::CreateBufferWithArray(RHICmdList, CreateDesc, ColorScrambleTable48),
+					FRHIViewDesc::CreateBufferSRV()
+					.SetType(FRHIViewDesc::EBufferType::Typed)
+					.SetFormat(PF_R8_UINT));
+			}
+			{
+				const FRHIBufferCreateDesc CreateDesc =
+					FRHIBufferCreateDesc::CreateVertex<uint8>(TEXT("ColorScrambleTable80"), 256)
+					.AddUsage(EBufferUsageFlags::Static | EBufferUsageFlags::ShaderResource)
+					.SetInitialState(ERHIAccess::VertexOrIndexBuffer | ERHIAccess::SRVMask);
+				Parameters.ColorScrambleTable80 = RHICmdList.CreateShaderResourceView(
+					UE::RHIResourceUtils::CreateBufferWithArray(RHICmdList, CreateDesc, ColorScrambleTable80),
+					FRHIViewDesc::CreateBufferSRV()
+					.SetType(FRHIViewDesc::EBufferType::Typed)
+					.SetFormat(PF_R8_UINT));
+			}
+			{
+				const FRHIBufferCreateDesc CreateDesc =
+					FRHIBufferCreateDesc::CreateVertex<uint8>(TEXT("ColorScrambleTable192"), 256)
+					.AddUsage(EBufferUsageFlags::Static | EBufferUsageFlags::ShaderResource)
+					.SetInitialState(ERHIAccess::VertexOrIndexBuffer | ERHIAccess::SRVMask);
+				Parameters.ColorScrambleTable192 = RHICmdList.CreateShaderResourceView(
+					UE::RHIResourceUtils::CreateBufferWithArray(RHICmdList, CreateDesc, ColorScrambleTable192),
+					FRHIViewDesc::CreateBufferSRV()
+					.SetType(FRHIViewDesc::EBufferType::Typed)
+					.SetFormat(PF_R8_UINT));
+			}
+
+			SetContentsNoUpdate(Parameters);
+			Super::InitRHI(RHICmdList);
+		}
+	};
+
+	const TUniformBufferRef<FAstcParameters>& GetAstcParametersUniformBufferRef()
+	{
+		static TGlobalResource<FAstcParametersUniformBuffer> AstcParametersUniformBuffer;
+		return AstcParametersUniformBuffer.GetUniformBufferRef();
+	}
+
 	bool UseEtcProfile(EShaderPlatform ShaderPlatform)
 	{
 		switch (ShaderPlatform)
@@ -127,6 +503,34 @@ namespace RuntimeVirtualTexture
 		return false;
 	}
 
+	bool UseAstcProfile(EShaderPlatform ShaderPlatform)
+	{
+		if (!CVarRVTAstc.GetValueOnAnyThread()) {
+			return false;
+		}
+		switch (ShaderPlatform)
+		{
+		case SP_METAL:
+		case SP_METAL_MRT:
+		case SP_METAL_SIM:
+		case SP_METAL_TVOS:
+		case SP_METAL_MRT_TVOS:
+		case SP_VULKAN_ES3_1_ANDROID:
+		case SP_OPENGL_ES3_1_ANDROID:
+		case SP_VULKAN_SM5_ANDROID:
+			return true;
+		default:
+			break;
+		}
+		return false;
+	}
+
+	bool UseAstcHighProfile(EShaderPlatform ShaderPlatform)
+	{
+		return UseAstcProfile(ShaderPlatform) && CVarRVTAstcHigh.GetValueOnAnyThread();
+
+	}
+
 	/** For platforms that do not support 2-channel images, write 64bit compressed texture outputs into RGBA16 instead of RG32. */
 	bool UseRGBA16(EShaderPlatform ShaderPlatform)
 	{
@@ -1199,6 +1603,7 @@ namespace RuntimeVirtualTexture
 			SHADER_PARAMETER(FUintVector2, DestPos1)
 			SHADER_PARAMETER(FUintVector2, DestPos2)
 			SHADER_PARAMETER_STRUCT_REF(FEtcParameters, EtcParameters)
+			SHADER_PARAMETER_STRUCT_REF(FAstcParameters, AstcParameters)
 			SHADER_PARAMETER_RDG_TEXTURE(Texture2D<float4>, RenderTexture0)
 			SHADER_PARAMETER_SAMPLER(SamplerState, TextureSampler0)
 			SHADER_PARAMETER_RDG_TEXTURE(Texture2D<float4>, RenderTexture1)
@@ -1213,13 +1618,27 @@ namespace RuntimeVirtualTexture
 			SHADER_PARAMETER_RDG_TEXTURE_UAV(RWTexture2D<uint4>, OutCompressTexture2_128bit)
 		END_SHADER_PARAMETER_STRUCT()
 
+		class FAstcHighProfile : SHADER_PERMUTATION_BOOL("ASTC_HIGH_PROFILE");
+		using FPermutationDomain = TShaderPermutationDomain<FAstcHighProfile>;
+
 		static void ModifyCompilationEnvironment(const FGlobalShaderPermutationParameters& Parameters, FShaderCompilerEnvironment& OutEnvironment)
 		{
 			FGlobalShader::ModifyCompilationEnvironment(Parameters, OutEnvironment);
 			OutEnvironment.SetDefine(TEXT("ETC_PROFILE"), UseEtcProfile(Parameters.Platform) ? 1 : 0);
+			OutEnvironment.SetDefine(TEXT("ASTC_PROFILE"), UseAstcProfile(Parameters.Platform) ? 1 : 0);
 			OutEnvironment.SetDefine(TEXT("PACK_RG32_RGBA16"), UseRGBA16(Parameters.Platform) ? 1 : 0);
 		}
 
+		static bool ShouldCompilePermutation(const FGlobalShaderPermutationParameters& Parameters)
+		{
+			FPermutationDomain PermutationVector(Parameters.PermutationId);
+			if (PermutationVector.Get<FAstcHighProfile>())
+			{
+				return UseAstcHighProfile(Parameters.Platform);
+			}
+			return true;
+		}
+
 		FShader_VirtualTextureCompress()
 		{}
 
@@ -1264,7 +1683,10 @@ namespace RuntimeVirtualTexture
 	void AddCompressPass(FRDGBuilder& GraphBuilder, ERHIFeatureLevel::Type FeatureLevel, FShader_VirtualTextureCompress::FParameters* Parameters, FIntVector GroupCount)
 	{
 		FGlobalShaderMap* GlobalShaderMap = GetGlobalShaderMap(FeatureLevel);
-		TShaderMapRef< FShader_VirtualTextureCompress_CS< MaterialType > > ComputeShader(GlobalShaderMap);
+
+		FShader_VirtualTextureCompress::FPermutationDomain PermutationVector;
+		PermutationVector.Set<FShader_VirtualTextureCompress::FAstcHighProfile>(UseAstcHighProfile(GShaderPlatformForFeatureLevel[FeatureLevel]));
+		TShaderMapRef< FShader_VirtualTextureCompress_CS< MaterialType > > ComputeShader(GlobalShaderMap, PermutationVector);
 
 		FComputeShaderUtils::AddPass(
 			GraphBuilder,
@@ -1638,6 +2060,7 @@ namespace RuntimeVirtualTexture
 			PassParameters->DestPos1 = GraphSetup.bDirectAliasing ? FUintVector2(DestBox1.Min.X / 4, DestBox1.Min.Y / 4) : FUintVector2(0, 0);
 			PassParameters->DestPos2 = GraphSetup.bDirectAliasing ? FUintVector2(DestBox2.Min.X / 4, DestBox2.Min.Y / 4) : FUintVector2(0, 0);
 			PassParameters->EtcParameters = GetEtcParametersUniformBufferRef();
+			PassParameters->AstcParameters = GetAstcParametersUniformBufferRef();
 			PassParameters->RenderTexture0 = GraphSetup.RenderTexture0;
 			PassParameters->TextureSampler0 = TStaticSamplerState<SF_Point, AM_Clamp, AM_Clamp, AM_Clamp>::GetRHI();
 			PassParameters->RenderTexture1 = GraphSetup.RenderTexture1;
