Added Epic Games copyright notice to shader header file

diff --git a/Engine/Shaders/Private/IntegerSequenceEncoding.ush b/Engine/Shaders/Private/IntegerSequenceEncoding.ush
--- a/Engine/Shaders/Private/IntegerSequenceEncoding.ush
+++ b/Engine/Shaders/Private/IntegerSequenceEncoding.ush
@@ -1,4 +1,4 @@
-
+// Copyright Epic Games, Inc. All Rights Reserved.
 
 // astc specification Integer Sequence Encoding
 // https://registry.khronos.org/DataFormat/specs/1.3/dataformat.1.3.html#astc-integer-sequence-encoding
