[Screenspace Subsurface Scattering] Use subsurface profile cache for Separable as well. It improves the performance when SSS covers a lot of screen.

#tests perfB test does not regress the rendering time.

diff --git a/Engine/Shaders/Private/PostProcessSubsurface.usf b/Engine/Shaders/Private/PostProcessSubsurface.usf
--- a/Engine/Shaders/Private/PostProcessSubsurface.usf
+++ b/Engine/Shaders/Private/PostProcessSubsurface.usf
@@ -307,6 +307,25 @@ FBurleyParameter GetBurleyParameters(uint SubsurfaceProfileInt, FGBufferData GBu
 }
 #endif
 
+// BufferUV back to grid posistion of the input
+uint2 ConvertUVPos2Input0GridPos(float2 BufferUV)
+{
+	return clamp(BufferUV, 0, 1) * SubsurfaceInput0_Extent - 0.5f;
+}
+
+uint ExtractSubsurfaceProfileIntWithInvalid(float2 UV,float Depth)
+{
+	uint ProfileID = SSS_PROFILE_ID_INVALID;
+
+	// the invalid status be determined by the scene depth.
+	if (Depth > 0)
+	{
+		uint2 PosOfSampledUV = ConvertUVPos2Input0GridPos(UV);
+		ProfileID = ProfileIdTexture[PosOfSampledUV];
+	}
+	return ProfileID;
+}
+
 // from [REDACTED_PROFILE_URL]/separable-sss/tree/master/Demo
 // [REDACTED_PERSON] [REDACTED_PROFILE_URL]
 // [REDACTED_PROFILE_URL]
@@ -838,7 +857,7 @@ void SetupIndirectCS(uint3 DT_ID : SV_DispatchThreadID,
 	}
 #endif // SUBSURFACE_HALF_RES
 
-#if !SUBSURFACE_FORCE_SEPARABLE && ENABLE_PROFILE_ID_CACHE
+#if ENABLE_PROFILE_ID_CACHE
 	uint ProfileId = GetSubsurfaceProfileId(BufferUV);
 #endif
 
@@ -888,7 +907,7 @@ void SetupIndirectCS(uint3 DT_ID : SV_DispatchThreadID,
 		Type |= TYPE_IN_VIEWPORT;
 		SetupTexture[Pos] = OutColor;
 
-#if !SUBSURFACE_FORCE_SEPARABLE && ENABLE_PROFILE_ID_CACHE
+#if ENABLE_PROFILE_ID_CACHE
 		ProfileIdTexture[Pos] = ProfileId;
 #endif
 	}
diff --git a/Engine/Shaders/Private/SeparableSSS.ush b/Engine/Shaders/Private/SeparableSSS.ush
--- a/Engine/Shaders/Private/SeparableSSS.ush
+++ b/Engine/Shaders/Private/SeparableSSS.ush
@@ -274,10 +274,16 @@ float4 SSSSBlurPS(uint2 BufferPos, float2 BufferUV, float2 dir, bool initStencil
 	SeparableFilterParameters.SubsurfaceProfileInt = SubsurfaceProfileInt;
 	FillBurleyParameters(SSSData, SeparableFilterParameters);
 #else
+
+#if ENABLE_PROFILE_ID_CACHE
+	const uint SubsurfaceProfileInt = ExtractSubsurfaceProfileIntWithInvalid(BufferUV, colorM.a);
+#else
 	const FGBufferData GBufferData = GetGBufferData(BufferUV);
 	// 0..255, which SubSurface profile to pick
 	// ideally this comes from a half res buffer as well - there are some minor artifacts
 	const uint SubsurfaceProfileInt = ExtractSubsurfaceProfileInt(GBufferData);
+#endif
+
 	SeparableFilterParameters.SubsurfaceProfileInt = SubsurfaceProfileInt;
 #endif
 
@@ -312,7 +318,11 @@ float4 SSSSBlurPS(uint2 BufferPos, float2 BufferUV, float2 dir, bool initStencil
 			// Fetch color and depth for current sample:
 			float2 LocalUV = BufferUV + UVOffset * Side;
 			float4 color = SSSSSampleSceneColor(LocalUV);
+#if ENABLE_PROFILE_ID_CACHE
+			uint LocalSubsurfaceProfileInt = ExtractSubsurfaceProfileIntWithInvalid(LocalUV,color.a);
+#else
 			uint LocalSubsurfaceProfileInt = SSSSSampleProfileId(LocalUV);
+#endif
 			float3 ColorTint = LocalSubsurfaceProfileInt == SubsurfaceProfileInt ? 1.0f : BoundaryColorBleed;
 
 
diff --git a/Engine/Shaders/Private/SubsurfaceBurleyNormalized.ush b/Engine/Shaders/Private/SubsurfaceBurleyNormalized.ush
--- a/Engine/Shaders/Private/SubsurfaceBurleyNormalized.ush
+++ b/Engine/Shaders/Private/SubsurfaceBurleyNormalized.ush
@@ -489,25 +489,6 @@ uint ExtractSubsurfaceProfileIntWithInvalid(FGBufferData BufferData)
 }
 #endif
 
-// BufferUV back to grid posistion of the input
-uint2 ConvertUVPos2Input0GridPos(float2 BufferUV)
-{
-	return clamp(BufferUV, 0, 1) * SubsurfaceInput0_Extent - 0.5f;
-}
-
-uint ExtractSubsurfaceProfileIntWithInvalid(float2 UV,float Depth)
-{
-	uint ProfileID = SSS_PROFILE_ID_INVALID;
-
-	// the invalid status be determined by the scene depth.
-	if (Depth > 0)
-	{
-		uint2 PosOfSampledUV = ConvertUVPos2Input0GridPos(UV);
-		ProfileID = ProfileIdTexture[PosOfSampledUV];
-	}
-	return ProfileID;
-}
-
 #define BILATERAL_FILTER_DEFAULT_NORMAL float3(1.0f,0.0f,0.0f)
 
 struct BurleySampleDiffuseNormal
@@ -733,7 +714,7 @@ BurleySampleDiffuseNormal SampleSSSColorConsideringLocalShared(float2 CenterUV,
 	// Fix border flickering when mipmaps got garbage data.
 	float2 ClampedUV = ClampUVLevel(CenterUV + UVOffset, SubsurfaceInput0_UVViewportBilinearMin, SubsurfaceInput0_UVViewportBilinearMax, MipLevel);
 
-#if !(ENABLE_PROFILE_ID_CACHE && (SUBSURFACE_FORCE_SEPARABLE !=1))
+#if !(ENABLE_PROFILE_ID_CACHE)
 	#if SUBTRATE_GBUFFER_FORMAT==1
 	const FSubstrateSubsurfaceHeader SSSHeader = LoadSubstrateSSSHeader(ClampedUV);
 	#else
@@ -760,7 +741,7 @@ BurleySampleDiffuseNormal SampleSSSColorConsideringLocalShared(float2 CenterUV,
 	Sample.WorldNormal.xyz = GetBilateralNormal(ClampedUV, SubsurfaceInput0_Extent);
 #endif
 
-#if ENABLE_PROFILE_ID_CACHE && (SUBSURFACE_FORCE_SEPARABLE !=1)
+#if ENABLE_PROFILE_ID_CACHE
 	Sample.ProfileID = ExtractSubsurfaceProfileIntWithInvalid(ClampedUV,Sample.DiffuseLighting.w);
 #else
 	Sample.ProfileID = ExtractSubsurfaceProfileIntWithInvalid(SSSHeader);
@@ -784,7 +765,7 @@ BurleySampleDiffuseNormal SampleSSSColorConsideringLocalShared(float2 CenterUV,
 			Sample.WorldNormal.xyz = GetBilateralNormal(ClampedUV, SubsurfaceInput0_Extent);
 #endif
 
-#if ENABLE_PROFILE_ID_CACHE && (SUBSURFACE_FORCE_SEPARABLE != 1)
+#if ENABLE_PROFILE_ID_CACHE
 			Sample.ProfileID = ExtractSubsurfaceProfileIntWithInvalid(ClampedUV, Sample.DiffuseLighting.w);
 #else
 			Sample.ProfileID = ExtractSubsurfaceProfileIntWithInvalid(SSSHeader);
@@ -801,7 +782,7 @@ BurleySampleDiffuseNormal SampleSSSColorConsideringLocalShared(float2 CenterUV,
 	Sample.WorldNormal.xyz = GetBilateralNormal(ClampedUV, SubsurfaceInput0_Extent);
 #endif
 	
-#if ENABLE_PROFILE_ID_CACHE && (SUBSURFACE_FORCE_SEPARABLE != 1)
+#if ENABLE_PROFILE_ID_CACHE
 	Sample.ProfileID = ExtractSubsurfaceProfileIntWithInvalid(ClampedUV, Sample.DiffuseLighting.w);
 #else
 	Sample.ProfileID = ExtractSubsurfaceProfileIntWithInvalid(SSSHeader);
diff --git a/Engine/Source/Runtime/Renderer/Private/PostProcess/PostProcessSubsurface.cpp b/Engine/Source/Runtime/Renderer/Private/PostProcess/PostProcessSubsurface.cpp
--- a/Engine/Source/Runtime/Renderer/Private/PostProcess/PostProcessSubsurface.cpp
+++ b/Engine/Source/Runtime/Renderer/Private/PostProcess/PostProcessSubsurface.cpp
@@ -1007,7 +1007,7 @@ void AddSubsurfaceViewPass(
 	//We run in separable mode under two conditions: 1) Run Burley fallback mode. 2) when the screen is in half resolution.
 	const bool bForceRunningInSeparable = CVarSSSBurleyQuality.GetValueOnRenderThread() == 0 || bHalfRes;
 
-	const bool bUseProfileIdCache = !bForceRunningInSeparable && IsProfileIdCacheEnabled();
+	const bool bUseProfileIdCache = IsProfileIdCacheEnabled();
 
 	const int32 MinGenerateMipsTileCount = FMath::Max(0, CVarSSSMipmapsMinTileCount.GetValueOnRenderThread());
 
