Mobile: Fixed case where unlit materials could write undefined values into GBuffer targets

diff --git a/Engine/Shaders/Private/MobileBasePassPixelShader.usf b/Engine/Shaders/Private/MobileBasePassPixelShader.usf
--- a/Engine/Shaders/Private/MobileBasePassPixelShader.usf
+++ b/Engine/Shaders/Private/MobileBasePassPixelShader.usf
@@ -84,7 +84,8 @@
 #define NONMETAL (MATERIAL_NONMETAL || MOBILE_QL_FORCE_NONMETAL)
 #define FORCE_VERTEX_NORMAL (MOBILE_QL_DISABLE_MATERIAL_NORMAL)
 #define SUPPORT_SPOTLIGHTS_SHADOW (MATERIALBLENDING_SOLID || MATERIALBLENDING_MASKED) && PROJECT_MOBILE_ENABLE_MOVABLE_SPOTLIGHT_SHADOWS
-#define MOBILE_USE_GBUFFER (MOBILE_DEFERRED_SHADING && ((MATERIALBLENDING_SOLID || MATERIALBLENDING_MASKED) && !MATERIAL_SHADINGMODEL_SINGLELAYERWATER))
+// Sky does not write to depth and must not overwrite GBuffer
+#define MOBILE_USE_GBUFFER (MOBILE_DEFERRED_SHADING && ((MATERIALBLENDING_SOLID || MATERIALBLENDING_MASKED) && !MATERIAL_SHADINGMODEL_SINGLELAYERWATER && !MATERIAL_IS_SKY))
 
 #if MATERIAL_SHADINGMODEL_THIN_TRANSLUCENT || SUBSTRATE_BLENDING_TRANSLUCENT_COLOREDTRANSMITTANCE
 	#define MOBILE_TRANSLUCENT_COLOR_TRANSMITTANCE 1
@@ -100,7 +101,8 @@
 #endif
 
 // SceneDepthAux always enabled for IOS, on other platforms only enabled when in forward and MobileHDR=true
-#define HAS_SCENE_DEPTH_AUX_OUTPUT (USE_SCENE_DEPTH_AUX && (MATERIALBLENDING_SOLID || MATERIALBLENDING_MASKED) && !MATERIAL_SHADINGMODEL_SINGLELAYERWATER)
+// Sky does not write to depth
+#define HAS_SCENE_DEPTH_AUX_OUTPUT (USE_SCENE_DEPTH_AUX && (MATERIALBLENDING_SOLID || MATERIALBLENDING_MASKED) && !MATERIAL_SHADINGMODEL_SINGLELAYERWATER && !MATERIAL_IS_SKY)
 
 #if MOBILE_USE_GBUFFER && MOBILE_EXTENDED_GBUFFER
 	#define SV_TargetDepthAux SV_Target5	
@@ -1062,7 +1064,7 @@ void Main(
 
 	GBuffer.PrecomputedShadowFactors = GetPrimaryPrecomputedShadowMask(LightmapVTPageTableResult, Interpolants, MaterialParameters);
 
-#if MOBILE_USE_GBUFFER && !MATERIAL_SHADINGMODEL_UNLIT
+#if MOBILE_USE_GBUFFER
 	GBuffer.IndirectIrradiance = IndirectIrradiance;
 	MobileEncodeGBuffer(GBuffer, OutGBufferA, OutGBufferB, OutGBufferC, OutGBufferD);
 #else
