Path Tracer: Fix compilation errors when RemoveDeadCode is turned off

This was preventing some EngineTests from running correctly.

#rnx

diff --git a/Engine/Shaders/Private/PathTracing/PathTracingVolumetricCloudMaterialShader.usf b/Engine/Shaders/Private/PathTracing/PathTracingVolumetricCloudMaterialShader.usf
--- a/Engine/Shaders/Private/PathTracing/PathTracingVolumetricCloudMaterialShader.usf
+++ b/Engine/Shaders/Private/PathTracing/PathTracingVolumetricCloudMaterialShader.usf
@@ -29,7 +29,17 @@
 #define CloudClipX						PathTracingCloudParameters.CloudClipX						
 #define CloudClipY						PathTracingCloudParameters.CloudClipY
 #define CloudClipZ						PathTracingCloudParameters.CloudClipZ
+
+#define CloudLayerBotKm					PathTracingCloudParameters.CloudLayerBotKm
+#define CloudLayerTopKm					PathTracingCloudParameters.CloudLayerTopKm
 #define CloudClipDistKm                 PathTracingCloudParameters.CloudClipDistKm
+#define CloudClipRadiusKm               PathTracingCloudParameters.CloudClipRadiusKm
+
+#define CloudClipCenterKm				PathTracingCloudParameters.CloudClipCenterKm
+
+#define CloudTracingMaxDistance			PathTracingCloudParameters.CloudTracingMaxDistance
+#define CloudRoughnessCutoff			PathTracingCloudParameters.CloudRoughnessCutoff
+
 #define CloudAccelMapResolution			PathTracingCloudParameters.CloudAccelMapResolution
 #define CloudVoxelWidth					PathTracingCloudParameters.CloudVoxelWidth
 #define CloudInvVoxelWidth				PathTracingCloudParameters.CloudInvVoxelWidth
@@ -49,8 +59,8 @@ FVolumeShadedResult CloudsGetDensity(float3 TranslatedWorldPos)
 	// Measure HeightKm as the distance to the planet center
 	const FDFVector3 OC = DFMultiply(DFSubtract(TranslatedWorldPos, MakeDFVector3(PlanetCenterTranslatedWorldHi, PlanetCenterTranslatedWorldLo)), CM_TO_SKY_UNIT);
 	const float H = DFLengthDemote(OC); // HeightKm
-	float B = PathTracingCloudParameters.CloudLayerBotKm;
-	float T = PathTracingCloudParameters.CloudLayerTopKm;
+	float B = CloudLayerBotKm;
+	float T = CloudLayerTopKm;
 
 	FSampleCloudMaterialResult CloudResult = SampleCloudMaterial(AbsoluteWorldPosition, H, B, T);
 	FVolumeShadedResult Result = (FVolumeShadedResult)0;
diff --git a/Engine/Shaders/Private/PathTracing/Volume/PathTracingAtmosphere.ush b/Engine/Shaders/Private/PathTracing/Volume/PathTracingAtmosphere.ush
--- a/Engine/Shaders/Private/PathTracing/Volume/PathTracingAtmosphere.ush
+++ b/Engine/Shaders/Private/PathTracing/Volume/PathTracingAtmosphere.ush
@@ -113,6 +113,8 @@ float3 RayPlanetIntersectPrecise(float3 Ro, float3 Rd)
 
 #endif
 
+#ifndef PlanetCenterTranslatedWorldHi // TODO: find a better way to recognize we are coming from the cloud shader
+
 FVolumeIntersection AtmosphereIntersect(float3 Origin, float3 Direction, float TMin, float TMax)
 {
 	float3 Result = RayPlanetIntersectPrecise(Origin, Direction);
@@ -145,6 +147,8 @@ FPackedPathTracingPayload AtmosphereGetBlockerHit(float3 Origin, float3 Directio
 	return PackPathTracingPayload(Payload);
 }
 
+#endif
+
 struct FSphericalPlanetParameters {
 	float HeightKm;         // Distance from planet center (in Km)
 	float ZenithCosAngle;   // angle between the ray and the planet normal (+1 goes up towards space, -1 points down toward the ground)
@@ -262,6 +266,8 @@ FVolumeShadedResult AtmosphereGetDensity(float3 TranslatedWorldPos)
 	return Result;
 }
 
+#ifndef PlanetCenterTranslatedWorldHi // TODO: find a better way to recognize we are coming from the cloud shader
+
 float3 AtmosphereGetOpticalDensity(FSphericalPlanetParameters Params)
 {
 	float2 UV = ToAtmosphereOpticalDepthLUTUV(Params.HeightKm, Params.ZenithCosAngle, Atmosphere.BottomRadiusKm, Atmosphere.TopRadiusKm);
@@ -279,3 +285,5 @@ float3 AtmosphereGetTransmittance(float3 Origin, float3 Direction, float TMin, f
 	// compute transmittance
 	return exp(-Tau);
 }
+
+#endif
\ No newline at end of file
