Fix for fog to match between mobile and deferred renderer, main vs real time sky capture.
Clarified parameter for heigh fog evaluation.

#rnx
#tests PC editor dx12/dx12 / fog&cloud enginetests

diff --git a/Engine/Shaders/Private/BasePassVertexShader.usf b/Engine/Shaders/Private/BasePassVertexShader.usf
--- a/Engine/Shaders/Private/BasePassVertexShader.usf
+++ b/Engine/Shaders/Private/BasePassVertexShader.usf
@@ -83,7 +83,7 @@ void Main(
 
 // Calculate the fog needed for translucency
 #if NEEDS_BASEPASS_VERTEX_FOGGING
-	Output.BasePassInterpolants.VertexFog = CalculateHeightFog(WorldPosition.xyz - ResolvedView.TranslatedWorldCameraOrigin, EyeIndex, ResolvedView);
+	Output.BasePassInterpolants.VertexFog = CalculateHeightFog(WorldPosition.xyz - ResolvedView.TranslatedWorldCameraOrigin, EyeIndex, ResolvedView); // WorldPosition is in fact TranslatedWorldPosition 
 
 	const float OneOverPreExposure = ResolvedView.OneOverPreExposure;
 
diff --git a/Engine/Shaders/Private/HeightFogCommon.ush b/Engine/Shaders/Private/HeightFogCommon.ush
--- a/Engine/Shaders/Private/HeightFogCommon.ush
+++ b/Engine/Shaders/Private/HeightFogCommon.ush
@@ -209,7 +209,7 @@ float PreComputeFogOriginFactor(float OriginHeight, float FogHeight, float FogFa
 	return FogDensity * pow(2.0f, CollapsedFogParameterPower);
 }
 
-// @param WorldPositionRelativeToCamera = WorldPosition - InCameraPosition
+// @param WorldPositionRelativeToCamera = WorldPosition - InCameraPosition. It is the vector from the camera to the world position to shade (independent from TranslatedPosition).
 half4 GetExponentialHeightFog(float3 WorldPositionRelativeToCamera, float ExcludeDistance, uint EyeIndex, ViewState InView, bool bOverrideOrigin = false, float3 OverrideRayStartRelativeToCamera=0.0f.xxx, float3 OverrideRayDir = 0.0f.xxx, float OverrideRayLength = 0.0f)
 {
 	FogStructData FogData = GetFogStructData(EyeIndex);
@@ -385,11 +385,13 @@ half4 GetExponentialHeightFog(float3 WorldPositionRelativeToCamera, float Exclud
 	return half4(FogColor, ExpFogFactor);
 }
 
+// @param WorldPositionRelativeToCamera = WorldPosition - InCameraPosition. It is the vector from the camera to the world position to shade (independent from TranslatedPosition).
 half4 GetExponentialHeightFog(float3 WorldPositionRelativeToCamera, float ExcludeDistance)
 {
 	return GetExponentialHeightFog(WorldPositionRelativeToCamera, ExcludeDistance, 0, GetPrimaryView());
 }
 
+// @param WorldPositionRelativeToCamera = WorldPosition - InCameraPosition. It is the vector from the camera to the world position to shade (independent from TranslatedPosition).
 half4 CalculateHeightFog(float3 WorldPositionRelativeToCamera, uint EyeIndex, ViewState InView)
 {
 	if (InView.RenderingReflectionCaptureMask == 0.0f && !IsExponentialFogRenderedInMain(InView.EnvironmentComponentsFlags))
@@ -409,6 +411,7 @@ half4 CalculateHeightFog(float3 WorldPositionRelativeToCamera, uint EyeIndex, Vi
 	return FogInscatteringAndOpacity;
 }
 
+// @param WorldPositionRelativeToCamera = WorldPosition - InCameraPosition. It is the vector from the camera to the world position to shade (independent from TranslatedPosition).
 half4 CalculateHeightFog(float3 WorldPositionRelativeToCamera)
 {
 	return CalculateHeightFog(WorldPositionRelativeToCamera, 0, GetPrimaryView());
diff --git a/Engine/Shaders/Private/HeightFogPixelShader.usf b/Engine/Shaders/Private/HeightFogPixelShader.usf
--- a/Engine/Shaders/Private/HeightFogPixelShader.usf
+++ b/Engine/Shaders/Private/HeightFogPixelShader.usf
@@ -81,13 +81,7 @@ void ExponentialPixelMain(
 
 	bool bIsRendered = (DeviceZ != 0.0);
 
-#if 1
-	// Using SvPositionToTranslatedWorld is more stable and accurate as compared to the ScreenVector over large distance (SkyAtmosphere like).
-	// However, we do not to clamp DeviceZ to a reasinable FarZ in order to avoid Inf values.
-	float3 WorldPositionRelativeToCamera = SvPositionToTranslatedWorld(float4(SVPos.xy, max(1.5e-10, DeviceZ), 1.0));
-#else
 	float3 WorldPositionRelativeToCamera = ScreenVector.xyz * SceneDepth;
-#endif
 
 	float ZSlice = log2(SceneDepth * View.VolumetricFogGridZParams.x + View.VolumetricFogGridZParams.y) * View.VolumetricFogGridZParams.z * View.VolumetricFogInvGridSize.z;
 	
diff --git a/Engine/Shaders/Private/MobileBasePassVertexShader.usf b/Engine/Shaders/Private/MobileBasePassVertexShader.usf
--- a/Engine/Shaders/Private/MobileBasePassVertexShader.usf
+++ b/Engine/Shaders/Private/MobileBasePassVertexShader.usf
@@ -79,7 +79,7 @@ void Main(
 #endif
 
 #if USE_VERTEX_FOG
-	half4 VertexFog = CalculateHeightFog(WorldPosition.xyz - ResolvedView.TranslatedWorldCameraOrigin, EyeIndex, ResolvedView);
+	half4 VertexFog = CalculateHeightFog(WorldPosition.xyz - ResolvedView.TranslatedWorldCameraOrigin, EyeIndex, ResolvedView); // WorldPosition is in fact TranslatedWorldPosition 
 
 #if PROJECT_SUPPORT_SKY_ATMOSPHERE && MATERIAL_IS_SKY==0 // Do not apply aerial perpsective on sky materials
 	if (ResolvedView.SkyAtmosphereApplyCameraAerialPerspectiveVolume > 0.0f)
diff --git a/Engine/Shaders/Private/MobileFog.usf b/Engine/Shaders/Private/MobileFog.usf
--- a/Engine/Shaders/Private/MobileFog.usf
+++ b/Engine/Shaders/Private/MobileFog.usf
@@ -55,7 +55,7 @@ void MobileFogPS(
 	half4 FinalFog = half4(0,0,0,1);
 	
 	float DeviceZ = LookupDeviceZ(TexCoord, GetEyeIndex(StereoInput));
-	float SceneDepth = ConvertFromDeviceZ(DeviceZ);
+	float SceneDepth = ConvertFromDeviceZ(max(1.0e-10, DeviceZ));
 	float3 WorldPositionRelativeToCamera = ScreenVector.xyz * SceneDepth;
 	float3 TranslatedWorldPosition = WorldPositionRelativeToCamera + GetTranslatedWorldCameraPosFromView(SVPos.xy);
 
diff --git a/Engine/Shaders/Private/VolumetricRenderTarget.usf b/Engine/Shaders/Private/VolumetricRenderTarget.usf
--- a/Engine/Shaders/Private/VolumetricRenderTarget.usf
+++ b/Engine/Shaders/Private/VolumetricRenderTarget.usf
@@ -628,7 +628,7 @@ float4 ApplyFogToCloudSampleFromVPosAndDepth(float2 SVPos, float DepthKilometer,
 	const float2 SVPosition = SVPos;
 	const float DistanceToPoint = DepthKilometer * KILOMETER_TO_CENTIMETER;
 
-	const float3 FogSampleWorldPositionRelativeToCameraCm = normalize(SvPositionToTranslatedWorld(float4(SVPosition, 0.5f, 1.0f))) * DistanceToPoint;
+	const float3 FogSampleWorldPositionRelativeToCameraCm = normalize(GetScreenWorldDir(float4(SVPosition, 0.5f, 1.0f))) * DistanceToPoint;
 
 #if PERMUTATION_APPLY_FOG
 	{
