VolumetricCloud - added option to evaluate fog on all pixel. A timny cost for seemless fog on cloud at the horizon.

#rnx
#tests PC editor dx12

diff --git a/Engine/Shaders/Private/VolumetricCloud.usf b/Engine/Shaders/Private/VolumetricCloud.usf
--- a/Engine/Shaders/Private/VolumetricCloud.usf
+++ b/Engine/Shaders/Private/VolumetricCloud.usf
@@ -1502,15 +1502,15 @@ void MainCommon(in FMaterialPixelParameters MaterialParameters, in float4 SvPosi
 	float4 ClipPos = mul(float4(RayTranslatedWorldOriginKm * KILOMETER_TO_CENTIMETER + FogSampleWorldPositionRelativeToCameraCm, 1.0f), PrimaryView.TranslatedWorldToClip);
 
 #if SUPPORT_PRIMITIVE_ALPHA_HOLDOUT
-	const bool bApplyFogOnCloud			= tAPWeightsSum;			// Use MeanTransmittance instead of tAPWeightsSum as it is more representative of what can happen since it has not been transformed. (but can result in more fog being applied, this makes alpha holdout correct)
+	const bool bApplyFogOnCloud			= RenderVolumetricCloudParameters.EnableHeightFog >= 2 || tAPWeightsSum;					// Use MeanTransmittance instead of tAPWeightsSum as it is more representative of what can happen since it has not been transformed. (but can result in more fog being applied, this makes alpha holdout correct)
 #if CLOUD_MIN_AND_MAX_DEPTH
-	const bool bApplyFogOnClosestCloud	= ClosestMeanTransmittance < 1.0f;	// Idem
+	const bool bApplyFogOnClosestCloud	= RenderVolumetricCloudParameters.EnableHeightFog >= 2 || ClosestMeanTransmittance < 1.0f;// Idem
 #endif
 #else // SUPPORT_PRIMITIVE_ALPHA_HOLDOUT
 	// We skip more pixels when using tAPWeightsSum so let;s keep that when alpha output is not enabled for the sake of performance
-	const bool bApplyFogOnCloud			= tAPWeightsSum > 0.0f;
+	const bool bApplyFogOnCloud			= RenderVolumetricCloudParameters.EnableHeightFog >= 2 || tAPWeightsSum > 0.0f;
 #if CLOUD_MIN_AND_MAX_DEPTH
-	const bool bApplyFogOnClosestCloud	= ClostestTAPWeightsSum > 0.0f;
+	const bool bApplyFogOnClosestCloud	= RenderVolumetricCloudParameters.EnableHeightFog >= 2 || ClostestTAPWeightsSum > 0.0f;
 #endif
 #endif // SUPPORT_PRIMITIVE_ALPHA_HOLDOUT
 
@@ -1566,7 +1566,7 @@ void MainCommon(in FMaterialPixelParameters MaterialParameters, in float4 SvPosi
 	{
 		float4 HeightFogInscatteringAndTransmittance = float4(0, 0, 0, 1);
 
-		if (RenderVolumetricCloudParameters.EnableHeightFog)
+		if (RenderVolumetricCloudParameters.EnableHeightFog > 0)
 		{
 			HeightFogInscatteringAndTransmittance = CalculateHeightFog(FogSampleWorldPositionRelativeToCameraCm);
 		}
@@ -1659,7 +1659,7 @@ void MainCommon(in FMaterialPixelParameters MaterialParameters, in float4 SvPosi
 		{
 			float4 HeightFogInscatteringAndTransmittance = float4(0, 0, 0, 1);
 
-			if (RenderVolumetricCloudParameters.EnableHeightFog)
+			if (RenderVolumetricCloudParameters.EnableHeightFog > 0)
 			{
 				HeightFogInscatteringAndTransmittance = CalculateHeightFog(FogSampleWorldPositionRelativeToCameraCm);
 			}
