VolumetricCloud - added option to evaluate fog on all pixel. A timny cost for seemless fog on cloud at the horizon.

#rnx
#tests PC editor dx12

diff --git a/Engine/Shaders/Private/VolumetricCloud.usf b/Engine/Shaders/Private/VolumetricCloud.usf
--- a/Engine/Shaders/Private/VolumetricCloud.usf
+++ b/Engine/Shaders/Private/VolumetricCloud.usf
@@ -1502,15 +1502,15 @@ void MainCommon(in FMaterialPixelParameters MaterialParameters, in float4 SvPosi
 	float4 ClipPos = mul(float4(RayTranslatedWorldOriginKm * KILOMETER_TO_CENTIMETER + FogSampleWorldPositionRelativeToCameraCm, 1.0f), PrimaryView.TranslatedWorldToClip);
 
 #if SUPPORT_PRIMITIVE_ALPHA_HOLDOUT
-	const bool bApplyFogOnCloud			= tAPWeightsSum;			// Use MeanTransmittance instead of tAPWeightsSum as it is more representative of what can happen since it has not been transformed. (but can result in more fog being applied, this makes alpha holdout correct)
+	const bool bApplyFogOnCloud			= RenderVolumetricCloudParameters.EnableHeightFog >= 2 || tAPWeightsSum;					// Use MeanTransmittance instead of tAPWeightsSum as it is more representative of what can happen since it has not been transformed. (but can result in more fog being applied, this makes alpha holdout correct)
 #if CLOUD_MIN_AND_MAX_DEPTH
-	const bool bApplyFogOnClosestCloud	= ClosestMeanTransmittance < 1.0f;	// Idem
+	const bool bApplyFogOnClosestCloud	= RenderVolumetricCloudParameters.EnableHeightFog >= 2 || ClosestMeanTransmittance < 1.0f;	// Idem
 #endif
 #else // SUPPORT_PRIMITIVE_ALPHA_HOLDOUT
 	// We skip more pixels when using tAPWeightsSum so let;s keep that when alpha output is not enabled for the sake of performance
-	const bool bApplyFogOnCloud			= tAPWeightsSum > 0.0f;
+	const bool bApplyFogOnCloud			= RenderVolumetricCloudParameters.EnableHeightFog >= 2 || tAPWeightsSum > 0.0f;
 #if CLOUD_MIN_AND_MAX_DEPTH
-	const bool bApplyFogOnClosestCloud	= ClostestTAPWeightsSum > 0.0f;
+	const bool bApplyFogOnClosestCloud	= RenderVolumetricCloudParameters.EnableHeightFog >= 2 || ClostestTAPWeightsSum > 0.0f;
 #endif
 #endif // SUPPORT_PRIMITIVE_ALPHA_HOLDOUT
 
@@ -1566,7 +1566,7 @@ void MainCommon(in FMaterialPixelParameters MaterialParameters, in float4 SvPosi
 	{
 		float4 HeightFogInscatteringAndTransmittance = float4(0, 0, 0, 1);
 
-		if (RenderVolumetricCloudParameters.EnableHeightFog)
+		if (RenderVolumetricCloudParameters.EnableHeightFog > 0)
 		{
 			HeightFogInscatteringAndTransmittance = CalculateHeightFog(FogSampleWorldPositionRelativeToCameraCm);
 		}
@@ -1659,7 +1659,7 @@ void MainCommon(in FMaterialPixelParameters MaterialParameters, in float4 SvPosi
 		{
 			float4 HeightFogInscatteringAndTransmittance = float4(0, 0, 0, 1);
 
-			if (RenderVolumetricCloudParameters.EnableHeightFog)
+			if (RenderVolumetricCloudParameters.EnableHeightFog > 0)
 			{
 				HeightFogInscatteringAndTransmittance = CalculateHeightFog(FogSampleWorldPositionRelativeToCameraCm);
 			}
diff --git a/Engine/Source/Runtime/Renderer/Private/VolumetricCloudRendering.cpp b/Engine/Source/Runtime/Renderer/Private/VolumetricCloudRendering.cpp
--- a/Engine/Source/Runtime/Renderer/Private/VolumetricCloudRendering.cpp
+++ b/Engine/Source/Runtime/Renderer/Private/VolumetricCloudRendering.cpp
@@ -75,7 +75,12 @@ static TAutoConsoleVariable<int32> CVarVolumetricCloudHighQualityAerialPerspecti
 static TAutoConsoleVariable<int32> CVarVolumetricCloudDisableCompute(
 	TEXT("r.VolumetricCloud.DisableCompute"), 0,
 	TEXT("Do not use compute shader for cloud tracing."),
-	ECVF_Scalability);
+	ECVF_RenderThreadSafe | ECVF_Scalability);
+
+static TAutoConsoleVariable<int32> CVarVolumetricCloudApplyFogOnAllPixel(
+	TEXT("r.VolumetricCloud.ApplyFogOnAllPixel"), 0,
+	TEXT("When true, fog will be applied on cloud tracing buffer on all pixels instead of only where clouds coverage is >0. This is a tiny bit more expenssive but can remove visual discontinuities when cloud coverage is close to or flip to 0."),
+	ECVF_RenderThreadSafe);
 
 static TAutoConsoleVariable<int32> CVarVolumetricCloudApplyFogLate(
 	TEXT("r.VolumetricCloud.ApplyFogLate"), 0,
@@ -690,7 +695,11 @@ void SetupDefaultRenderVolumetricCloudGlobalParameters(FRDGBuilder& GraphBuilder
 
 	VolumetricCloudParams.BlueNoise = GetBlueNoiseGlobalParameters();
 
-	VolumetricCloudParams.EnableHeightFog = ViewInfo.Family->Scene->HasAnyExponentialHeightFog() && ShouldRenderFog(*ViewInfo.Family);
+	VolumetricCloudParams.EnableHeightFog = 0;
+	if (ViewInfo.Family->Scene->HasAnyExponentialHeightFog() && ShouldRenderFog(*ViewInfo.Family))
+	{
+		VolumetricCloudParams.EnableHeightFog = CVarVolumetricCloudApplyFogOnAllPixel.GetValueOnRenderThread() > 0 ? 2 : 1;
+	}
 	SetupFogUniformParameters(GraphBuilder, ViewInfo, VolumetricCloudParams.FogStruct);
 	if (ShouldVolumetricCloudsApplyFogDuringReconstruction(ViewInfo) && !bForceNonLateFogOnCloud)
 	{
@@ -2260,7 +2269,7 @@ static TRDGUniformBufferRef<FRenderVolumetricCloudGlobalParameters> CreateCloudP
 	}
 
 	VolumetricCloudParams.EnableAerialPerspectiveSampling = CloudRC.bSkipAerialPerspective ? 0 : 1;
-	VolumetricCloudParams.EnableHeightFog = VolumetricCloudParams.EnableHeightFog && !CloudRC.bSkipHeightFog ? 1 : 0;
+	VolumetricCloudParams.EnableHeightFog = VolumetricCloudParams.EnableHeightFog > 0 && !CloudRC.bSkipHeightFog ? VolumetricCloudParams.EnableHeightFog : 0;
 	VolumetricCloudParams.EnableDistantSkyLightSampling = CVarVolumetricCloudEnableDistantSkyLightSampling.GetValueOnAnyThread() > 0 ? 1 : 0;
 	VolumetricCloudParams.EnableAtmosphericLightsSampling = CVarVolumetricCloudEnableAtmosphericLightsSampling.GetValueOnAnyThread() > 0 ? 1 : 0;
 
