Mobile - fog on meshes rendered during the sky pass now matches the deferred renderer: fog is evaluated for far depth.

#rnx
#tests PC editor dx12 / mobile preview

diff --git a/Engine/Shaders/Private/MobileBasePassVertexShader.usf b/Engine/Shaders/Private/MobileBasePassVertexShader.usf
--- a/Engine/Shaders/Private/MobileBasePassVertexShader.usf
+++ b/Engine/Shaders/Private/MobileBasePassVertexShader.usf
@@ -79,7 +79,21 @@ void Main(
 #endif
 
 #if USE_VERTEX_FOG
-	half4 VertexFog = CalculateHeightFog(WorldPosition.xyz - ResolvedView.TranslatedWorldCameraOrigin, EyeIndex, ResolvedView); // WorldPosition is in fact TranslatedWorldPosition 
+
+	half4 VertexFog = float4(0, 0, 0, 1);
+	if (ResolvedView.RealTimeReflectionCapture == 0.0f)
+	{
+		// We only render fog on the sky dome per vertex if not rendering a real time reflection capture. A reflection capture have its own fog pass.
+
+		float3 WorldPositionRelativeToCamera = WorldPosition.xyz - ResolvedView.TranslatedWorldCameraOrigin;
+		if (MATERIAL_IS_SKY)
+		{
+			// If this is a IsSky material, then we evaluated the fog as if it was up to Far DeviceZ. This is to match deferred renderer SkyPass not writting depth.
+			WorldPositionRelativeToCamera = normalize(WorldPosition.xyz - ResolvedView.TranslatedWorldCameraOrigin) * ConvertFromDeviceZ(max(1.0e-38, FarDepthValue));
+		}
+
+		VertexFog = CalculateHeightFog(WorldPositionRelativeToCamera, EyeIndex, ResolvedView); // WorldPosition is in fact TranslatedWorldPosition
+	}
 
 #if PROJECT_SUPPORT_SKY_ATMOSPHERE && MATERIAL_IS_SKY==0 // Do not apply aerial perpsective on sky materials
 	if (ResolvedView.SkyAtmosphereApplyCameraAerialPerspectiveVolume > 0.0f)
