VolumetricFog Fix for PSO creation issue of light function shaders when InstancedStereo is enabled.

#tests PC editor dx12 with and without instanced stereo
#rnx

diff --git a/Engine/Shaders/Private/VolumetricFogLightFunction.usf b/Engine/Shaders/Private/VolumetricFogLightFunction.usf
--- a/Engine/Shaders/Private/VolumetricFogLightFunction.usf
+++ b/Engine/Shaders/Private/VolumetricFogLightFunction.usf
@@ -17,7 +17,8 @@ float2 LightFunctionTexelSize;
 
 void Main(
 	in noperspective float4 UV : TEXCOORD0,
-	in float4 SvPosition : SV_Position,	
+	in FStereoPSInput StereoInput,
+	in float4 SvPosition : SV_Position,
 	out float4 OutColor : SV_Target0
 	)
 {
