[Nanite] Fixed incorrect packing of instance cell data which was causing intermittent culling errors.

#rnx

diff --git a/Engine/Shaders/Private/SceneCulling/SceneCulling.ush b/Engine/Shaders/Private/SceneCulling/SceneCulling.ush
--- a/Engine/Shaders/Private/SceneCulling/SceneCulling.ush
+++ b/Engine/Shaders/Private/SceneCulling/SceneCulling.ush
@@ -46,9 +46,9 @@ FCellHeader UnpackCellHeader(FPackedCellHeader Packed)
 	uint4 Bits = uint4(Packed.Packed0, Packed.Packed1, 0u, 0u);
 	FCellHeader Result;
 	uint Offset = 0u;
-	Result.ItemChunksOffset = ReadBits(Bits, Offset, INSTANCE_HIERARCHY_CELL_HEADER_OFFSET_BITS);
-	Result.NumStaticChunks  = ReadBits(Bits, Offset, INSTANCE_HIERARCHY_CELL_HEADER_COUNT_BITS);
 	Result.NumDynamicChunks = ReadBits(Bits, Offset, INSTANCE_HIERARCHY_CELL_HEADER_COUNT_BITS); 
+	Result.NumStaticChunks  = ReadBits(Bits, Offset, INSTANCE_HIERARCHY_CELL_HEADER_COUNT_BITS);
+	Result.ItemChunksOffset = ReadBits(Bits, Offset, INSTANCE_HIERARCHY_CELL_HEADER_OFFSET_BITS);
 	
 	Result.bIsValid = Result.NumDynamicChunks != 0u;
 	if (Result.bIsValid)
diff --git a/Engine/Source/Runtime/Renderer/Private/SceneCulling/SceneCulling.inl b/Engine/Source/Runtime/Renderer/Private/SceneCulling/SceneCulling.inl
--- a/Engine/Source/Runtime/Renderer/Private/SceneCulling/SceneCulling.inl
+++ b/Engine/Source/Runtime/Renderer/Private/SceneCulling/SceneCulling.inl
@@ -122,7 +122,7 @@ inline bool IsValidCell(const FCellHeader& CellHeader)
 inline bool IsValidCell(const FPackedCellHeader& CellHeader)
 {
 	// For a valid cell the value is always nonzero
-	return CellHeader.Packed1 != 0u;
+	return CellHeader.Packed0 != 0u;
 }
 
 inline FCellHeader UnpackCellHeader(const FPackedCellHeader& Packed)
