[SSS] Fix black halo around thin geometries surrounded by SSS.

diff --git a/Engine/Shaders/Private/SeparableSSS.ush b/Engine/Shaders/Private/SeparableSSS.ush
--- a/Engine/Shaders/Private/SeparableSSS.ush
+++ b/Engine/Shaders/Private/SeparableSSS.ush
@@ -290,7 +290,8 @@ float4 SSSSBlurPS(uint2 BufferPos, float2 BufferUV, float2 dir, bool initStencil
     // Accumulate the center sample:
     float3 colorAccum = 0;
 	// >0 to avoid division by 0, not 100% correct to not visible
-	float3 colorInvDiv = 0.00001f;
+	const float3 initColorInvDiv = 0.00001f;
+	float3 colorInvDiv = initColorInvDiv;
 
 	// center sample
 	half3  CentralKernelWeight = GetSubsurfaceProfileKernel(0, SeparableFilterParameters).rgb;
@@ -323,7 +324,7 @@ float4 SSSSBlurPS(uint2 BufferPos, float2 BufferUV, float2 dir, bool initStencil
 #else
 			uint LocalSubsurfaceProfileInt = SSSSSampleProfileId(LocalUV);
 #endif
-			float3 ColorTint = LocalSubsurfaceProfileInt == SubsurfaceProfileInt ? 1.0f : BoundaryColorBleed;
+			float3 ColorTint = (LocalSubsurfaceProfileInt == SubsurfaceProfileInt || LocalSubsurfaceProfileInt == SSS_PROFILE_ID_INVALID) ? 1.0f : BoundaryColorBleed;
 
 
 			float LocalDepth = color.a;
@@ -337,6 +338,10 @@ float4 SSSSBlurPS(uint2 BufferPos, float2 BufferUV, float2 dir, bool initStencil
 
 			color.a *= 1 - s;
 #endif
+
+#if SUBSURFACE_HALF_RES
+			color.a *= (LocalSubsurfaceProfileInt != SSS_PROFILE_ID_INVALID) ? 1 : 0;
+#endif
 			// approximation, ideally we would reconstruct the mask with GetMaskFromDepthInAlpha() and do manual bilinear filter
 			// needed?
 			color.rgb *= color.a * ColorTint;
@@ -354,6 +359,11 @@ float4 SSSSBlurPS(uint2 BufferPos, float2 BufferUV, float2 dir, bool initStencil
 	// done for each color channel to avoid color shift
 	float3 OutColor = colorAccum / colorInvDiv; 
 
+#if SUBSURFACE_HALF_RES
+	// For half resolution, subsurface is dark if profile id is invalid, which will contribute dark halo to thin geometries surrounded by SSS.
+	// Fallback to the pixel color at the center.
+	OutColor = select(colorInvDiv == initColorInvDiv, colorM.rgb, OutColor);
+#endif
 	// alpha stored the SceneDepth (0 if there is no subsurface scattering)
     return float4(OutColor, OutDepth);
 }
