Substrate - fixed nana due to area light and phase function of simple volume diffuse.

#rnx
#tests PC editor dx12

diff --git a/Engine/Shaders/Private/BRDF.ush b/Engine/Shaders/Private/BRDF.ush
--- a/Engine/Shaders/Private/BRDF.ush
+++ b/Engine/Shaders/Private/BRDF.ush
@@ -137,6 +137,13 @@ void SphereMaxNoH( inout BxDFContext Context, float SinAlpha, bool bNewtonIterat
 #endif
 			Context.VoL = Context.VoL * CosAlpha + VoTr;
 
+#if SUBSTRATE_ENABLED
+			// It seems that VoL can end up out of the [-1, 1] range.
+			// This causes NaN when evaluating the phase function for Substrate SimpleVolume diffuse (sqrt of a negative number is NaN, see HenyeyGreensteinPhase).
+			// So we enforce a valid VoL range here for substrate which does do math when VoL is close to -1 in this case.
+			Context.VoL = clamp(Context.VoL, -1.0, 1.0);
+#endif
+
 			float InvLenH = rsqrt( 2 + 2 * Context.VoL );
 			Context.NoH = saturate( ( Context.NoL + Context.NoV ) * InvLenH );
 // Enable once anisotropic materials support area lights
