VolumetricCloud - added permutation for minmax depth being available.

#rnx
#tests pc editor, console

diff --git a/Engine/Shaders/Private/VolumetricRenderTarget.usf b/Engine/Shaders/Private/VolumetricRenderTarget.usf
--- a/Engine/Shaders/Private/VolumetricRenderTarget.usf
+++ b/Engine/Shaders/Private/VolumetricRenderTarget.usf
@@ -544,7 +544,6 @@ uint4 VolumetricTextureValidCoordRect;
 float4 VolumetricTextureValidUvRect;
 uint ForwardShadingEnable;
 uint OutputAlphaHoldout;
-uint MinMaxDepthAvailable;
 
 #if PERMUTATION_RENDER_UNDERWATER_BUFFER || PERMUTATION_COMPOSE_WITH_WATER
 Texture2D WaterLinearDepthTexture;
@@ -987,108 +986,109 @@ void ComposeVolumetricRTOverScenePS(
 
 
 
-	if (MinMaxDepthAvailable == 0)
-	{	// ==> This path is when compute shader is disable and so no minmax depth texture is available. It is also used by VRT.Mode=1 for upsampling based on front depth only.
+#if PERMUTATION_MINMAXDEPTH_AVAILABLE == 0
+	// ==> This path is when compute shader is disable and so no minmax depth texture is available. It is also used by VRT.Mode=1 for upsampling based on front depth only.
 		
-		// Check that cloud medium sample are in front of the surface we are upsampling over water to make sure we will have valid sample.
-		// This is especially needed when upsampling cloud over water surface where depth is not matching the depth for which clouds have been traced for.
-		// If samples are behind, we assume clouds should not be visible (it will be a harsh transition for now)
-		const bool bAllCloudSamplesInFrontOfWater = all(float4(VolumeFrontDepth0.x, VolumeFrontDepth1.x, VolumeFrontDepth2.x, VolumeFrontDepth3.x) < PixelDistanceFromViewKm);
-		const bool bAnyCloudSamplesInFrontOfWater = any(float4(VolumeFrontDepth0.x, VolumeFrontDepth1.x, VolumeFrontDepth2.x, VolumeFrontDepth3.x) < PixelDistanceFromViewKm);
+	// Check that cloud medium sample are in front of the surface we are upsampling over water to make sure we will have valid sample.
+	// This is especially needed when upsampling cloud over water surface where depth is not matching the depth for which clouds have been traced for.
+	// If samples are behind, we assume clouds should not be visible (it will be a harsh transition for now)
+	const bool bAllCloudSamplesInFrontOfWater = all(float4(VolumeFrontDepth0.x, VolumeFrontDepth1.x, VolumeFrontDepth2.x, VolumeFrontDepth3.x) < PixelDistanceFromViewKm);
+	const bool bAnyCloudSamplesInFrontOfWater = any(float4(VolumeFrontDepth0.x, VolumeFrontDepth1.x, VolumeFrontDepth2.x, VolumeFrontDepth3.x) < PixelDistanceFromViewKm);
 
-		// clamp to the maximum of the read depth to avoid no depth matching
-		PixelDistanceFromViewKm = min(PixelDistanceFromViewKm, max(max(VolumeFrontDepth0.y, VolumeFrontDepth1.y), max(VolumeFrontDepth2.y, VolumeFrontDepth3.y)));
+	// clamp to the maximum of the read depth to avoid no depth matching
+	PixelDistanceFromViewKm = min(PixelDistanceFromViewKm, max(max(VolumeFrontDepth0.y, VolumeFrontDepth1.y), max(VolumeFrontDepth2.y, VolumeFrontDepth3.y)));
 
-		float4 Depth4 = float4(VolumeFrontDepth0.y, VolumeFrontDepth1.y, VolumeFrontDepth2.y, VolumeFrontDepth3.y);
-		float4 Depth4Diff = PixelDistanceFromViewKm - Depth4;
-		Depth4Diff = abs(Depth4Diff);
-		float MaxDepth4Diff = max(max(Depth4Diff.x, Depth4Diff.y), max(Depth4Diff.z, Depth4Diff.w));
+	float4 Depth4 = float4(VolumeFrontDepth0.y, VolumeFrontDepth1.y, VolumeFrontDepth2.y, VolumeFrontDepth3.y);
+	float4 Depth4Diff = PixelDistanceFromViewKm - Depth4;
+	Depth4Diff = abs(Depth4Diff);
+	float MaxDepth4Diff = max(max(Depth4Diff.x, Depth4Diff.y), max(Depth4Diff.z, Depth4Diff.w));
 
-		float ValidSampleCount = 0;
-		float4 DataAcc = 0;
-		const float WeightMultiplier = 1000.0f;
-		const float ThresholdToBilinear = PixelDistanceFromViewKm * 0.1;
-		if (MaxDepth4Diff > ThresholdToBilinear)
-		{
-			float4 VolumeRGBT0 = SafeLoadVolumetricTexture(VolumeCoord0);
-			float4 VolumeRGBT1 = SafeLoadVolumetricTexture(VolumeCoord1);
-			float4 VolumeRGBT2 = SafeLoadVolumetricTexture(VolumeCoord2);
-			float4 VolumeRGBT3 = SafeLoadVolumetricTexture(VolumeCoord3);
+	float ValidSampleCount = 0;
+	float4 DataAcc = 0;
+	const float WeightMultiplier = 1000.0f;
+	const float ThresholdToBilinear = PixelDistanceFromViewKm * 0.1;
+	if (MaxDepth4Diff > ThresholdToBilinear)
+	{
+		float4 VolumeRGBT0 = SafeLoadVolumetricTexture(VolumeCoord0);
+		float4 VolumeRGBT1 = SafeLoadVolumetricTexture(VolumeCoord1);
+		float4 VolumeRGBT2 = SafeLoadVolumetricTexture(VolumeCoord2);
+		float4 VolumeRGBT3 = SafeLoadVolumetricTexture(VolumeCoord3);
 
-			VolumeRGBT0 = ApplyFogToCloudFromUVsAndDepth(VolumeCoord0 * VolumetricTextureSizeAndInvSize.zw, VolumeFrontDepth0.x, VolumeRGBT0);
-			VolumeRGBT1 = ApplyFogToCloudFromUVsAndDepth(VolumeCoord1 * VolumetricTextureSizeAndInvSize.zw, VolumeFrontDepth1.x, VolumeRGBT1);
-			VolumeRGBT2 = ApplyFogToCloudFromUVsAndDepth(VolumeCoord2 * VolumetricTextureSizeAndInvSize.zw, VolumeFrontDepth2.x, VolumeRGBT2);
-			VolumeRGBT3 = ApplyFogToCloudFromUVsAndDepth(VolumeCoord3 * VolumetricTextureSizeAndInvSize.zw, VolumeFrontDepth3.x, VolumeRGBT3);
+		VolumeRGBT0 = ApplyFogToCloudFromUVsAndDepth(VolumeCoord0 * VolumetricTextureSizeAndInvSize.zw, VolumeFrontDepth0.x, VolumeRGBT0);
+		VolumeRGBT1 = ApplyFogToCloudFromUVsAndDepth(VolumeCoord1 * VolumetricTextureSizeAndInvSize.zw, VolumeFrontDepth1.x, VolumeRGBT1);
+		VolumeRGBT2 = ApplyFogToCloudFromUVsAndDepth(VolumeCoord2 * VolumetricTextureSizeAndInvSize.zw, VolumeFrontDepth2.x, VolumeRGBT2);
+		VolumeRGBT3 = ApplyFogToCloudFromUVsAndDepth(VolumeCoord3 * VolumetricTextureSizeAndInvSize.zw, VolumeFrontDepth3.x, VolumeRGBT3);
 
-			if (bIsWaterPixel)
+		if (bIsWaterPixel)
+		{
+			// Only do the following is the pixel is water
+			const float ConsideredFarWaterDistanceKm = 1.0f;				// We only apply the fix for pixel that are further than 1km
+			const float NotRenderedPixelDistanceKm = 500.0f;				// We assume pixel have a depth greater than that such distant depth pixel
+			if (any(Depth4 > NotRenderedPixelDistanceKm)					// Some pixels have non rendered pixel...
+				&& !all(Depth4 > NotRenderedPixelDistanceKm)				// ...but not all (in this case we want to fall back to bilinear filtering).
+				&& PixelDistanceFromViewKm > ConsideredFarWaterDistanceKm)	// Only treat this way pixel that are far enough to not mix up close object with far water.
 			{
-				// Only do the following is the pixel is water
-				const float ConsideredFarWaterDistanceKm = 1.0f;				// We only apply the fix for pixel that are further than 1km
-				const float NotRenderedPixelDistanceKm = 500.0f;				// We assume pixel have a depth greater than that such distant depth pixel
-				if (any(Depth4 > NotRenderedPixelDistanceKm)					// Some pixels have non rendered pixel...
-					&& !all(Depth4 > NotRenderedPixelDistanceKm)				// ...but not all (in this case we want to fall back to bilinear filtering).
-					&& PixelDistanceFromViewKm > ConsideredFarWaterDistanceKm)	// Only treat this way pixel that are far enough to not mix up close object with far water.
+				// This is a special case / hack added for to fix some visual issues encoutnered with water as of today.
+				// Water is render after checker boarded minmax depth is taken for cloud rendering (because clouds also needs to be rendered within the water pass).
+				// As such, the reconstructure fails a taking a matching color for a matching depth (cloud depth = 1000km but water is closer, less than 2 kilometers).
+				// This does not handle all the cases but fixes current issues with the water system.
+				// ==> this code can be removed when we automatically fix up edges of objects to have sharp silhouette under strong conflicting motion.
+				const float AcceptanceOffsetKm = 0.001f;
+				if (VolumeFrontDepth0.y > (PixelDistanceFromViewKm - AcceptanceOffsetKm))
 				{
-					// This is a special case / hack added for to fix some visual issues encoutnered with water as of today.
-					// Water is render after checker boarded minmax depth is taken for cloud rendering (because clouds also needs to be rendered within the water pass).
-					// As such, the reconstructure fails a taking a matching color for a matching depth (cloud depth = 1000km but water is closer, less than 2 kilometers).
-					// This does not handle all the cases but fixes current issues with the water system.
-					// ==> this code can be removed when we automatically fix up edges of objects to have sharp silhouette under strong conflicting motion.
-					const float AcceptanceOffsetKm = 0.001f;
-					if (VolumeFrontDepth0.y > (PixelDistanceFromViewKm - AcceptanceOffsetKm))
-					{
-						DataAcc += VolumeRGBT0;
-						ValidSampleCount++;
-					}
-					if (VolumeFrontDepth1.y > (PixelDistanceFromViewKm - AcceptanceOffsetKm))
-					{
-						DataAcc += VolumeRGBT1;
-						ValidSampleCount++;
-					}
-					if (VolumeFrontDepth2.y > (PixelDistanceFromViewKm - AcceptanceOffsetKm))
-					{
-						DataAcc += VolumeRGBT2;
-						ValidSampleCount++;
-					}
-					if (VolumeFrontDepth3.y > (PixelDistanceFromViewKm - AcceptanceOffsetKm))
-					{
-						DataAcc += VolumeRGBT3;
-						ValidSampleCount++;
-					}
+					DataAcc += VolumeRGBT0;
+					ValidSampleCount++;
+				}
+				if (VolumeFrontDepth1.y > (PixelDistanceFromViewKm - AcceptanceOffsetKm))
+				{
+					DataAcc += VolumeRGBT1;
+					ValidSampleCount++;
+				}
+				if (VolumeFrontDepth2.y > (PixelDistanceFromViewKm - AcceptanceOffsetKm))
+				{
+					DataAcc += VolumeRGBT2;
+					ValidSampleCount++;
+				}
+				if (VolumeFrontDepth3.y > (PixelDistanceFromViewKm - AcceptanceOffsetKm))
+				{
+					DataAcc += VolumeRGBT3;
+					ValidSampleCount++;
 				}
-			}
-
-			if (!bIsWaterPixel || (bIsWaterPixel && ValidSampleCount == 0 && bAnyCloudSamplesInFrontOfWater))
-			{
-				// Depth discontinuities edges
-				float4 weights = 1.0f / (Depth4Diff * WeightMultiplier + 1.0f);
-				const float weightsSum = dot(weights, float4(1.0f, 1.0f, 1.0f, 1.0f));
-				weights /= weightsSum;
-
-				ValidSampleCount = weightsSum > 0.0f ? 1.0 : 0.0f;
-				DataAcc = weights.x * VolumeRGBT0 + weights.y * VolumeRGBT1 + weights.z * VolumeRGBT2 + weights.w * VolumeRGBT3;
 			}
 		}
-		else
-		{
-			// Now do a bilinear sample to have a soft look region without depth edges.
-			ValidSampleCount = 1.0;
-			float2 SampleUVs = (float2(VolumeCoord0) + 0.25 + (float2(OffsetX, OffsetY) * 0.5 + 0.5) * 0.5) * VolumetricTextureSizeAndInvSize.zw;
-			DataAcc = SafeSampleVolumetricTexture(SampleUVs);
-			DataAcc = ApplyFogToCloudFromUVs(SampleUVs, DataAcc);
-		}
 
-		OutputRt0 = float4(0.0f, 0.0f, 0.0f, 1.0f);
-		if (ValidSampleCount > 0.0f)
+		if (!bIsWaterPixel || (bIsWaterPixel && ValidSampleCount == 0 && bAnyCloudSamplesInFrontOfWater))
 		{
-			OutputRt0 = DataAcc / ValidSampleCount;
-		}
-		else
-		{
-			clip(-1.0f);
+			// Depth discontinuities edges
+			float4 weights = 1.0f / (Depth4Diff * WeightMultiplier + 1.0f);
+			const float weightsSum = dot(weights, float4(1.0f, 1.0f, 1.0f, 1.0f));
+			weights /= weightsSum;
+
+			ValidSampleCount = weightsSum > 0.0f ? 1.0 : 0.0f;
+			DataAcc = weights.x * VolumeRGBT0 + weights.y * VolumeRGBT1 + weights.z * VolumeRGBT2 + weights.w * VolumeRGBT3;
 		}
-		return;
 	}
+	else
+	{
+		// Now do a bilinear sample to have a soft look region without depth edges.
+		ValidSampleCount = 1.0;
+		float2 SampleUVs = (float2(VolumeCoord0) + 0.25 + (float2(OffsetX, OffsetY) * 0.5 + 0.5) * 0.5) * VolumetricTextureSizeAndInvSize.zw;
+		DataAcc = SafeSampleVolumetricTexture(SampleUVs);
+		DataAcc = ApplyFogToCloudFromUVs(SampleUVs, DataAcc);
+	}
+
+	OutputRt0 = float4(0.0f, 0.0f, 0.0f, 1.0f);
+	if (ValidSampleCount > 0.0f)
+	{
+		OutputRt0 = DataAcc / ValidSampleCount;
+	}
+	else
+	{
+		clip(-1.0f);
+	}
+	return;
+
+#else // PERMUTATION_MINMAXDEPTH_AVAILABLE == 1
 
 	// ==> Now this is the higher quality upsampling using minmax depth/luminance/trabnsmittance per pixel from reconstruction.
 
@@ -1211,6 +1211,9 @@ void ComposeVolumetricRTOverScenePS(
 	{
 		clip(-1.0f);
 	}
+
+#endif // PERMUTATION_MINMAXDEPTH_AVAILABLE
+
 #endif // PERMUTATION_RENDER_UNDERWATER_BUFFER
 
 #endif
diff --git a/Engine/Source/Runtime/Renderer/Private/VolumetricRenderTarget.cpp b/Engine/Source/Runtime/Renderer/Private/VolumetricRenderTarget.cpp
--- a/Engine/Source/Runtime/Renderer/Private/VolumetricRenderTarget.cpp
+++ b/Engine/Source/Runtime/Renderer/Private/VolumetricRenderTarget.cpp
@@ -870,7 +870,8 @@ class FComposeVolumetricRTOverScenePS : public FGlobalShader
 	class FMSAASampleCount : SHADER_PERMUTATION_SPARSE_INT("PERMUTATION_MSAA_SAMPLE_COUNT", 1, 2, 4, 8);
 	class FApplyFog : SHADER_PERMUTATION_BOOL("PERMUTATION_APPLY_FOG");
 	class FApplyLocalFogVolume : SHADER_PERMUTATION_BOOL("PERMUTATION_APPLY_LOCAL_FOG_VOLUME");
-	using FPermutationDomain = TShaderPermutationDomain<FUpsamplingMode, FRenderUnderWaterBuffer, FRenderCameraComposeWithWater, FMSAASampleCount, FApplyFog, FApplyLocalFogVolume>;
+	class FMinMaxDepthAvailable : SHADER_PERMUTATION_BOOL("PERMUTATION_MINMAXDEPTH_AVAILABLE");
+	using FPermutationDomain = TShaderPermutationDomain<FUpsamplingMode, FRenderUnderWaterBuffer, FRenderCameraComposeWithWater, FMSAASampleCount, FApplyFog, FApplyLocalFogVolume, FMinMaxDepthAvailable>;
 
 	BEGIN_SHADER_PARAMETER_STRUCT(FParameters, )
 		SHADER_PARAMETER_STRUCT_REF(FViewUniformShaderParameters, ViewUniformBuffer)
@@ -886,7 +887,6 @@ class FComposeVolumetricRTOverScenePS : public FGlobalShader
 		RENDER_TARGET_BINDING_SLOTS()
 		SHADER_PARAMETER(uint32, ForwardShadingEnable)
 		SHADER_PARAMETER(uint32, OutputAlphaHoldout)
-		SHADER_PARAMETER(uint32, MinMaxDepthAvailable)
 		SHADER_PARAMETER(float, UvOffsetSampleAcceptanceWeight)
 		SHADER_PARAMETER(FVector4f, VolumetricTextureSizeAndInvSize)
 		SHADER_PARAMETER(FVector2f, FullResolutionToVolumetricBufferResolutionScale)
@@ -1039,6 +1039,7 @@ void ComposeVolumetricRenderTargetOverScene(
 		PermutationVector.Set<FComposeVolumetricRTOverScenePS::FMSAASampleCount>(SceneDepthTexture->Desc.NumSamples);
 		PermutationVector.Set<FComposeVolumetricRTOverScenePS::FApplyFog>(bShouldVolumetricCloudsApplyFogDuringReconstruction && ViewInfo.Family->Scene->HasAnyExponentialHeightFog() && ShouldRenderFog(*ViewInfo.Family));
 		PermutationVector.Set<FComposeVolumetricRTOverScenePS::FApplyLocalFogVolume>(bShouldVolumetricCloudsApplyFogDuringReconstruction && ViewInfo.LocalFogVolumeViewData.GPUInstanceCount > 0);
+		PermutationVector.Set<FComposeVolumetricRTOverScenePS::FMinMaxDepthAvailable>(bMinMaxDepth);
 		TShaderMapRef<FComposeVolumetricRTOverScenePS> PixelShader(ViewInfo.ShaderMap, PermutationVector);
 
 		FComposeVolumetricRTOverScenePS::FParameters* PassParameters = GraphBuilder.AllocParameters<FComposeVolumetricRTOverScenePS::FParameters>();
@@ -1055,7 +1056,6 @@ void ComposeVolumetricRenderTargetOverScene(
 		PassParameters->Substrate = Substrate::BindSubstrateGlobalUniformParameters(ViewInfo);
 		PassParameters->OutputAlphaHoldout = 0;
 		PassParameters->ForwardShadingEnable = bForwardShading ? 1 : 0;
-		PassParameters->MinMaxDepthAvailable = bMinMaxDepth ? 1 : 0;
 		FVector2f DummyUVScale;
 		GetTextureSafeUvCoordBound(VolumetricViewRect, PassParameters->VolumetricTexture, PassParameters->VolumetricTextureValidCoordRect, PassParameters->VolumetricTextureValidUvRect, DummyUVScale);
 
@@ -1153,6 +1153,7 @@ void ComposeVolumetricRenderTargetOverSceneUnderWater(
 		PermutationVector.Set<FComposeVolumetricRTOverScenePS::FRenderUnderWaterBuffer>(1);
 		PermutationVector.Set<FComposeVolumetricRTOverScenePS::FRenderCameraComposeWithWater>(0);
 		PermutationVector.Set<FComposeVolumetricRTOverScenePS::FMSAASampleCount>(1);
+		PermutationVector.Set<FComposeVolumetricRTOverScenePS::FMinMaxDepthAvailable>(bMinMaxDepth);
 		TShaderMapRef<FComposeVolumetricRTOverScenePS> PixelShader(ViewInfo.ShaderMap, PermutationVector);
 
 		FComposeVolumetricRTOverScenePS::FParameters* PassParameters = GraphBuilder.AllocParameters<FComposeVolumetricRTOverScenePS::FParameters>();
@@ -1173,7 +1174,6 @@ void ComposeVolumetricRenderTargetOverSceneUnderWater(
 		PassParameters->Substrate = Substrate::BindSubstrateGlobalUniformParameters(ViewInfo);
 		PassParameters->OutputAlphaHoldout = 0;
 		PassParameters->ForwardShadingEnable = 0;
-		PassParameters->MinMaxDepthAvailable = bMinMaxDepth ? 1 : 0;
 		FVector2f DummyUVScale;
 		GetTextureSafeUvCoordBound(VolumetricViewRect, PassParameters->VolumetricTexture, PassParameters->VolumetricTextureValidCoordRect, PassParameters->VolumetricTextureValidUvRect, DummyUVScale);
 
@@ -1233,6 +1233,7 @@ void ComposeVolumetricRenderTargetOverSceneForVisualization(
 		PermutationVector.Set<FComposeVolumetricRTOverScenePS::FRenderUnderWaterBuffer>(0);
 		PermutationVector.Set<FComposeVolumetricRTOverScenePS::FRenderCameraComposeWithWater>(0);
 		PermutationVector.Set<FComposeVolumetricRTOverScenePS::FMSAASampleCount>(1);
+		PermutationVector.Set<FComposeVolumetricRTOverScenePS::FMinMaxDepthAvailable>(bMinMaxDepth);
 		TShaderMapRef<FComposeVolumetricRTOverScenePS> PixelShader(ViewInfo.ShaderMap, PermutationVector);
 
 		FComposeVolumetricRTOverScenePS::FParameters* PassParameters = GraphBuilder.AllocParameters<FComposeVolumetricRTOverScenePS::FParameters>();
@@ -1248,7 +1249,6 @@ void ComposeVolumetricRenderTargetOverSceneForVisualization(
 		PassParameters->Substrate = Substrate::BindSubstrateGlobalUniformParameters(ViewInfo);
 		PassParameters->OutputAlphaHoldout = 0;
 		PassParameters->ForwardShadingEnable = 0;
-		PassParameters->MinMaxDepthAvailable = bMinMaxDepth ? 1 : 0;
 		FVector2f DummyUVScale;
 		GetTextureSafeUvCoordBound(VolumetricViewRect, PassParameters->VolumetricTexture, PassParameters->VolumetricTextureValidCoordRect, PassParameters->VolumetricTextureValidUvRect, DummyUVScale);
 
