VolumetricCloud - upsampling improvements: less leaking on edges, better decisions for bilinear filtering, less flickering on edge for cloud over mountains.
This is also a perf win. Upsampling to 2848x1600 was down from 0.17ms to 0.1ms on a console.

#rnx
#tests pc editor, console

diff --git a/Engine/Shaders/Private/VolumetricRenderTarget.usf b/Engine/Shaders/Private/VolumetricRenderTarget.usf
--- a/Engine/Shaders/Private/VolumetricRenderTarget.usf
+++ b/Engine/Shaders/Private/VolumetricRenderTarget.usf
@@ -1115,89 +1115,160 @@ void ComposeVolumetricRTOverScenePS(
 	{
 		float4 VolumeRGBT0, VolumeRGBT1, VolumeRGBT2, VolumeRGBT3;
 		float4 Depth4;
-		BRANCH
-		if (abs(PixelDistanceFromViewKm - VolumeFrontDepth0.z) < abs(PixelDistanceFromViewKm - VolumeFrontDepth0.w))// && PixelDistanceFromViewKm < VolumeFrontDepth0.x)
+
+
+		float MinMaxDepthRange = 0;
+
+		uint2 VolumeCoord;
+		float4 VolumeFrontDepth;
+		float DepthDiff;
+		bool bSampleFront;
+		float Tmp;
+
+		Tmp = abs(PixelDistanceFromViewKm - VolumeFrontDepth0.w);
+		//if (Tmp < DepthDiff)	// Accept the first fragment
 		{
-			VolumeRGBT0 = SafeLoadVolumetricSecondaryTexture(VolumeCoord0);
-			Depth4.x = VolumeFrontDepth0.z;
+			VolumeCoord = VolumeCoord0;
+			VolumeFrontDepth = VolumeFrontDepth0.w;
+			DepthDiff = Tmp;
+			bSampleFront = false;
 		}
-		else
+		Tmp = abs(PixelDistanceFromViewKm - VolumeFrontDepth0.z);
+		if (Tmp < DepthDiff)
 		{
-			VolumeRGBT0 = SafeLoadVolumetricTexture(VolumeCoord0);
-			Depth4.x = VolumeFrontDepth0.w;
+			VolumeCoord = VolumeCoord0;
+			VolumeFrontDepth = VolumeFrontDepth0.z;
+			DepthDiff = Tmp;
+			bSampleFront = true;
 		}
-		BRANCH
-		if (abs(PixelDistanceFromViewKm - VolumeFrontDepth1.z) < abs(PixelDistanceFromViewKm - VolumeFrontDepth1.w))// && PixelDistanceFromViewKm < VolumeFrontDepth1.x)
+
+		Tmp = abs(PixelDistanceFromViewKm - VolumeFrontDepth1.w);
+		if (Tmp < DepthDiff)
 		{
-			VolumeRGBT1 = SafeLoadVolumetricSecondaryTexture(VolumeCoord1);
-			Depth4.y = VolumeFrontDepth1.z;
+			VolumeCoord = VolumeCoord1;
+			VolumeFrontDepth = VolumeFrontDepth1.w;
+			DepthDiff = Tmp;
+			bSampleFront = false;
 		}
-		else
+		Tmp = abs(PixelDistanceFromViewKm - VolumeFrontDepth1.z);
+		if (Tmp < DepthDiff)
 		{
-			VolumeRGBT1 = SafeLoadVolumetricTexture(VolumeCoord1);
-			Depth4.y = VolumeFrontDepth1.w;
+			VolumeCoord = VolumeCoord1;
+			VolumeFrontDepth = VolumeFrontDepth1.z;
+			DepthDiff = Tmp;
+			bSampleFront = true;
 		}
-		BRANCH
-		if (abs(PixelDistanceFromViewKm - VolumeFrontDepth2.z) < abs(PixelDistanceFromViewKm - VolumeFrontDepth2.w))// && PixelDistanceFromViewKm < VolumeFrontDepth2.x)
+
+		Tmp = abs(PixelDistanceFromViewKm - VolumeFrontDepth2.w);
+		if (Tmp < DepthDiff)
 		{
-			VolumeRGBT2 = SafeLoadVolumetricSecondaryTexture(VolumeCoord2);
-			Depth4.z = VolumeFrontDepth2.z;
+			VolumeCoord = VolumeCoord2;
+			VolumeFrontDepth = VolumeFrontDepth2.w;
+			DepthDiff = Tmp;
+			bSampleFront = false;
 		}
-		else
+		Tmp = abs(PixelDistanceFromViewKm - VolumeFrontDepth2.z);
+		if (Tmp < DepthDiff)
 		{
-			VolumeRGBT2 = SafeLoadVolumetricTexture(VolumeCoord2);
-			Depth4.z = VolumeFrontDepth2.w;
+			VolumeCoord = VolumeCoord2;
+			VolumeFrontDepth = VolumeFrontDepth2.z;
+			DepthDiff = Tmp;
+			bSampleFront = true;
 		}
-		BRANCH
-		if (abs(PixelDistanceFromViewKm - VolumeFrontDepth3.z) < abs(PixelDistanceFromViewKm - VolumeFrontDepth3.w))// && PixelDistanceFromViewKm < VolumeFrontDepth3.x)
+
+		Tmp = abs(PixelDistanceFromViewKm - VolumeFrontDepth3.w);
+		if (Tmp < DepthDiff)
 		{
-			VolumeRGBT3 = SafeLoadVolumetricSecondaryTexture(VolumeCoord3);
-			Depth4.w = VolumeFrontDepth3.z;
+			VolumeCoord = VolumeCoord3;
+			VolumeFrontDepth = VolumeFrontDepth3.w;
+			DepthDiff = Tmp;
+			bSampleFront = false;
 		}
-		else
+		Tmp = abs(PixelDistanceFromViewKm - VolumeFrontDepth3.z);
+		if (Tmp < DepthDiff)
 		{
-			VolumeRGBT3 = SafeLoadVolumetricTexture(VolumeCoord3);
-			Depth4.w = VolumeFrontDepth3.w;
+			VolumeCoord = VolumeCoord3;
+			VolumeFrontDepth = VolumeFrontDepth3.z;
+			DepthDiff = Tmp;
+			bSampleFront = true;
 		}
 
+		float MaxDepth = max(max(max(VolumeFrontDepth0.w, VolumeFrontDepth0.z), max(VolumeFrontDepth1.w, VolumeFrontDepth1.z)), max(max(VolumeFrontDepth2.w, VolumeFrontDepth2.z), max(VolumeFrontDepth3.w, VolumeFrontDepth3.z)));
+		float MinDepth = min(min(min(VolumeFrontDepth0.w, VolumeFrontDepth0.z), min(VolumeFrontDepth1.w, VolumeFrontDepth1.z)), min(min(VolumeFrontDepth2.w, VolumeFrontDepth2.z), min(VolumeFrontDepth3.w, VolumeFrontDepth3.z)));
+		MinMaxDepthRange = MaxDepth - MinDepth;
 
-		float4 Depth4Diff = PixelDistanceFromViewKm - Depth4;
-		Depth4Diff = abs(Depth4Diff);
-		float MaxDepth4Diff = max(max(Depth4Diff.x, Depth4Diff.y), max(Depth4Diff.z, Depth4Diff.w));
+#define DEBUG_UPSAMPLING 0
 
-		MaxDepth4Diff = max(max(abs(VolumeFrontDepth0.w- VolumeFrontDepth0.z), abs(VolumeFrontDepth1.w - VolumeFrontDepth1.z)), max(abs(VolumeFrontDepth2.w - VolumeFrontDepth2.z), abs(VolumeFrontDepth3.w - VolumeFrontDepth3.z)));
+		const bool bAllCloudSamplesBehindFrontDepth = all(float4(VolumeFrontDepth0.x, VolumeFrontDepth1.x, VolumeFrontDepth2.x, VolumeFrontDepth3.x) > (PixelDistanceFromViewKm + 0.001));
+		if (bAllCloudSamplesBehindFrontDepth)
+		{
+		#if DEBUG_UPSAMPLING
+			OutputRt0 = float4(0.0f, 1.0f, 0.0f, 0.5f);
+		#else
+			clip(-1.0f);
+		#endif
+			return;
+		}
 
-		const float MaximumDistanceKmToAllowUpsample = 2.0f;
-		const bool bAllowUpsamplingOnDistantPixel = true;// any(Depth4 < MaximumDistanceKmToAllowUpsample); // Not used since it causes upsampling problems on such area
 
-		const float WeightMultiplier = 1.0f;
-		const float ThresholdToBilinear = 0.5;// Could be parameterised or based on distance from view, max(PixelDistanceFromViewKm, 2.0f) * 0.1;
-		if (MaxDepth4Diff > ThresholdToBilinear && !bIsWaterPixel && bAllowUpsamplingOnDistantPixel)// && !bAllCloudSamplesInFrontOfDepth)
+		// TODO parameterise from cvar
+		const float DistanceThresholdToBilinear	= 2.0;// Km. Could be parameterised or based on distance from view, max(PixelDistanceFromViewKm, 2.0f) * 0.1;
+		const float DepthDifferenceToBilinear	= 1.0;// IDEM
+		if (!bIsWaterPixel)
 		{
-			// Bilateral upsampling using depth
+			//const bool bAllCloudToBilinear = all(float4(VolumeFrontDepth0.x, VolumeFrontDepth1.x, VolumeFrontDepth2.x, VolumeFrontDepth3.x) > 3)
+			//	&& all(float4(VolumeFrontDepth0.x, VolumeFrontDepth1.x, VolumeFrontDepth2.x, VolumeFrontDepth3.x) < float4(VolumeFrontDepth0.z, VolumeFrontDepth1.z, VolumeFrontDepth2.z, VolumeFrontDepth3.z));
 
-			VolumeRGBT0 = ApplyFogToCloudFromUVsAndDepth(VolumeCoord0 * VolumetricTextureSizeAndInvSize.zw, Depth4.x, VolumeRGBT0);
-			VolumeRGBT1 = ApplyFogToCloudFromUVsAndDepth(VolumeCoord1 * VolumetricTextureSizeAndInvSize.zw, Depth4.y, VolumeRGBT1);
-			VolumeRGBT2 = ApplyFogToCloudFromUVsAndDepth(VolumeCoord2 * VolumetricTextureSizeAndInvSize.zw, Depth4.z, VolumeRGBT2);
-			VolumeRGBT3 = ApplyFogToCloudFromUVsAndDepth(VolumeCoord3 * VolumetricTextureSizeAndInvSize.zw, Depth4.w, VolumeRGBT3);
+			const bool bAllCloudToBilinear = all(float4(VolumeFrontDepth0.x, VolumeFrontDepth1.x, VolumeFrontDepth2.x, VolumeFrontDepth3.x) <= float4(VolumeFrontDepth0.z, VolumeFrontDepth1.z, VolumeFrontDepth2.z, VolumeFrontDepth3.z));
 
-			// Depth discontinuities edges
-			float4 weights = 1.0f / (Depth4Diff * WeightMultiplier + 1.0f);
-			const float weightsSum = dot(weights, float4(1.0f, 1.0f, 1.0f, 1.0f));
-			weights /= weightsSum;
+			if ((all(float2(MinDepth, MaxDepth) > DistanceThresholdToBilinear) && MinMaxDepthRange < DepthDifferenceToBilinear) || bAllCloudToBilinear)
+			{
+				// Bilateral upsampling for a soft look.
+				float2 SampleUVs = (float2(VolumeCoord0) + 0.25 + (float2(OffsetX, OffsetY) * 0.5 + 0.5) * 0.5) * VolumetricTextureSizeAndInvSize.zw;
+				DataAcc = SafeSampleVolumetricTexture(SampleUVs);
 
-			ValidSampleCount = weightsSum > 0.0f ? 1.0 : 0.0f;
-			DataAcc = weights.x * VolumeRGBT0 + weights.y * VolumeRGBT1 + weights.z * VolumeRGBT2 + weights.w * VolumeRGBT3;
-		//	DataAcc.x += 0.2;
+				ValidSampleCount = 1.0;
+				DataAcc = ApplyFogToCloudFromUVs(SampleUVs, DataAcc);
+
+				#if DEBUG_UPSAMPLING
+				DataAcc += float4(0.0, 0.0, 0.1, 0.0);
+				#endif
+			}
+			else
+			{
+				// Choose the best fit for the 2x2 neightborhood around edges
+				float FogDepth = 0;
+				BRANCH
+				if (bSampleFront)
+				{
+					DataAcc = SafeLoadVolumetricSecondaryTexture(VolumeCoord);
+					FogDepth = VolumeFrontDepth.z;
+				}
+				else
+				{
+					DataAcc = SafeLoadVolumetricTexture(VolumeCoord);
+					FogDepth = VolumeFrontDepth.w; 
+				}
+
+				ValidSampleCount = 1;
+				DataAcc = ApplyFogToCloudFromUVsAndDepth(VolumeCoord * VolumetricTextureSizeAndInvSize.zw, FogDepth, DataAcc);
+
+				#if DEBUG_UPSAMPLING
+				DataAcc += bSampleFront ? float4(0.1, 0.0, 0.0, 0.0) : float4(0.0, 0.1, 0.0, 0.0);
+				#endif
+			}
 		}
 		else
 		{
-			// Bilinear upsampling to have a soft look region without depth edges.
+			// Bilateral upsampling for a soft look.
 			ValidSampleCount = 1.0;
 			float2 SampleUVs = (float2(VolumeCoord0) + 0.25 + (float2(OffsetX, OffsetY) * 0.5 + 0.5) * 0.5) * VolumetricTextureSizeAndInvSize.zw;
 			DataAcc = SafeSampleVolumetricTexture(SampleUVs);
 			DataAcc = ApplyFogToCloudFromUVs(SampleUVs, DataAcc);
-		//	DataAcc.y += 0.2;
+
+			#if DEBUG_UPSAMPLING
+			DataAcc += float4(0.1, 0.0, 0.1, 0.0);
+			#endif
 		}
 	}
 
