Fix for banding that occurs in TSR when saving 16 bit float data into 11 or 10 bit float data due to quantization. Added dither based on mantissa precision to make the value fluctuate between its closest neighbors over time.

diff --git a/Engine/Shaders/Private/TemporalSuperResolution/TSRRejectShading.usf b/Engine/Shaders/Private/TemporalSuperResolution/TSRRejectShading.usf
--- a/Engine/Shaders/Private/TemporalSuperResolution/TSRRejectShading.usf
+++ b/Engine/Shaders/Private/TemporalSuperResolution/TSRRejectShading.usf
@@ -41,6 +41,7 @@ float2 ResurrectionJacobianYAdd;
 float2 TranslucencyTextureUVMin;
 float2 TranslucencyTextureUVMax;
 float3 HistoryGuideQuantizationError;
+float3 SceneColorOutputQuantizationError;
 float TheoricBlendFactor;
 uint TileOverscan;
 uint bEnableResurrection;
@@ -590,8 +591,20 @@ void MainCS(
 			for (uint ElementIndex = 0; ElementIndex < SIMD_SIZE; ElementIndex++)
 			{
 				const tsr_short2 OutputPixelPos = ComputeElementOutputPixelPos(ElementIndex);
-		
+				
+			#if CONFIG_ENABLE_STOCASTIC_QUANTIZATION
+				// Banding removal due to loss of precision when saving 16 bit float data into 11_11_10 float.
+				// Adding dithering based on mantissa bits precision stored in SceneColorOutputQuantizationError
+				uint2 Random = Rand3DPCG16(int3(OutputPixelPos.xy, View.StateFrameIndexMod8)).xy;
+				tsr_half E = tsr_half(Hammersley16(0, 1, Random).x);
+				tsr_halfC DitheredColor = CenterFinalSceneColor.GetElement(ElementIndex);
+				DitheredColor.xyz = (tsr_half3)QuantizeForFloatRenderTarget(DitheredColor.xyz, E, SceneColorOutputQuantizationError);
+
+				InputSceneColorOutput[OutputPixelPos] = DitheredColor;
+			#else
 				InputSceneColorOutput[OutputPixelPos] = CenterFinalSceneColor.GetElement(ElementIndex);
+			#endif
+				
 				InputSceneColorLdrLumaOutput[OutputPixelPos] = CenterLumaLDR.GetElement(ElementIndex);
 			}
 		}
diff --git a/Engine/Source/Runtime/Renderer/Private/PostProcess/TemporalSuperResolution.cpp b/Engine/Source/Runtime/Renderer/Private/PostProcess/TemporalSuperResolution.cpp
--- a/Engine/Source/Runtime/Renderer/Private/PostProcess/TemporalSuperResolution.cpp
+++ b/Engine/Source/Runtime/Renderer/Private/PostProcess/TemporalSuperResolution.cpp
@@ -884,6 +884,7 @@ class FTSRRejectShadingCS : public FTSRConvolutionNetworkShader
 		SHADER_PARAMETER(FVector2f, ResurrectionJacobianYMul)
 		SHADER_PARAMETER(FVector2f, ResurrectionJacobianYAdd)
 		SHADER_PARAMETER(FVector3f, HistoryGuideQuantizationError)
+		SHADER_PARAMETER(FVector3f, SceneColorOutputQuantizationError)
 		SHADER_PARAMETER(float, FlickeringFramePeriod)
 		SHADER_PARAMETER(float, TheoricBlendFactor)
 		SHADER_PARAMETER(int32, TileOverscan)
@@ -2838,6 +2839,7 @@ FDefaultTemporalUpscaler::FOutputs AddTemporalSuperResolutionPasses(
 			PassParameters->ResurrectionJacobianYAdd = ScreenVelocityToInputPixelVelocity * FVector2f(0.0f, InputPixelVelocityToScreenVelocity.Y);
 		}
 		PassParameters->HistoryGuideQuantizationError = ComputePixelFormatQuantizationError(History.GuideArray->Desc.Format);
+		PassParameters->SceneColorOutputQuantizationError = ComputePixelFormatQuantizationError(HistoryColorFormat);
 		PassParameters->FlickeringFramePeriod = FlickeringFramePeriod;
 		PassParameters->TheoricBlendFactor = 1.0f / (1.0f + MaxHistorySampleCount / OutputToInputResolutionFractionSquare);
 		PassParameters->TileOverscan = TileOverscan;
