[SSS] Fix package and test error for tiling shader file due to missing function for substrate, and type.

diff --git a/Engine/Shaders/Private/PostProcessSubsurfaceTile.usf b/Engine/Shaders/Private/PostProcessSubsurfaceTile.usf
--- a/Engine/Shaders/Private/PostProcessSubsurfaceTile.usf
+++ b/Engine/Shaders/Private/PostProcessSubsurfaceTile.usf
@@ -132,6 +132,15 @@ int MaxGroupCount;
 int2 TiledViewRes;
 RWStructuredBuffer<uint> TileMaskBufferOut;
 
+//Duplicate from PostProcessSubsurface.usf
+#if SUBTRATE_GBUFFER_FORMAT==1
+FSubstrateSubsurfaceHeader LoadSubstrateSSSHeader(float2 UV)
+{
+	const float2 PixelPos = UV.xy * View.BufferSizeAndInvSize.xy;
+	return SubstrateLoadSubsurfaceHeader(Substrate.MaterialTextureArray, Substrate.FirstSliceStoringSubstrateSSSData, PixelPos);
+}
+#endif
+
 // Only access shading model GBuffer
 float HasSubsurfaceForOnePixel(float2 UVSceneColor)
 {
@@ -143,7 +152,7 @@ float HasSubsurfaceForOnePixel(float2 UVSceneColor)
 	const bool bHasSSSDiffusion = UseSubsurfaceProfile(ScreenSpaceData.GBuffer.ShadingModelID);
 #endif
 
-	return select(bHasSSSDiffusion, 1, 0);
+	return select(bHasSSSDiffusion, 1.0f, 0.0f);
 }
 
 // Simpler logic from the SSS::Setup pass
