[SingleLayerWater] Refraction optimization now just uses a black color instead of water volume albedo as fallback when there's no refraction data.

diff --git a/Engine/Shaders/Private/SingleLayerWaterShading.ush b/Engine/Shaders/Private/SingleLayerWaterShading.ush
--- a/Engine/Shaders/Private/SingleLayerWaterShading.ush
+++ b/Engine/Shaders/Private/SingleLayerWaterShading.ush
@@ -222,7 +222,7 @@ WaterVolumeLightingOutput EvaluateWaterVolumeLighting(
 		// We use the pixel SvPosition instead of the scene one pre refraction/distortion to avoid those extra ALUs.
 		float3 BehindWaterSceneLuminance = SceneColorWithoutSingleLayerWaterTexture.SampleLevel(SceneColorWithoutSingleLayerWaterSampler, ViewportUV, 0).rgb;
 		BehindWaterSceneLuminance = MeanTransmittanceToLightSources * ResolvedView.OneOverPreExposure * BehindWaterSceneLuminance;
-		BehindWaterSceneLuminance = lerp(Albedo * ResolvedView.OneOverPreExposure, BehindWaterSceneLuminance, RefractionMask);
+		BehindWaterSceneLuminance *= RefractionMask; // Multiply by refraction mask to essentially turn invalid refraction data black.
 #endif
 
 		float3 SunScattLuminance = DirLightPhaseValue * SunIlluminance;
