Path Tracer: Implement support for local lights on volumetric clouds

This implementation is not yet complete. The equi-angular sampling strategy that is enabled here is not always the best choice, but this puts in place the framework necessary to eventually combine strategies more optimally.

diff --git a/Engine/Shaders/Private/PathTracing/PathTracingCommon.ush b/Engine/Shaders/Private/PathTracing/PathTracingCommon.ush
--- a/Engine/Shaders/Private/PathTracing/PathTracingCommon.ush
+++ b/Engine/Shaders/Private/PathTracing/PathTracingCommon.ush
@@ -534,11 +534,22 @@ struct FPackedPathTracingPayload : FMinimalPayload
 	{
 		PackedData[0] = RandSeq.SampleIndex;
 		PackedData[1] = RandSeq.SampleSeed;
+		// sampling mode: need to pack transmittance and albedo together
 		PackedData[2] = PackFloat2ToUInt(Throughput.xy);
 		PackedData[3] = PackFloat2ToUInt(float2(Throughput.z, Albedo.x));
 		PackedData[4] = PackFloat2ToUInt(Albedo.yz);
 	}
 
+	void SetVolumetricCallableShaderOutput(RandomSequence RandSeq, float3 Throughput)
+	{
+		// Increased precision for Transmittance mode
+		PackedData[0] = RandSeq.SampleIndex;
+		PackedData[1] = RandSeq.SampleSeed;
+		PackedData[2] = asuint(Throughput.x);
+		PackedData[3] = asuint(Throughput.y);
+		PackedData[4] = asuint(Throughput.z);
+	}
+
 	void SetVolumetricCallableShaderOutputSample(float SampleT, float3 AlbedoRayleigh, float3 AlbedoHG1, float PhaseG, float3 AlbedoHG2, float3 DualPhaseData, float3 PayloadThroughput)
 	{
 		HitT = SampleT;
@@ -581,6 +592,12 @@ struct FPackedPathTracingPayload : FMinimalPayload
 		return float3(UnpackFloat2FromUInt(PackedData[3]).y, UnpackFloat2FromUInt(PackedData[4]));
 	}
 
+	float3 GetVolumetricCallableShaderOutputTransmittanceThroughput()
+	{
+		// Increased precision for Transmittance mode
+		return asfloat(uint3(PackedData[2], PackedData[3], PackedData[4]));
+	}
+
 	float3 GetVolumetricCallableShaderOutputPayloadThroughput()
 	{
 		return float3(UnpackFloat2FromUInt(PackedData[5]), UnpackFloat2FromUInt(PackedData[6]).x);
@@ -616,7 +633,7 @@ struct FPackedPathTracingPayload : FMinimalPayload
 		return f16tof32(PackedData[12]);
 	}
 
-	float GetVolumetricCallableCloudDensityFactor()
+	float GetVolumetricCallableShaderOutputCloudDensityFactor()
 	{
 		return f16tof32(PackedData[12] >> 16);
 	}
diff --git a/Engine/Shaders/Private/PathTracing/PathTracingCore.ush b/Engine/Shaders/Private/PathTracing/PathTracingCore.ush
--- a/Engine/Shaders/Private/PathTracing/PathTracingCore.ush
+++ b/Engine/Shaders/Private/PathTracing/PathTracingCore.ush
@@ -740,7 +740,7 @@ FPackedPathTracingPayload TraceTransparentRay(inout FPathState PathState, int Bo
 					PathState.RandSequence,
 					Interval.VolumeMask |
 					(VolumeFlags & ~PATH_TRACER_VOLUME_ENABLE_MASK) |
-					(min(Bounce, 31) << PATH_TRACER_VOLUME_CALLABLE_FLAGS_BOUNCE_SHIFT) |
+					min(Bounce << PATH_TRACER_VOLUME_CALLABLE_FLAGS_BOUNCE_SHIFT, PATH_TRACER_VOLUME_CALLABLE_FLAGS_BOUNCE_MASK) |
 					(bNeedSample  ? PATH_TRACER_VOLUME_CALLABLE_FLAGS_GET_SAMPLE    : 0),
 					PathState.PathThroughput,
 					Interval.CloudDensity
@@ -777,7 +777,7 @@ FPackedPathTracingPayload TraceTransparentRay(inout FPathState PathState, int Bo
 						CloudPayload.GetVolumetricCallableShaderOutputDualHGWeight(),
 						CloudPayload.GetVolumetricCallableShaderOutputDualHGPhaseData()
 					);
-					Result.SetCloudFactor(CloudPayload.GetVolumetricCallableCloudDensityFactor());
+					Result.SetCloudFactor(CloudPayload.GetVolumetricCallableShaderOutputCloudDensityFactor());
 					Payload = PackPathTracingPayload(Result);
 				}
 			}
@@ -787,7 +787,7 @@ FPackedPathTracingPayload TraceTransparentRay(inout FPathState PathState, int Bo
 				// In order to use analytic transmittance, the user must have requested it _and_ we must be on an interval with no clouds or heterogeneous volumes
 				const bool bUseAnalyticTransmittance =
 					(VolumeFlags & PATH_TRACER_VOLUME_USE_ANALYTIC_TRANSMITTANCE) != 0 &&
-					(Interval.VolumeMask & (PATH_TRACER_VOLUME_ENABLE_CLOUDS | PATH_TRACER_VOLUME_ENABLE_HETEROGENEOUS_VOLUMES)) == 0;
+					(Interval.VolumeMask & PATH_TRACER_VOLUME_ENABLE_HETEROGENEOUS_VOLUMES) == 0;
 				// are there any holdouts on the current segment?
 				const bool bHaveHoldouts =
 					(Interval.VolumeMask << PATH_TRACER_MAX_VOLUMES) & VolumeFlags & PATH_TRACER_VOLUME_HOLDOUT_MASK != 0;
@@ -1771,30 +1771,27 @@ bool PathTracingKernel(inout FPathState PathState, int Bounce)
 			float EquiAngularT = SampleResult.x;
 			float EquiAngularPathPdf = SampleResult.y;
 
-			// find out how much volume exists at the current point
-			const float3 TranslatedWorldPos = Ro + EquiAngularT * Rd;
-			const FVolumeShadedResult Result = VolumeGetDensity(TranslatedWorldPos, VolumeSegment.Interval.VolumeMask, bIsCameraRay);
-			const float3 SigmaT = Result.SigmaT;
-			const float3 SigmaS = min(Result.SigmaSRayleigh + Result.SigmaSHG + Result.SigmaSDualHG, SigmaT);
 
-			if (any(SigmaS > 0))
+			// Account for the transmittance up to the current point within the current slice of volume
+			VolumeSegment.Interval.VolumeTMax = EquiAngularT;
+			float CloudFactor = 1.0;
+			const FVolumeShadedResult Result = VolumeGetDensityAndTransmittance(Ro, Rd, VolumeSegment.Interval, Bounce, PathState.RandSequence, VolumeSegment.Throughput, CloudFactor);
+			float3 Contrib = VolumeSegment.Throughput;
+			if (any(Contrib > 0))
 			{
-				// Account for the transmittance up to the current point within the current slice of volume
-				VolumeSegment.Interval.VolumeTMax = EquiAngularT;
-				VolumeSegment.Throughput = VolumeGetTransmittance(VolumeSegment.Throughput, Ro, Rd, VolumeSegment.Interval, PathState.RandSequence);
-
-				float3 Contrib = VolumeSegment.Throughput * SigmaS;
-
 				// account for probability of the path prefix
 				Contrib /= EquiAngularPathPdf * LightPickPdf;
 
+				// find out how much volume exists at the current point
+				const float3 TranslatedWorldPos = Ro + EquiAngularT * Rd;
+
 				// prepare a minimal payload that describes the hit we need to shade
 				FPathTracingPayload VolPayload = CreateMediumHitPayload(EquiAngularT, TranslatedWorldPos, Result);
 				float3 LightRandValue = RandomSequence_GenerateSample3D(PathState.RandSequence);
 
 				const bool bCastShadows = CastsVolumeShadow(LightSampleLightId);
 				const bool bCloudCastShadows = true; // volumes should get shadows from clouds
-				const float CloudFactor = 1.0; // TODO: the handling of local lights does not yet account for clouds correctly, need to fetch this from another call to the callable shader
+
 				const uint MissShaderIndex = GetLightMissShaderIndex(LightSampleLightId);
 				// compute direct light sampling?
 				if (MISMode != 0)
diff --git a/Engine/Shaders/Private/PathTracing/PathTracingVolumetricCloudMaterialShader.usf b/Engine/Shaders/Private/PathTracing/PathTracingVolumetricCloudMaterialShader.usf
--- a/Engine/Shaders/Private/PathTracing/PathTracingVolumetricCloudMaterialShader.usf
+++ b/Engine/Shaders/Private/PathTracing/PathTracingVolumetricCloudMaterialShader.usf
@@ -101,13 +101,19 @@ RAY_TRACING_ENTRY_CALLABLE(PathTracingVolumetricCloudMaterialShader, FPackedPath
 	FRayDesc Ray = PackedPayload.GetVolumetricCallableShaderInputRay();
 	float3 PathThroughput = PackedPayload.GetVolumetricCallableShaderInputThroughput();
 	float CloudDensity = PackedPayload.GetVolumetricCallableShaderInputCloudDensity();
-	uint Flags = PackedPayload.GetVolumetricCallableShaderInputFlags();
+	const uint Flags = PackedPayload.GetVolumetricCallableShaderInputFlags();
+
+	const int Bounce = (Flags & PATH_TRACER_VOLUME_CALLABLE_FLAGS_BOUNCE_MASK) >> PATH_TRACER_VOLUME_CALLABLE_FLAGS_BOUNCE_SHIFT;
+	const bool bIsCameraRay = Bounce == 0;
 
 	// TODO: evaluate if it would be better to make a second shader for this (so that we don't get two copies of the shader eval in inner loops)
 	// TODO: evaluate if voxel traversal is worth it for transmittance
 	if (Flags & PATH_TRACER_VOLUME_CALLABLE_FLAGS_TRANSMITTANCE)
 	{
-		float CloudFactor = PackedPayload.GetVolumetricCallableShaderInputCloudFactor();
+		const float CloudFactor = PackedPayload.GetVolumetricCallableShaderInputCloudFactor();
+		// Reset payload
+		PackedPayload = (FPackedPathTracingPayload)0;
+
 		// Transmittance loop for the cloud shader only
 		// NOTE: We have specialized this since we know our upper bound is a scalar value, most of the pdf math cancels out and we can take a faster path
 		float TMin = Ray.TMin;
@@ -139,15 +145,84 @@ RAY_TRACING_ENTRY_CALLABLE(PathTracingVolumetricCloudMaterialShader, FPackedPath
 			}
 		}
 
-		PackedPayload = (FPackedPathTracingPayload)0;
-		PackedPayload.SetVolumetricCallableShaderOutput(RandSeq, Throughput, 0.0 /* Albedo not used for transmittance */);
+		// Do we need to know the sample at the _end_ of the ray interval? capture it first in case it does not scatter any light
+		if (Flags & PATH_TRACER_VOLUME_CALLABLE_FLAGS_GET_SAMPLE)
+		{
+#if MATERIAL_VOLUMETRIC_ADVANCED_MULTISCATTERING_OCTAVE_COUNT > 0
+			FCloudMaterialMSParams MS = (FCloudMaterialMSParams)0; 
+			if (PathTracingCloudParameters.CloudShaderMultipleScatterApproxEnabled)
+			{
+				MS = CloudMaterialGetMSParams();
+			}
+			FRISContext HitSample = InitRISContext(RandomSequence_GenerateSample1D(RandSeq));
+			float3 FinalThroughput = 0;
+#endif
+
+			float3 TranslatedWorldPos = Ray.Origin + Ray.Direction * Ray.TMax;
+			FVolumeShadedResult CloudResult = (FVolumeShadedResult)0; 
+			MergeVolumeShadedResult(CloudResult, CloudsGetDensity(TranslatedWorldPos), bIsCameraRay && (Flags & PATH_TRACER_VOLUME_HOLDOUT_CLOUDS) != 0);
+			FVolumeShadedResult Result = CloudResult;
+			if (Flags & PATH_TRACER_VOLUME_ENABLE_ATMOSPHERE)
+			{
+				MergeVolumeShadedResult(Result, AtmosphereGetDensity(TranslatedWorldPos), bIsCameraRay && (Flags & PATH_TRACER_VOLUME_HOLDOUT_ATMOSPHERE) != 0);
+			}
+			if (Flags & PATH_TRACER_VOLUME_ENABLE_FOG)
+			{
+				MergeVolumeShadedResult(Result, FogGetDensity(TranslatedWorldPos), bIsCameraRay && (Flags & PATH_TRACER_VOLUME_HOLDOUT_FOG) != 0);
+			}
+
+			float3 SigmaS = min(Result.SigmaSRayleigh + Result.SigmaSHG + Result.SigmaSDualHG, Result.SigmaT);
+#if MATERIAL_VOLUMETRIC_ADVANCED_MULTISCATTERING_OCTAVE_COUNT > 0
+			float3 Contrib = Throughput * SigmaS;
+			float SelectionWeight = max3(Contrib.x, Contrib.y, Contrib.z);
+			if (HitSample.Accept(SelectionWeight))
+#else
+			Throughput *= SigmaS;
+#endif
+			{
+				const float3 RayleighWeight = select(SigmaS > 0.0, Result.SigmaSRayleigh / SigmaS, 0.0);
+				const float3 HGWeight       = select(SigmaS > 0.0, Result.SigmaSHG       / SigmaS, 0.0);
+				const float3 DualHGWeight   = select(SigmaS > 0.0, Result.SigmaSDualHG   / SigmaS, 0.0);
+				PackedPayload.SetVolumetricCallableShaderOutputSample(Ray.TMax, RayleighWeight, HGWeight, Result.PhaseG, DualHGWeight, Result.DualPhaseData, 0.0);
+				PackedPayload.SetVolumetricCallableShaderOutputCloudDensityFactor(1.0);
+#if MATERIAL_VOLUMETRIC_ADVANCED_MULTISCATTERING_OCTAVE_COUNT > 0
+				FinalThroughput = Contrib / SelectionWeight;
+#endif
+			}
+#if MATERIAL_VOLUMETRIC_ADVANCED_MULTISCATTERING_OCTAVE_COUNT > 0
+			float S = MS.ScattFactor, E = MS.ExtinFactor, P = MS.PhaseFactor;
+			for (int ms = 1; ms <= MATERIAL_VOLUMETRIC_ADVANCED_MULTISCATTERING_OCTAVE_COUNT; ms++)
+			{
+				CloudResult.SigmaSDualHG *= S;
+
+				// Propose a candidate hit with just the scaled cloud contribution
+				Contrib = Throughput * CloudResult.SigmaSDualHG;
+				SelectionWeight = max3(Contrib.x, Contrib.y, Contrib.z);
+				if (HitSample.Accept(SelectionWeight))
+				{
+					// stash this hit for next time
+					FinalThroughput = Contrib / SelectionWeight;
+					// We create a blend of SigmaHG (set to isotropic) and SigmaSDualHG (set to the original phase function) to match what raster does (lerp the eval)
+					PackedPayload.SetVolumetricCallableShaderOutputSample(Ray.TMax, 0.0, P, 0.0, 1.0 - P, CloudResult.DualPhaseData, 0.0);
+					PackedPayload.SetVolumetricCallableShaderOutputCloudDensityFactor(E); // so that future cloud shadow rays get attenuated
+				}
+
+				// Update scale factors for next iteration
+				S *= S;
+				P *= P;
+				E *= E;
+			}
+			Throughput = FinalThroughput * HitSample.GetNormalization();
+#endif
+		}
+
+		PackedPayload.SetVolumetricCallableShaderOutput(RandSeq, Throughput);
 	}
 	else
 	{
-		FRISContext HitSample = PackedPayload.GetVolumetricCallableShaderRISContext();
 		// Loop through combined density for sampling
-		const int Bounce = (Flags & PATH_TRACER_VOLUME_CALLABLE_FLAGS_BOUNCE_MASK) >> PATH_TRACER_VOLUME_CALLABLE_FLAGS_BOUNCE_SHIFT;
-		const bool bIsCameraRay = Bounce == 0;
+
+		FRISContext HitSample = PackedPayload.GetVolumetricCallableShaderRISContext();
 
 #if MATERIAL_VOLUMETRIC_ADVANCED_MULTISCATTERING_OCTAVE_COUNT > 0
 		FCloudMaterialMSParams MS = (FCloudMaterialMSParams)0; 
diff --git a/Engine/Shaders/Private/PathTracing/Volume/PathTracingVolume.ush b/Engine/Shaders/Private/PathTracing/Volume/PathTracingVolume.ush
--- a/Engine/Shaders/Private/PathTracing/Volume/PathTracingVolume.ush
+++ b/Engine/Shaders/Private/PathTracing/Volume/PathTracingVolume.ush
@@ -162,20 +162,8 @@ float3 VolumeGetTransmittance(float3 StartThroughput, float3 Origin, float3 Dire
 			PATH_TRACER_REFERENCE_TRANSMITTANCE_LOOP(AtmosphereGetDensityBounds, AtmosphereGetDensity);
 		}
 	}
-#if PATH_TRACER_USE_CLOUD_SHADER
-	if (Interval.VolumeMask & PATH_TRACER_VOLUME_ENABLE_CLOUDS)
-	{
-		float TMin = Interval.VolumeTMin;
-		float TMax = Interval.VolumeTMax;
-		float SigmaBar = Interval.CloudDensity;
-		FPackedPathTracingPayload CloudPayload = (FPackedPathTracingPayload) 0;
-		CloudPayload.SetVolumetricCallableShaderInput(Origin, Direction, TMin, TMax, RandSequence, PATH_TRACER_VOLUME_CALLABLE_FLAGS_TRANSMITTANCE, Throughput, SigmaBar);
-		CloudPayload.SetVolumetricCallableShaderInputCloudFactor(1.0);
-		CallShader(CloudCallableShaderId, CloudPayload);
-		RandSequence = CloudPayload.GetVolumetricCallableShaderOutputRandSeq();
-		Throughput = CloudPayload.GetVolumetricCallableShaderOutputThroughput();
-	}
-#endif
+	// NOTE: This function does not include cloud transmittance as it should only be called when cloud density
+	//       has already been accounted for or is not present on this interval.
 	if (Interval.VolumeMask & PATH_TRACER_VOLUME_ENABLE_FOG)
 	{
 		float TMin = Interval.VolumeTMin;
@@ -227,7 +215,7 @@ float3 VolumeGetTransmittance(float3 StartThroughput, float3 Origin, float3 Dire
 		CloudPayload.SetVolumetricCallableShaderInputCloudFactor(CloudFactor);
 		CallShader(CloudCallableShaderId, CloudPayload);
 		RandSequence = CloudPayload.GetVolumetricCallableShaderOutputRandSeq();
-		Throughput = CloudPayload.GetVolumetricCallableShaderOutputThroughput();
+		Throughput = CloudPayload.GetVolumetricCallableShaderOutputTransmittanceThroughput();
 	}
 #endif
 	if (IntersectionList.VolumeMask & PATH_TRACER_VOLUME_ENABLE_FOG)
@@ -252,6 +240,54 @@ float3 VolumeGetTransmittance(float3 StartThroughput, float3 Origin, float3 Dire
 	return Throughput;
 }
 
+FVolumeShadedResult VolumeGetDensityAndTransmittance(float3 Origin, float3 Direction, FVolumeIntersectionInterval Interval, uint Bounce, inout RandomSequence RandSequence, inout float3 Throughput, inout float CloudFactor)
+{
+	FVolumeShadedResult Result = (FVolumeShadedResult) 0;
+#if PATH_TRACER_USE_CLOUD_SHADER
+	if (Interval.VolumeMask & PATH_TRACER_VOLUME_ENABLE_CLOUDS)
+	{
+		float TMin = Interval.VolumeTMin;
+		float TMax = Interval.VolumeTMax;
+		float SigmaBar = Interval.CloudDensity;
+		FPackedPathTracingPayload CloudPayload = (FPackedPathTracingPayload) 0;
+		uint Flags = Interval.VolumeMask;
+		Flags |= VolumeFlags & ~PATH_TRACER_VOLUME_ENABLE_MASK;
+		Flags |= PATH_TRACER_VOLUME_CALLABLE_FLAGS_TRANSMITTANCE;
+		Flags |= PATH_TRACER_VOLUME_CALLABLE_FLAGS_GET_SAMPLE;
+		Flags |= min(PATH_TRACER_VOLUME_CALLABLE_FLAGS_BOUNCE_MASK, Bounce << PATH_TRACER_VOLUME_CALLABLE_FLAGS_BOUNCE_SHIFT);
+		CloudPayload.SetVolumetricCallableShaderInput(Origin, Direction, TMin, TMax, RandSequence, Flags, Throughput, SigmaBar);
+		CloudPayload.SetVolumetricCallableShaderInputCloudFactor(1.0);
+		CallShader(CloudCallableShaderId, CloudPayload);
+		RandSequence = CloudPayload.GetVolumetricCallableShaderOutputRandSeq();
+		Throughput = CloudPayload.GetVolumetricCallableShaderOutputTransmittanceThroughput();
+
+		Result = (FVolumeShadedResult) 0;
+		// NOTE: SigmaT is not used, SigmaS values are weights
+		Result.SigmaSRayleigh = CloudPayload.GetVolumetricCallableShaderOutputRayleighWeight();
+		float4 HGData = CloudPayload.GetVolumetricCallableShaderOutputHG();
+		Result.SigmaSHG = HGData.xyz;
+		Result.PhaseG = HGData.w;
+		Result.SigmaSDualHG = CloudPayload.GetVolumetricCallableShaderOutputDualHGWeight();
+		Result.DualPhaseData = CloudPayload.GetVolumetricCallableShaderOutputDualHGPhaseData();
+		CloudFactor = CloudPayload.GetVolumetricCallableShaderOutputCloudDensityFactor();
+	}
+	else
+#endif
+	{
+		// non-cloud case, just get density of all volumes
+		Result = VolumeGetDensity(Origin + Direction * Interval.VolumeTMax, Interval.VolumeMask, Bounce == 0);
+		CloudFactor = 1.0;
+		// Factor in SigmaS to throughput
+		Throughput *= Result.SigmaSRayleigh + Result.SigmaSHG + Result.SigmaSDualHG;
+	}
+	// Factor in remainder of transmittance terms from non-cloud volumes along the segment if we still have some energy left
+	if (any(Throughput > 0))
+	{
+		Throughput = VolumeGetTransmittance(Throughput, Origin, Direction, Interval, RandSequence);
+	}
+	return Result;
+}
+
 FVolumeTrackingResult VolumeSampleDistance(float3 PathThroughput, float3 Origin, float3 Direction, FVolumeIntersectionInterval Interval, bool bIsCameraRay, inout RandomSequence RandSequence)
 {
 	FVolumeDensityBounds VolumeDensityBounds = VolumeGetDensityBounds(Origin, Direction, Interval);
