Volumetric cloud - improved and optimised apply late fog.
This is now enabled by deaffult for volumetric fog on clouds to always be correct when lumen is enabled.

#tests PC editor dx12, console
#rnx

diff --git a/Engine/Shaders/Private/VolumetricRenderTarget.usf b/Engine/Shaders/Private/VolumetricRenderTarget.usf
--- a/Engine/Shaders/Private/VolumetricRenderTarget.usf
+++ b/Engine/Shaders/Private/VolumetricRenderTarget.usf
@@ -623,6 +623,12 @@ float4 ApplyFogToCloudSampleFromVPosAndDepth(float2 SVPos, float DepthKilometer,
 {
 #if APPLY_FOG_LATE
 
+	const float CloudCoverage = saturate(1.0 - CloudRGBT.a);
+	if (CloudCoverage <= 0.0f)
+	{
+		return CloudRGBT;
+	}
+
 	float4 HeightFogInscatteringAndTransmittance = float4(0, 0, 0, 1);
 	const float2 SVPosition = SVPos;
 	const float DistanceToPoint = DepthKilometer * KILOMETER_TO_CENTIMETER;
@@ -663,7 +669,7 @@ float4 ApplyFogToCloudSampleFromVPosAndDepth(float2 SVPos, float DepthKilometer,
 
 	HeightFogInscatteringAndTransmittance.rgb *= PrimaryView.PreExposure;
 
-	CloudRGBT.rgb = CloudRGBT.rgb * HeightFogInscatteringAndTransmittance.a + (1.0 - CloudRGBT.a) * HeightFogInscatteringAndTransmittance.rgb;
+	CloudRGBT.rgb = CloudRGBT.rgb * HeightFogInscatteringAndTransmittance.a + CloudCoverage * HeightFogInscatteringAndTransmittance.rgb;
 
 #endif // APPLY_FOG_LATE
 
@@ -906,15 +912,26 @@ void ComposeVolumetricRTOverScenePS(
 
 		float ValidSampleCount = 0.0f;
 		float4 DataAcc = 0.0f;
+		float DepthForFog = 0.0f;
 		const float CloudFrontDepthTinyOffset = 0.001;
 
 		// We are testing if the depth buffer is further than the cloud front depth and that the cloud front depth is actually in front of traced depth.
-		if (PixelDistanceFromViewKm > (VolumeFrontDepth0.x - CloudFrontDepthTinyOffset) && VolumeFrontDepth0.x < (VolumeFrontDepth0.y)) { VolumeRGBT0 = ApplyFogToCloudFromUVsAndDepth(VolumeUVOffset0, VolumeFrontDepth0.x, VolumeRGBT0); DataAcc += VolumeRGBT0; ValidSampleCount += 1.0f; }
-		if (PixelDistanceFromViewKm > (VolumeFrontDepth1.x - CloudFrontDepthTinyOffset) && VolumeFrontDepth1.x < (VolumeFrontDepth1.y)) { VolumeRGBT1 = ApplyFogToCloudFromUVsAndDepth(VolumeUVOffset1, VolumeFrontDepth1.x, VolumeRGBT1); DataAcc += VolumeRGBT1; ValidSampleCount += 1.0f; }
-		if (PixelDistanceFromViewKm > (VolumeFrontDepth2.x - CloudFrontDepthTinyOffset) && VolumeFrontDepth2.x < (VolumeFrontDepth2.y)) { VolumeRGBT2 = ApplyFogToCloudFromUVsAndDepth(VolumeUVOffset2, VolumeFrontDepth2.x, VolumeRGBT2); DataAcc += VolumeRGBT2; ValidSampleCount += 1.0f; }
-		if (PixelDistanceFromViewKm > (VolumeFrontDepth3.x - CloudFrontDepthTinyOffset) && VolumeFrontDepth3.x < (VolumeFrontDepth3.y)) { VolumeRGBT3 = ApplyFogToCloudFromUVsAndDepth(VolumeUVOffset3, VolumeFrontDepth3.x, VolumeRGBT3); DataAcc += VolumeRGBT3; ValidSampleCount += 1.0f; }
+		if (PixelDistanceFromViewKm > (VolumeFrontDepth0.x - CloudFrontDepthTinyOffset) && VolumeFrontDepth0.x < (VolumeFrontDepth0.y)) { DepthForFog += VolumeFrontDepth0.x; DataAcc += VolumeRGBT0; ValidSampleCount += 1.0f; }
+		if (PixelDistanceFromViewKm > (VolumeFrontDepth1.x - CloudFrontDepthTinyOffset) && VolumeFrontDepth1.x < (VolumeFrontDepth1.y)) { DepthForFog += VolumeFrontDepth1.x; DataAcc += VolumeRGBT1; ValidSampleCount += 1.0f; }
+		if (PixelDistanceFromViewKm > (VolumeFrontDepth2.x - CloudFrontDepthTinyOffset) && VolumeFrontDepth2.x < (VolumeFrontDepth2.y)) { DepthForFog += VolumeFrontDepth2.x; DataAcc += VolumeRGBT2; ValidSampleCount += 1.0f; }
+		if (PixelDistanceFromViewKm > (VolumeFrontDepth3.x - CloudFrontDepthTinyOffset) && VolumeFrontDepth3.x < (VolumeFrontDepth3.y)) { DepthForFog += VolumeFrontDepth3.x; DataAcc += VolumeRGBT3; ValidSampleCount += 1.0f; }
 		if (ValidSampleCount > 0.0f)
 		{
+		#if APPLY_FOG_LATE
+			if (ValidSampleCount > 0.0f)
+			{
+				// Apply fog on a single sample cloud + depth result.
+				DataAcc /= ValidSampleCount;
+				DepthForFog /= ValidSampleCount;
+				ValidSampleCount = 1.0;
+				DataAcc = ApplyFogToCloudFromUVsAndDepth(WaterVolumeUV, DepthForFog, DataAcc);
+			}
+		#endif
 			OutputRt0 = DataAcc / ValidSampleCount;
 		}
 		else
@@ -936,13 +953,23 @@ void ComposeVolumetricRTOverScenePS(
 			VolumeFrontDepth3 = SafeLoadVolumetricDepthTexture(VolumeUVOffset3 * VolumetricTextureSizeAndInvSize.xy);
 
 			// We are testing if the depth buffer is further than the cloud front depth and that the cloud front depth is actually in front of traced depth.
-			if (PixelDistanceFromViewKm > (VolumeFrontDepth0.x - CloudFrontDepthTinyOffset) && VolumeFrontDepth0.x < (VolumeFrontDepth0.y)) { VolumeRGBT0 = ApplyFogToCloudFromUVsAndDepth(VolumeUVOffset0, VolumeFrontDepth0.x, VolumeRGBT0); DataAcc += VolumeRGBT0; ValidSampleCount += 1.0f; }
-			if (PixelDistanceFromViewKm > (VolumeFrontDepth1.x - CloudFrontDepthTinyOffset) && VolumeFrontDepth1.x < (VolumeFrontDepth1.y)) { VolumeRGBT1 = ApplyFogToCloudFromUVsAndDepth(VolumeUVOffset1, VolumeFrontDepth1.x, VolumeRGBT1); DataAcc += VolumeRGBT1; ValidSampleCount += 1.0f; }
-			if (PixelDistanceFromViewKm > (VolumeFrontDepth2.x - CloudFrontDepthTinyOffset) && VolumeFrontDepth2.x < (VolumeFrontDepth2.y)) { VolumeRGBT2 = ApplyFogToCloudFromUVsAndDepth(VolumeUVOffset2, VolumeFrontDepth2.x, VolumeRGBT2); DataAcc += VolumeRGBT2; ValidSampleCount += 1.0f; }
-			if (PixelDistanceFromViewKm > (VolumeFrontDepth3.x - CloudFrontDepthTinyOffset) && VolumeFrontDepth3.x < (VolumeFrontDepth3.y)) { VolumeRGBT3 = ApplyFogToCloudFromUVsAndDepth(VolumeUVOffset3, VolumeFrontDepth3.x, VolumeRGBT3); DataAcc += VolumeRGBT3; ValidSampleCount += 1.0f; }
+			if (PixelDistanceFromViewKm > (VolumeFrontDepth0.x - CloudFrontDepthTinyOffset) && VolumeFrontDepth0.x < (VolumeFrontDepth0.y)) { DepthForFog += VolumeFrontDepth0.x; DataAcc += VolumeRGBT0; ValidSampleCount += 1.0f; }
+			if (PixelDistanceFromViewKm > (VolumeFrontDepth1.x - CloudFrontDepthTinyOffset) && VolumeFrontDepth1.x < (VolumeFrontDepth1.y)) { DepthForFog += VolumeFrontDepth1.x; DataAcc += VolumeRGBT1; ValidSampleCount += 1.0f; }
+			if (PixelDistanceFromViewKm > (VolumeFrontDepth2.x - CloudFrontDepthTinyOffset) && VolumeFrontDepth2.x < (VolumeFrontDepth2.y)) { DepthForFog += VolumeFrontDepth2.x; DataAcc += VolumeRGBT2; ValidSampleCount += 1.0f; }
+			if (PixelDistanceFromViewKm > (VolumeFrontDepth3.x - CloudFrontDepthTinyOffset) && VolumeFrontDepth3.x < (VolumeFrontDepth3.y)) { DepthForFog += VolumeFrontDepth3.x; DataAcc += VolumeRGBT3; ValidSampleCount += 1.0f; }
 			if (ValidSampleCount > 0.0f)
 			{
-				OutputRt0 = DataAcc / ValidSampleCount;// +float4(10.0, 0.0, 0.0, 0.0);
+			#if APPLY_FOG_LATE
+				if (ValidSampleCount > 0.0f)
+				{
+					// Apply fog on a single sample cloud + depth result.
+					DataAcc /= ValidSampleCount;
+					DepthForFog /= ValidSampleCount;
+					ValidSampleCount = 1.0;
+					DataAcc = ApplyFogToCloudFromUVsAndDepth(WaterVolumeUV, DepthForFog, DataAcc);
+				}
+			#endif
+				OutputRt0 = DataAcc / ValidSampleCount;
 			}
 			else
 			{
@@ -1014,10 +1041,7 @@ void ComposeVolumetricRTOverScenePS(
 		float4 VolumeRGBT2 = SafeLoadVolumetricTexture(VolumeCoord2);
 		float4 VolumeRGBT3 = SafeLoadVolumetricTexture(VolumeCoord3);
 
-		VolumeRGBT0 = ApplyFogToCloudFromUVsAndDepth(VolumeCoord0 * VolumetricTextureSizeAndInvSize.zw, VolumeFrontDepth0.x, VolumeRGBT0);
-		VolumeRGBT1 = ApplyFogToCloudFromUVsAndDepth(VolumeCoord1 * VolumetricTextureSizeAndInvSize.zw, VolumeFrontDepth1.x, VolumeRGBT1);
-		VolumeRGBT2 = ApplyFogToCloudFromUVsAndDepth(VolumeCoord2 * VolumetricTextureSizeAndInvSize.zw, VolumeFrontDepth2.x, VolumeRGBT2);
-		VolumeRGBT3 = ApplyFogToCloudFromUVsAndDepth(VolumeCoord3 * VolumetricTextureSizeAndInvSize.zw, VolumeFrontDepth3.x, VolumeRGBT3);
+		float DepthForFog = 0.0f;
 
 		if (bIsWaterPixel)
 		{
@@ -1037,21 +1061,25 @@ void ComposeVolumetricRTOverScenePS(
 				if (VolumeFrontDepth0.y > (PixelDistanceFromViewKm - AcceptanceOffsetKm))
 				{
 					DataAcc += VolumeRGBT0;
+					DepthForFog += VolumeFrontDepth0.x;
 					ValidSampleCount++;
 				}
 				if (VolumeFrontDepth1.y > (PixelDistanceFromViewKm - AcceptanceOffsetKm))
 				{
 					DataAcc += VolumeRGBT1;
+					DepthForFog += VolumeFrontDepth1.x;
 					ValidSampleCount++;
 				}
 				if (VolumeFrontDepth2.y > (PixelDistanceFromViewKm - AcceptanceOffsetKm))
 				{
 					DataAcc += VolumeRGBT2;
+					DepthForFog += VolumeFrontDepth2.x;
 					ValidSampleCount++;
 				}
 				if (VolumeFrontDepth3.y > (PixelDistanceFromViewKm - AcceptanceOffsetKm))
 				{
 					DataAcc += VolumeRGBT3;
+					DepthForFog += VolumeFrontDepth3.x;
 					ValidSampleCount++;
 				}
 			}
@@ -1066,7 +1094,19 @@ void ComposeVolumetricRTOverScenePS(
 
 			ValidSampleCount = weightsSum > 0.0f ? 1.0 : 0.0f;
 			DataAcc = weights.x * VolumeRGBT0 + weights.y * VolumeRGBT1 + weights.z * VolumeRGBT2 + weights.w * VolumeRGBT3;
+			DepthForFog = weights.x * VolumeFrontDepth0.x + weights.y * VolumeFrontDepth1.x + weights.z * VolumeFrontDepth2.x + weights.w * VolumeFrontDepth3.x;
+		}
+
+	#if APPLY_FOG_LATE
+		if (ValidSampleCount > 0.0f)
+		{
+			// Apply fog on a single sample cloud + depth result.
+			DataAcc /= ValidSampleCount;
+			DepthForFog /= ValidSampleCount;
+			ValidSampleCount = 1.0;
+			DataAcc = ApplyFogToCloudFromUVsAndDepth(ScreenUV, DepthForFog, DataAcc);
 		}
+	#endif
 	}
 	else
 	{
diff --git a/Engine/Source/Runtime/Renderer/Private/VolumetricCloudRendering.cpp b/Engine/Source/Runtime/Renderer/Private/VolumetricCloudRendering.cpp
--- a/Engine/Source/Runtime/Renderer/Private/VolumetricCloudRendering.cpp
+++ b/Engine/Source/Runtime/Renderer/Private/VolumetricCloudRendering.cpp
@@ -82,9 +82,13 @@ static TAutoConsoleVariable<int32> CVarVolumetricCloudApplyFogOnAllPixel(
 	TEXT("When true, fog will be applied on cloud tracing buffer on all pixels instead of only where clouds coverage is >0. This is a tiny bit more expenssive but can remove visual discontinuities when cloud coverage is close to or flip to 0."),
 	ECVF_RenderThreadSafe);
 
+// This is recommended to be 1 by default so that volumetric fog is correctly applied when lumen is enabled and r.VolumetricRenderTarget.PreferAsyncCompute=1.
+// - Order of operations when PreferAsyncCompute=1 are (1) Cloud tracing in async compute pipe => (2) lumen update translucent lighting volume => (3) VFog upsampling (4) Cloud upsampling.
+// - But cloud tracing in (1) has bad volumetric since it is generated later, and it cannot be moved up because it depends on lumen generating the translucent lighting volume.
+// - So ApplyFogLate=1 allows to to have fog on clouds applied in (4). Which is fine for out single depth operation.
 static TAutoConsoleVariable<int32> CVarVolumetricCloudApplyFogLate(
 	TEXT("r.VolumetricCloud.ApplyFogLate"), 0,
-	TEXT("When true, fog will be applied on cloud later in the frame. This applyies fog on cloud without visible latency due to reprojection and reconstruction. It is however more expenssive."),
+	TEXT("When true, fog will be applied on cloud later in the frame. This applies fog on cloud without visible latency due to reprojection and reconstruction. It is a little bit more expenssive but the only correct look when Lumen is enabled and tracing is run on async compute using r.VolumetricRenderTarget.PreferAsyncCompute=1."),
 	ECVF_RenderThreadSafe);
 
 ////////////////////////////////////////////////////////////////////////// Shadow tracing
diff --git a/Engine/Source/Runtime/Renderer/Private/VolumetricRenderTarget.cpp b/Engine/Source/Runtime/Renderer/Private/VolumetricRenderTarget.cpp
--- a/Engine/Source/Runtime/Renderer/Private/VolumetricRenderTarget.cpp
+++ b/Engine/Source/Runtime/Renderer/Private/VolumetricRenderTarget.cpp
@@ -44,7 +44,7 @@ static TAutoConsoleVariable<float> CVarVolumetricRenderTargetScale(
 
 static TAutoConsoleVariable<int32> CVarVolumetricRenderTargetPreferAsyncCompute(
 	TEXT("r.VolumetricRenderTarget.PreferAsyncCompute"), 0,
-	TEXT("Whether to prefer using async compute to generate volumetric cloud render targets."),
+	TEXT("Whether to prefer using async compute to generate volumetric cloud render targets. When this is set to true, it is recommend to also use r.VolumetricCloud.ApplyFogLate=1 for correct volumetric fog lighting on clouds."),
 	ECVF_RenderThreadSafe | ECVF_Scalability);
 
 static TAutoConsoleVariable<int32> CVarVolumetricRenderTargetViewRect(
@@ -1037,7 +1037,7 @@ void ComposeVolumetricRenderTargetOverScene(
 		PermutationVector.Set<FComposeVolumetricRTOverScenePS::FRenderUnderWaterBuffer>(0);
 		PermutationVector.Set<FComposeVolumetricRTOverScenePS::FRenderCameraComposeWithWater>((bComposeWithWater && ViewInfo.IsUnderwater()) ? 1 : 0);
 		PermutationVector.Set<FComposeVolumetricRTOverScenePS::FMSAASampleCount>(SceneDepthTexture->Desc.NumSamples);
-		PermutationVector.Set<FComposeVolumetricRTOverScenePS::FApplyFog>(bShouldVolumetricCloudsApplyFogDuringReconstruction && ViewInfo.Family->Scene->HasAnyExponentialHeightFog() && ShouldRenderFog(*ViewInfo.Family));
+		PermutationVector.Set<FComposeVolumetricRTOverScenePS::FApplyFog>(bShouldVolumetricCloudsApplyFogDuringReconstruction && (ViewInfo.Family->Scene->HasAnyExponentialHeightFog() || ViewInfo.LocalFogVolumeViewData.GPUInstanceCount > 0) && ShouldRenderFog(*ViewInfo.Family));
 		PermutationVector.Set<FComposeVolumetricRTOverScenePS::FApplyLocalFogVolume>(bShouldVolumetricCloudsApplyFogDuringReconstruction && ViewInfo.LocalFogVolumeViewData.GPUInstanceCount > 0);
 		PermutationVector.Set<FComposeVolumetricRTOverScenePS::FMinMaxDepthAvailable>(bMinMaxDepth);
 		TShaderMapRef<FComposeVolumetricRTOverScenePS> PixelShader(ViewInfo.ShaderMap, PermutationVector);
@@ -1142,6 +1142,8 @@ void ComposeVolumetricRenderTargetOverSceneUnderWater(
 		int32 UpsamplingMode = GetCompositionUpsamplingMode(VRTMode, VolumetricCloudRT.GetUpsamplingMode());
 		const bool bMinMaxDepth = ShouldVolumetricCloudTraceWithMinMaxDepth(ViewInfo);
 
+		const bool bShouldVolumetricCloudsApplyFogDuringReconstruction = ShouldVolumetricCloudsApplyFogDuringReconstruction(ViewInfo);
+
 		FRDGTextureRef VolumetricTexture = nullptr;
 		FRDGTextureRef VolumetricSecondaryTexture = nullptr;
 		FRDGTextureRef VolumetricDepthTexture = nullptr;
@@ -1153,6 +1155,8 @@ void ComposeVolumetricRenderTargetOverSceneUnderWater(
 		PermutationVector.Set<FComposeVolumetricRTOverScenePS::FRenderUnderWaterBuffer>(1);
 		PermutationVector.Set<FComposeVolumetricRTOverScenePS::FRenderCameraComposeWithWater>(0);
 		PermutationVector.Set<FComposeVolumetricRTOverScenePS::FMSAASampleCount>(1);
+		PermutationVector.Set<FComposeVolumetricRTOverScenePS::FApplyFog>(bShouldVolumetricCloudsApplyFogDuringReconstruction && (ViewInfo.Family->Scene->HasAnyExponentialHeightFog() || ViewInfo.LocalFogVolumeViewData.GPUInstanceCount > 0) && ShouldRenderFog(*ViewInfo.Family));
+		PermutationVector.Set<FComposeVolumetricRTOverScenePS::FApplyLocalFogVolume>(bShouldVolumetricCloudsApplyFogDuringReconstruction && ViewInfo.LocalFogVolumeViewData.GPUInstanceCount > 0);
 		PermutationVector.Set<FComposeVolumetricRTOverScenePS::FMinMaxDepthAvailable>(bMinMaxDepth);
 		TShaderMapRef<FComposeVolumetricRTOverScenePS> PixelShader(ViewInfo.ShaderMap, PermutationVector);
 
@@ -1174,6 +1178,13 @@ void ComposeVolumetricRenderTargetOverSceneUnderWater(
 		PassParameters->Substrate = Substrate::BindSubstrateGlobalUniformParameters(ViewInfo);
 		PassParameters->OutputAlphaHoldout = 0;
 		PassParameters->ForwardShadingEnable = 0;
+
+		if (bShouldVolumetricCloudsApplyFogDuringReconstruction)
+		{
+			PassParameters->FogStruct = CreateFogUniformBuffer(GraphBuilder, ViewInfo);
+			PassParameters->LFV = ViewInfo.LocalFogVolumeViewData.UniformParametersStruct;
+		}
+
 		FVector2f DummyUVScale;
 		GetTextureSafeUvCoordBound(VolumetricViewRect, PassParameters->VolumetricTexture, PassParameters->VolumetricTextureValidCoordRect, PassParameters->VolumetricTextureValidUvRect, DummyUVScale);
 
