Substrate - removing non unroll loop which are just too register intensive.
It seems that DXC and other compiler part of (vulkan) have converged enough for that.
Substrate code has evolved and we should be able to understand better any new errrors coming from this.
And the mobile renderer use the special cast to legacy gbuffer now.

#tests PC editor dx12 and vulkan with ray tracing / path tracing on Substrate level. Similar performance level measured in deferred for dx12 while vulkan was faster.
#rnx

diff --git a/Engine/Shaders/Private/Substrate/Substrate.ush b/Engine/Shaders/Private/Substrate/Substrate.ush
--- a/Engine/Shaders/Private/Substrate/Substrate.ush
+++ b/Engine/Shaders/Private/Substrate/Substrate.ush
@@ -31,17 +31,10 @@
 #define SUBSTRATE_CLAMPED_CLOSURE_COUNT 1
 #endif
 
-// DXC has issue with unrolling certain part of the code. To address temporarly the warning-as-error issues, 
-// we redefine the unroll function and fall back to loop on DXC
-// To remove this, one must test PC + DX12 + RayTracing and make sure raytracing shaders compile.
-#define SUBSTRATE_COMPILER (COMPILER_VULKAN)
-#if SUBSTRATE_COMPILER == 1
-#define SUBSTRATE_UNROLL LOOP
-#define SUBSTRATE_UNROLL_N(X) LOOP
-#else
-#define SUBSTRATE_UNROLL UNROLL
-#define SUBSTRATE_UNROLL_N(X) UNROLL_N(X)
-#endif
+// Substrate should never use real loop, that can result in too high register allocation in some compiler.
+// Only a maximum unrolled loop according to the maximum slab/closures per pixel is preferable.
+#define SUBSTRATE_UNROLL		UNROLL
+#define SUBSTRATE_UNROLL_N(X)	UNROLL_N(X)
 
 #if SUBSTRATE_CLAMPED_CLOSURE_COUNT == 1
 	#define Substrate_for_unroll(X,Y,Z) X; if(Y)
