[SSS] Add mixed SSS implementation exposure for quality and performance.
1. Add a new field 'Implementation' in Subsurface profile under `Burley Normalized`, with 'AFIS' and `Separable`. Default to `AFIS`. When set to `Separable`, the corresponding material will directly use Separable. This allows mixed setting for best quality(skin) and performance(snow, ice).
2. Add half resolution support to AFIS.
3. Add a CVar `r.SSS.HalfRes.ForceSeparalbe 1` so that the old behavior still holds. Users will need to disable this CVar to have the expected mixed support for half resolution.

#tests test on metahuman and one internal project

diff --git a/Engine/Shaders/Private/PostProcessSubsurface.usf b/Engine/Shaders/Private/PostProcessSubsurface.usf
--- a/Engine/Shaders/Private/PostProcessSubsurface.usf
+++ b/Engine/Shaders/Private/PostProcessSubsurface.usf
@@ -752,6 +752,26 @@ void AddSubsurfaceComputeGroup(uint GI, uint Type, uint2 GroupXY)
 	}
 }
 
+// Get random index that has Mask[index] = 1 (s.t. Sum(Mask)>=1)
+uint GetRandomIndex(uint2 Pos, uint4 Mask)
+{
+	uint4 PrefixSum = uint4(
+		Mask[0],
+		Mask[0]+Mask[1],
+		Mask[0]+Mask[1]+Mask[2],
+		Mask[0]+Mask[1]+Mask[2]+Mask[3]);
+
+	uint RandomSelectedBit = (Rand3DPCG16(int3(Pos, View.FrameNumber)).x % PrefixSum[3]) + 1u;
+	
+	uint m0 = (RandomSelectedBit == PrefixSum[0]) & Mask[0];
+	uint m1 = (RandomSelectedBit == PrefixSum[1]) & Mask[1];
+	uint m2 = (RandomSelectedBit == PrefixSum[2]) & Mask[2];
+	uint m3 = (RandomSelectedBit == PrefixSum[3]) & Mask[3];
+
+	uint SelectedIndexWithValidMask = firstbitlow((m0 << 0u)|(m1 << 1u)|(m2 << 2u)|(m3 << 3u));
+	return min(SelectedIndexWithValidMask,3u);
+}
+
 RWTexture2D<float4> SetupTexture;
 RWTexture2D<float4> ProfileEdgeMask;
 
@@ -836,39 +856,52 @@ void SetupIndirectCS(
 	uint SelectedProfile = 0;
 	if (bHasSubsurface)
 	{
-#if SUBSURFACE_HALF_RES || SUBSURFACE_FORCE_SEPARABLE
-
-	#if 0
-		// the code here is designed for half resolution burley.
-		float2 UVA = min(BufferUV + float2(-0.5, 0.5f) * SubsurfaceInput0_ExtentInverse, SubsurfaceInput0_UVViewportBilinearMax);
-		uint PidA = GetSubsurfaceProfileId(UVA);
-		float2 UVB = min(BufferUV + float2(0.5, 0.5f) * SubsurfaceInput0_ExtentInverse, SubsurfaceInput0_UVViewportBilinearMax);
-		uint PidB = GetSubsurfaceProfileId(UVB);
-		float2 UVC = min(BufferUV + float2(0.5, -0.5f) * SubsurfaceInput0_ExtentInverse, SubsurfaceInput0_UVViewportBilinearMax);
-		uint PidC = GetSubsurfaceProfileId(UVC);
-		float2 UVD = min(BufferUV + float2(-0.5, -0.5f) * SubsurfaceInput0_ExtentInverse, SubsurfaceInput0_UVViewportBilinearMax);
-		uint PidD = GetSubsurfaceProfileId(UVD);
-
-		// Gatter available and use the last subsurface type for the current pixel.
-		SelectedProfile = lerp(SelectedProfile, PidA, A.a);
-		SelectedProfile = lerp(SelectedProfile, PidB, B.a);
-		SelectedProfile = lerp(SelectedProfile, PidC, C.a);
-		SelectedProfile = lerp(SelectedProfile, PidD, D.a);
-
-		// record the uv of this pixel
-		BufferUV = lerp(BufferUV, UVA, A.a);
-		BufferUV = lerp(BufferUV, UVB, B.a);
-		BufferUV = lerp(BufferUV, UVC, C.a);
-		BufferUV = lerp(BufferUV, UVD, D.a);
+#if SUBSURFACE_HALF_RES
+
+	#if !SUBSURFACE_FORCE_SEPARABLE
+
+        #if 1
+		{
+			float2 UVA = min(BufferUV + float2(-0.5, 0.5f) * SubsurfaceInput0_ExtentInverse, SubsurfaceInput0_UVViewportBilinearMax);
+			float2 UVB = min(BufferUV + float2(0.5, 0.5f) * SubsurfaceInput0_ExtentInverse, SubsurfaceInput0_UVViewportBilinearMax);
+			float2 UVC = min(BufferUV + float2(0.5, -0.5f) * SubsurfaceInput0_ExtentInverse, SubsurfaceInput0_UVViewportBilinearMax);
+			float2 UVD = min(BufferUV + float2(-0.5, -0.5f) * SubsurfaceInput0_ExtentInverse, SubsurfaceInput0_UVViewportBilinearMax);
+
+			BufferUV = lerp(BufferUV, UVA, A.a);
+			BufferUV = lerp(BufferUV, UVB, B.a);
+			BufferUV = lerp(BufferUV, UVC, C.a);
+			BufferUV = lerp(BufferUV, UVD, D.a);
+		}
+		#else
+		{
+			// Gatter available and randomize.
+			float4x2 UVs;
+			UVs[0] = min(BufferUV + float2(-0.5, 0.5f) * SubsurfaceInput0_ExtentInverse, SubsurfaceInput0_UVViewportBilinearMax);
+			UVs[1] = min(BufferUV + float2(0.5, 0.5f) * SubsurfaceInput0_ExtentInverse, SubsurfaceInput0_UVViewportBilinearMax);
+			UVs[2] = min(BufferUV + float2(0.5, -0.5f) * SubsurfaceInput0_ExtentInverse, SubsurfaceInput0_UVViewportBilinearMax);
+			UVs[3] = min(BufferUV + float2(-0.5, -0.5f) * SubsurfaceInput0_ExtentInverse, SubsurfaceInput0_UVViewportBilinearMax);
+
+			uint UVIndex = GetRandomIndex(Pos, uint4(A.a, B.a, C.a, D.a));
+			BufferUV = UVs[UVIndex];
+		}
+		#endif
 		
 		SelectedProfile = GetSubsurfaceProfileId(BufferUV);
 		Type = GetSubsurfaceType(SelectedProfile);
+#if ENABLE_PROFILE_ID_CACHE
+		ProfileId = SelectedProfile;
+#endif
 	#else
-		Type = TYPE_SEPARABLE; // use separable for half resolution
+		Type = TYPE_SEPARABLE;
 	#endif
 		
+#else
+#if ENABLE_PROFILE_ID_CACHE
+		SelectedProfile = ProfileId;
 #else
 		SelectedProfile = GetSubsurfaceProfileId(BufferUV);
+#endif
+		
 		Type = GetSubsurfaceType(SelectedProfile);
 #endif
 	}
diff --git a/Engine/Shaders/Private/SubsurfaceBurleyNormalized.ush b/Engine/Shaders/Private/SubsurfaceBurleyNormalized.ush
--- a/Engine/Shaders/Private/SubsurfaceBurleyNormalized.ush
+++ b/Engine/Shaders/Private/SubsurfaceBurleyNormalized.ush
@@ -18,8 +18,8 @@
 
 // Set to 1 to be more correct and higher quality. we estimate true distribution variance. It would be slower under some conditions
 // Set to 0 to be more efficient. That we use the previous sample count as the mean.
-// By default we set it to 0.
-#define USE_TRUE_DISTRIBUTION_VAR 0
+// By default we set it to 0. Enable it for half resolution.
+#define USE_TRUE_DISTRIBUTION_VAR (SUBSURFACE_HALF_RES)
 
 // Used to avoid low sampling count due to low variance
 #define BETA_LIMIT 8
diff --git a/Engine/Source/Runtime/Engine/Classes/Engine/SubsurfaceProfile.h b/Engine/Source/Runtime/Engine/Classes/Engine/SubsurfaceProfile.h
--- a/Engine/Source/Runtime/Engine/Classes/Engine/SubsurfaceProfile.h
+++ b/Engine/Source/Runtime/Engine/Classes/Engine/SubsurfaceProfile.h
@@ -9,6 +9,19 @@
 
 #define MAX_SUBSURFACE_PROFILE_COUNT 256
 
+UENUM()
+enum class ESubsurfaceImplementationTechniqueHint : uint8
+{
+	/** 
+	* Best for high quality like skin. Perform direct subsurface scattering with acceleration from adaptive filtered importance sampling (AFIS). It is expensive but more accurate. 
+	* Will fallback to Separable if lower quality (better performance) is preferred (r.SSS.Burley.Quality 0) or force separable in half resolution (r.SSS.HalfRes 1 and r.SSS.HalfRes.ForceSeparable 1) for performance.
+	*/
+	SIH_AFIS UMETA(DisplayName = "Adaptive Filtered Importance Sampling"),
+
+	/** Best for high performance if large ice/snow covering. Perform subsurface scattering with vertial and horizontal separable filters.*/
+	SIH_Separable UMETA(DisplayName = "Separable"),
+};
+
 // struct with all the settings we want in USubsurfaceProfile, separate to make it easer to pass this data around in the engine.
 USTRUCT(BlueprintType)
 struct FSubsurfaceProfileStruct
@@ -82,6 +95,12 @@ struct FSubsurfaceProfileStruct
 	UPROPERTY(Category = "Burley Normalized", EditAnywhere, BlueprintReadOnly, meta = (HideAlphaChannel))
 	FLinearColor BoundaryColorBleed;
 
+	/**
+	* This allows users to use mixed implementations for best quality and performance (e.g., High quality for skin with AFIS, and high performance on ice with Separable in the same scene).
+	*/
+	UPROPERTY(Category = "Burley Normalized", EditAnywhere, BlueprintReadOnly)
+	ESubsurfaceImplementationTechniqueHint Implementation;
+
 	UPROPERTY(Category = "Transmission", EditAnywhere, BlueprintReadOnly, meta = (ClampMin = "0.01", UIMax = "100.0", ClampMax = "100.0", SliderExponent="24"))
 	float ExtinctionScale;
 	
@@ -132,7 +151,7 @@ struct FSubsurfaceProfileStruct
 		WorldUnitScale = 0.1f;
 		TransmissionTintColor = FLinearColor(1.0f, 1.0f, 1.0f);
 		Tint = FLinearColor(1.0f, 1.0f, 1.0f);
-
+		Implementation = ESubsurfaceImplementationTechniqueHint::SIH_AFIS;
 		bEnableMeanFreePath = false;
 	}
 
diff --git a/Engine/Source/Runtime/Engine/Private/Rendering/SubsurfaceProfile.cpp b/Engine/Source/Runtime/Engine/Private/Rendering/SubsurfaceProfile.cpp
--- a/Engine/Source/Runtime/Engine/Private/Rendering/SubsurfaceProfile.cpp
+++ b/Engine/Source/Runtime/Engine/Private/Rendering/SubsurfaceProfile.cpp
@@ -509,7 +509,8 @@ void FSubsurfaceProfileTexture::CreateTexture(FRHICommandListImmediate& RHICmdLi
 
 		TextureRow[SSSS_BOUNDARY_COLOR_BLEED_OFFSET] = Data.BoundaryColorBleed;
 
-		TextureRow[SSSS_BOUNDARY_COLOR_BLEED_OFFSET].A = Data.bEnableBurley ? SSS_TYPE_BURLEY : SSS_TYPE_SSSS;
+		TextureRow[SSSS_BOUNDARY_COLOR_BLEED_OFFSET].A = 
+			(Data.Implementation == ESubsurfaceImplementationTechniqueHint::SIH_AFIS) ? SSS_TYPE_BURLEY : SSS_TYPE_SSSS;
 
 		float MaterialRoughnessToAverage = Data.Roughness0 * (1.0f - Data.LobeMix) + Data.Roughness1 * Data.LobeMix;
 
@@ -524,7 +525,7 @@ void FSubsurfaceProfileTexture::CreateTexture(FRHICommandListImmediate& RHICmdLi
 		TextureRow[SSSS_TRANSMISSION_OFFSET].B = EncodeScatteringDistribution(Data.ScatteringDistribution);
 		TextureRow[SSSS_TRANSMISSION_OFFSET].A = 1.0f / Data.IOR;
 
-		if (Data.bEnableBurley)
+		if (true)
 		{
 			Data.ScatterRadius = FMath::Max(DifffuseMeanFreePathInMm.GetMax()*MmToCm, 0.1f);
 
@@ -648,7 +649,7 @@ bool FSubsurfaceProfileTexture::GetEntryString(uint32 Index, FString& Out) const
 
 
 	Out = FString::Printf(TEXT(" %c. %p SurfaceAlbedo=%.1f %.1f %.1f, MeanFreePathColor=%.1f %.1f %.1f, MeanFreePathDistance=%.1f, WorldUnitScale=%.1f,\
-								Tint=%.1f %.1f %.1f, BoundaryColorBleeding=%.1f %.1f %.1f, TransmissionTintColor=%.1f %.1f %.1f, EnableBurley=%.1f"), 
+								Tint=%.1f %.1f %.1f, BoundaryColorBleeding=%.1f %.1f %.1f, TransmissionTintColor=%.1f %.1f %.1f, %s"), 
 		MiniFontCharFromIndex(Index), 
 		SubsurfaceProfileEntries[Index].Profile,
 		ref.SurfaceAlbedo.R, ref.SurfaceAlbedo.G, ref.SurfaceAlbedo.B,
@@ -658,7 +659,7 @@ bool FSubsurfaceProfileTexture::GetEntryString(uint32 Index, FString& Out) const
 		ref.Tint.R, ref.Tint.G, ref.Tint.B,
 		ref.BoundaryColorBleed.R, ref.BoundaryColorBleed.G, ref.BoundaryColorBleed.B, 
 		ref.TransmissionTintColor.R, ref.TransmissionTintColor.G, ref.TransmissionTintColor.B,
-		ref.bEnableBurley?1.0f:0.0f);
+		ref.Implementation == ESubsurfaceImplementationTechniqueHint::SIH_AFIS?TEXT("AFIS") : TEXT("Separable"));
 
 	return true;
 }
@@ -702,8 +703,8 @@ void FSubsurfaceProfileTexture::Dump()
 				SubsurfaceProfileEntries[i].Settings.BoundaryColorBleed.R, SubsurfaceProfileEntries[i].Settings.BoundaryColorBleed.G, SubsurfaceProfileEntries[i].Settings.BoundaryColorBleed.B);
 			UE_LOG(LogSubsurfaceProfile, Log, TEXT("     TransmissionTintColor=%f %f %f"),
 				SubsurfaceProfileEntries[i].Settings.TransmissionTintColor.R, SubsurfaceProfileEntries[i].Settings.TransmissionTintColor.G, SubsurfaceProfileEntries[i].Settings.TransmissionTintColor.B);
-			UE_LOG(LogSubsurfaceProfile, Log, TEXT("     EnableBurley=%.1f"),
-				SubsurfaceProfileEntries[i].Settings.bEnableBurley?1.0f:0.0f);
+			UE_LOG(LogSubsurfaceProfile, Log, TEXT("     %s"),
+				SubsurfaceProfileEntries[i].Settings.Implementation == ESubsurfaceImplementationTechniqueHint::SIH_AFIS?TEXT("AFIS") : TEXT("Separable"));
 		}
 	}
 
diff --git a/Engine/Source/Runtime/Renderer/Private/PostProcess/PostProcessSubsurface.cpp b/Engine/Source/Runtime/Renderer/Private/PostProcess/PostProcessSubsurface.cpp
--- a/Engine/Source/Runtime/Renderer/Private/PostProcess/PostProcessSubsurface.cpp
+++ b/Engine/Source/Runtime/Renderer/Private/PostProcess/PostProcessSubsurface.cpp
@@ -52,6 +52,13 @@ namespace
 		TEXT(" 1: parts of the algorithm runs in half resolution which is lower quality but faster (default, Separable only)"),
 		ECVF_RenderThreadSafe | ECVF_Scalability);
 
+	TAutoConsoleVariable<int32> CVarSSSHalfResForceSeparable(
+		TEXT("r.SSS.HalfRes.ForceSeparable"),
+		1,
+		TEXT(" 0: Use the implementation hint from the subsurface profile.\n")
+		TEXT(" 1: Force to use separable filter in HalfRes mode (default)."),
+		ECVF_RenderThreadSafe | ECVF_Scalability);
+
 	TAutoConsoleVariable<int32> CVarSSSQuality(
 		TEXT("r.SSS.Quality"),
 		0,
@@ -1012,7 +1019,8 @@ void AddSubsurfaceViewPass(
 	const uint32 ScaleFactor = bHalfRes ? 2 : 1;
 	
 	//We run in separable mode under two conditions: 1) Run Burley fallback mode. 2) when the screen is in half resolution.
-	const bool bForceRunningInSeparable = CVarSSSBurleyQuality.GetValueOnRenderThread() == 0 || bHalfRes;
+	const bool bForceRunningInSeparable = CVarSSSBurleyQuality.GetValueOnRenderThread() == 0 || 
+		 (bHalfRes && CVarSSSHalfResForceSeparable.GetValueOnRenderThread() != 0 );
 
 	const bool bUseProfileIdCache = IsProfileIdCacheEnabled();
 
@@ -1074,7 +1082,7 @@ void AddSubsurfaceViewPass(
 	TRefCountPtr<IPooledRenderTarget>* QualityHistoryState = ViewState ? &ViewState->SubsurfaceScatteringQualityHistoryRT : NULL;
 
 	//allocate/reallocate the quality history texture. 
-	FRDGTextureRef QualityHistoryTexture = RegisterExternalRenderTarget(GraphBuilder, QualityHistoryState, SceneColorTextureDescriptor.Extent, TEXT("Subsurface.QualityHistoryTexture"));
+	FRDGTextureRef QualityHistoryTexture = RegisterExternalRenderTarget(GraphBuilder, QualityHistoryState, SubsurfaceTextureDescriptor.Extent, TEXT("Subsurface.QualityHistoryTexture"));
 	FRDGTextureRef NewQualityHistoryTexture = nullptr;
 
 	RDG_EVENT_SCOPE_STAT(GraphBuilder, SubsurfaceScattering, "SubsurfaceScattering");
@@ -1219,10 +1227,10 @@ void AddSubsurfaceViewPass(
 					GraphBuilder,
 					SubsurfaceTextureDescriptor,
 					TEXT("Subsurface.SubpassTwoTex"),
-					RDG_EVENT_NAME("SSS::ClearUAV(%s, %s.TileCount > 0)", TEXT("Subsurface.SubpassTwoTex"), ToString(FSubsurfaceTiles::ETileType::SEPARABLE)),
+					RDG_EVENT_NAME("SSS::ClearUAV(%s, %s.TileCount > 0)", TEXT("Subsurface.SubpassTwoTex"), ToString(FSubsurfaceTiles::ETileType::All)),
 					SubsurfaceViewport,
 					Tiles.TileTypeCountBuffer,
-					Offset);
+					ToIndex(FSubsurfaceTiles::ETileType::All));
 			}
 			else
 			{
