Substrate - fix compilation problem with new SSS code.

#rnx
#tests PC editor dx12, switch preview, sss profile material

diff --git a/Engine/Shaders/Private/PostProcessSubsurfaceTile.usf b/Engine/Shaders/Private/PostProcessSubsurfaceTile.usf
--- a/Engine/Shaders/Private/PostProcessSubsurfaceTile.usf
+++ b/Engine/Shaders/Private/PostProcessSubsurfaceTile.usf
@@ -137,13 +137,24 @@ RWStructuredBuffer<uint> TileMaskBufferOut;
 RWStructuredBuffer<uint> TileMaskPassthroughBufferOut;
 RWStructuredBuffer<uint> TileMaskNonPassthroughBufferOut;
 
-//Duplicate from PostProcessSubsurface.usf
 #if SUBTRATE_GBUFFER_FORMAT==1
+//Duplicate from PostProcessSubsurface.usf
+FSubstrateSubsurfaceData LoadSubstrateSSSData(float2 UV)
+{
+	const float2 PixelPos = UV.xy * View.BufferSizeAndInvSize.xy;
+	return SubstrateLoadSubsurfaceData(Substrate.MaterialTextureArray, Substrate.FirstSliceStoringSubstrateSSSData, PixelPos);
+}
 FSubstrateSubsurfaceHeader LoadSubstrateSSSHeader(float2 UV)
 {
 	const float2 PixelPos = UV.xy * View.BufferSizeAndInvSize.xy;
 	return SubstrateLoadSubsurfaceHeader(Substrate.MaterialTextureArray, Substrate.FirstSliceStoringSubstrateSSSData, PixelPos);
 }
+//Duplicate from SubsurfaceBurleyNormalized.ush
+FSubstrateTopLayerData LoadSubstrateTopLayerData(float2 UV)
+{
+	const float2 PixelPos = UV.xy * View.BufferSizeAndInvSize.xy;
+	return SubstrateUnpackTopLayerData(Substrate.TopLayerTexture.Load(uint3(PixelPos, 0)));
+}
 #endif
 
 // Only access shading model GBuffer
