Add (optional) support for local light shadow on forward lit surfaces.

diff --git a/Engine/Shaders/Private/ForwardLightingCommon.ush b/Engine/Shaders/Private/ForwardLightingCommon.ush
--- a/Engine/Shaders/Private/ForwardLightingCommon.ush
+++ b/Engine/Shaders/Private/ForwardLightingCommon.ush
@@ -296,6 +296,45 @@ FDeferredLightingSplit GetForwardDirectLightingSplit(
 
 		float DynamicShadowing = dot(PreviewShadowMapChannelMask, DynamicShadowFactors);
 		float4 LightAttenuation = float4(1, 1, DynamicShadowing.x, DynamicShadowing.x);
+
+		#if VIRTUAL_SHADOW_MAP && SUPPORT_LOCAL_LIGHT_SHADOW_ON_FORWARD_LIT_TRANSLUCENT
+		BRANCH
+		if (LocalLight.Internal.VirtualShadowMapId != INDEX_NONE)
+		{
+			float DynamicShadowFactor = 1.f;
+			LightData.ShadowedBits = 1;
+			#if SUPPORT_VSM_FOWARD_QUALITY == 1
+			{
+				const FLightShaderParameters LocalLightShaderParameters = ConvertToLightShaderParameters(LightData);
+
+				const float ScreenRayLengthWorld = VirtualShadowMap.ScreenRayLength * View.ClipToView[1][1] * GBufferData.Depth;
+				float SMRTRayOffset = ScreenRayLengthWorld;
+
+				FVirtualShadowMapSampleResult VirtualShadowMapSample = TraceLocalLight(
+					LocalLight.VirtualShadowMapId,
+					LocalLightShaderParameters,
+					PixelPos,
+					GBufferData.SceneDepth,
+					TranslatedWorldPosition,
+					SMRTRayOffset,
+					Dither,
+					GBufferData.WorldNormal);
+
+				FilterVirtualShadowMapSampleResult(PixelPos, VirtualShadowMapSample);
+
+				DynamicShadowFactor *= VirtualShadowMapSample.ShadowFactor;
+			}
+			#else
+			{
+				FVirtualShadowMapSampleResult VirtualShadowMapSample = SampleVirtualShadowMapLocal(LocalLight.Internal.VirtualShadowMapId, TranslatedWorldPosition);
+				DynamicShadowFactor *= VirtualShadowMapSample.ShadowFactor;
+			}
+			#endif
+
+			LightAttenuation = DynamicShadowFactor;
+		}
+		#endif
+
 		float SurfaceShadow = 1.0f;
 		FDeferredLightingSplit NewLighting = GetDynamicLightingSplit(TranslatedWorldPosition, -CameraVector, GBufferData, 1, LightData, LightAttenuation, Dither, uint2(0,0), SurfaceShadow);
 
diff --git a/Engine/Source/Runtime/Engine/Classes/Engine/RendererSettings.h b/Engine/Source/Runtime/Engine/Classes/Engine/RendererSettings.h
--- a/Engine/Source/Runtime/Engine/Classes/Engine/RendererSettings.h
+++ b/Engine/Source/Runtime/Engine/Classes/Engine/RendererSettings.h
@@ -1083,6 +1083,15 @@ class URendererSettings : public UDeveloperSettings
 		ConfigRestartRequired = true))
 		uint32 bSupportRectLightOnTranslucent : 1;
 
+	/**
+	"Enable local light shadow evaluation on translucent materials when using the Forward Shading mode."
+	*/
+	UPROPERTY(config, EditAnywhere, Category = Optimizations, meta = (
+		ConsoleVariable = "r.Translucent.UsesShadowedLocalLights", DisplayName = "Support Local Light Shadow On Translucent Materials (When Using ForwardShading)",
+		ToolTip = "Enable local light shadow evaluation on translucent materials when using the Forward Shading mode.",
+		ConfigRestartRequired = true))
+		uint32 bSupportShadowedLocalLightOnTranslucent : 1;
+
 	UPROPERTY(config, EditAnywhere, Category = Debugging, meta = (
 		ConsoleVariable = "r.GPUCrashDebugging", DisplayName = "Enable vendor specific and in-engine GPU crash analysis tools.",
 		ToolTip = "Enables vendor specific and in-engine GPU crash analysis tools.",
diff --git a/Engine/Source/Runtime/Engine/Private/ShaderCompiler/ShaderCompiler.cpp b/Engine/Source/Runtime/Engine/Private/ShaderCompiler/ShaderCompiler.cpp
--- a/Engine/Source/Runtime/Engine/Private/ShaderCompiler/ShaderCompiler.cpp
+++ b/Engine/Source/Runtime/Engine/Private/ShaderCompiler/ShaderCompiler.cpp
@@ -3841,6 +3841,11 @@ void GlobalBeginCompileShader(
 		SET_SHADER_DEFINE(Input.Environment, SUPPORT_RECTLIGHT_ON_FORWARD_LIT_TRANSLUCENT, bTranslucentUsesLightRectLights ? 1 : 0);
 	}
 
+	{
+		const bool bTranslucentUsesShadowedLocalLights = GetTranslucentUsesShadowedLocalLights();
+		SET_SHADER_DEFINE(Input.Environment, SUPPORT_SHADOWED_LOCAL_LIGHT_ON_FORWARD_LIT_TRANSLUCENT, bTranslucentUsesShadowedLocalLights ? 1 : 0);
+	}
+
 	{
 		const bool bTranslucentUsesLightIESProfiles = GetTranslucentUsesLightIESProfiles();
 		SET_SHADER_DEFINE(Input.Environment, SUPPORT_IESPROFILE_ON_FORWARD_LIT_TRANSLUCENT, bTranslucentUsesLightIESProfiles ? 1 : 0);
diff --git a/Engine/Source/Runtime/RenderCore/Private/RenderUtils.cpp b/Engine/Source/Runtime/RenderCore/Private/RenderUtils.cpp
--- a/Engine/Source/Runtime/RenderCore/Private/RenderUtils.cpp
+++ b/Engine/Source/Runtime/RenderCore/Private/RenderUtils.cpp
@@ -1767,7 +1767,13 @@ static TAutoConsoleVariable<int32> CVarTranslucentUsesLightFunctionAtlas(
 static TAutoConsoleVariable<int32> CVarRectLighTranslucent(
 	TEXT("r.Translucent.UsesRectLights"),
 	0,
-	TEXT("Enable rect light support for translucent surfaces. When enabled, it will add an extrat sampler to the pixel shader (limited to 16 on dx11 based system)"),
+	TEXT("Enable rect light support for translucent surfaces. When enabled, it will add an extra sampler to the pixel shader (limited to 16 on dx11 based system)"),
+	ECVF_ReadOnly | ECVF_RenderThreadSafe);
+
+static TAutoConsoleVariable<int32> CVarShadowedLocalLightsTranslucent(
+	TEXT("r.Translucent.UsesShadowedLocalLights"),
+	0,
+	TEXT("Enable shadow support on local lights for translucent surfaces."),
 	ECVF_ReadOnly | ECVF_RenderThreadSafe);
 
 static TAutoConsoleVariable<int32> CVarIESProfileTranslucent(
@@ -1808,6 +1814,11 @@ bool GetTranslucentUsesLightRectLights()
 	return CVarRectLighTranslucent.GetValueOnAnyThread() > 0;
 }
 
+bool GetTranslucentUsesShadowedLocalLights()
+{
+	return CVarShadowedLocalLightsTranslucent.GetValueOnAnyThread() > 0;
+}
+
 bool GetTranslucentUsesLightIESProfiles()
 {
 	return CVarIESProfileTranslucent.GetValueOnAnyThread() > 0;
diff --git a/Engine/Source/Runtime/RenderCore/Private/Shader.cpp b/Engine/Source/Runtime/RenderCore/Private/Shader.cpp
--- a/Engine/Source/Runtime/RenderCore/Private/Shader.cpp
+++ b/Engine/Source/Runtime/RenderCore/Private/Shader.cpp
@@ -2396,6 +2396,15 @@ void ShaderMapAppendKey(EShaderPlatform Platform, FShaderKeyGenerator& KeyGen)
 		}
 	}
 
+	{
+		const bool bTranslucentUsesShadowedLocalLights = GetTranslucentUsesShadowedLocalLights();
+		if (bTranslucentUsesShadowedLocalLights)
+		{
+			KeyGen.AppendSeparator();
+			KeyGen.Append(TEXT("SHALOLIT"));
+		}
+	}
+
 	{
 		const bool bTranslucentUsesLightIESProfiles = GetTranslucentUsesLightIESProfiles();
 		if (bTranslucentUsesLightIESProfiles)
diff --git a/Engine/Source/Runtime/RenderCore/Public/RenderUtils.h b/Engine/Source/Runtime/RenderCore/Public/RenderUtils.h
--- a/Engine/Source/Runtime/RenderCore/Public/RenderUtils.h
+++ b/Engine/Source/Runtime/RenderCore/Public/RenderUtils.h
@@ -630,6 +630,7 @@ RENDERCORE_API bool GetSingleLayerWaterUsesLightFunctionAtlas();
 RENDERCORE_API bool GetTranslucentUsesLightFunctionAtlas();
 RENDERCORE_API bool GetTranslucentUsesLightRectLights();
 RENDERCORE_API bool GetTranslucentUsesLightIESProfiles();
+RENDERCORE_API bool GetTranslucentUsesShadowedLocalLights();
 
 RENDERCORE_API bool GetHairStrandsUsesTriangleStrips();
 RENDERCORE_API uint32 GetHairStrandsLODMode();
