Enable Susbtrate VRS on the shader side.

diff --git a/Engine/Shaders/Private/BasePassPixelShader.usf b/Engine/Shaders/Private/BasePassPixelShader.usf
--- a/Engine/Shaders/Private/BasePassPixelShader.usf
+++ b/Engine/Shaders/Private/BasePassPixelShader.usf
@@ -86,8 +86,7 @@
 #endif
 
 #if SUBSTRATE_ENABLED && COMPUTE_SHADED
-// To enable proper  data propagation when VRS is enable, uncomment this.
-// #define SUBSTRATE_VRS_WRITE 1
+#define SUBSTRATE_VRS_WRITE 1
 #endif 
 
 // For the base pass SUBSTRATE_INLINE_SHADING is defined in BasePassRendering.h
