[HWRT] Fix OutputVertexIndex not taking MinVertexIndex into account after recent changes.

#tests engine tests
#rnx

diff --git a/Engine/Shaders/Private/RayTracing/RayTracingDynamicMesh.usf b/Engine/Shaders/Private/RayTracing/RayTracingDynamicMesh.usf
--- a/Engine/Shaders/Private/RayTracing/RayTracingDynamicMesh.usf
+++ b/Engine/Shaders/Private/RayTracing/RayTracingDynamicMesh.usf
@@ -34,7 +34,7 @@ void HideTriangle(uint TriangleIndex, float3 ReferenceLocalPositon)
 {
 	for (int VertexIndex = 0; VertexIndex < 3; VertexIndex++)
 	{		
-		const uint OutputVertexIndex = OutputVertexBaseIndex + (TriangleIndex * 3 + VertexIndex) * 3;
+		const uint OutputVertexIndex = OutputVertexBaseIndex + (MinVertexIndex + TriangleIndex * 3 + VertexIndex) * 3;
 
 		RWVertexPositions[OutputVertexIndex + 0] = ReferenceLocalPositon.x;
 		RWVertexPositions[OutputVertexIndex + 1] = ReferenceLocalPositon.y;
@@ -66,10 +66,10 @@ void RayTracingDynamicGeometryConverterCS(uint3 GroupId : SV_GroupID, int GroupT
 	for(int Index = 0; Index < 3; Index++)
 	{
 		const uint VertexIndex = Indices[Index];
-		const uint OutputVertexIndex = OutputVertexBaseIndex + (LinearThreadId * 3 + Index) * 3;
+		const uint OutputVertexIndex = OutputVertexBaseIndex + (MinVertexIndex + LinearThreadId * 3 + Index) * 3;
 #else
 		const uint VertexIndex = LinearThreadId;
-		const uint OutputVertexIndex = OutputVertexBaseIndex + VertexIndex * 3;
+		const uint OutputVertexIndex = OutputVertexBaseIndex + (MinVertexIndex + VertexIndex) * 3;
 #endif // PER_VERTEX_MASK_ENABLED
 
 		if (VertexIndex < NumActualVertices)
