Fix InverseExposureLerp NaNs when eye adaption is disabled

#tests Used PLATFORM_ASSERT to verify log(0) occurs in InverseExposureLerp with EyeAdaptationInverse material node in Editor debug views and that this fix prevents that code path.

diff --git a/Engine/Shaders/Private/EyeAdaptationCommon.ush b/Engine/Shaders/Private/EyeAdaptationCommon.ush
--- a/Engine/Shaders/Private/EyeAdaptationCommon.ush
+++ b/Engine/Shaders/Private/EyeAdaptationCommon.ush
@@ -53,6 +53,9 @@ float EyeAdaptationLookup()
 
 float InverseExposureLerp(float Exposure, float Alpha)
 {
+#if EYE_ADAPTATION_DISABLED
+	return 1.0f;
+#else
 	// When Alpha=0.0, we want to multiply by 1.0. when Alpha = 1.0, we want to multiply by 1/Exposure.
 	// So the lerped value is:
 	//     LerpLogScale = Lerp(log(1),log(1/Exposure),T)
@@ -63,6 +66,7 @@ float InverseExposureLerp(float Exposure, float Alpha)
 	float LerpLogScale = -Alpha * log(Exposure);
 	float Scale = exp(LerpLogScale);
 	return Scale;
+#endif
 }
 
 float EyeAdaptationInverseLookup(float Alpha)
