Workaround a shader compiler bug. Some shader compilers require structs stored in groupshared to be 16-byte aligned.

diff --git a/Engine/Shaders/Private/ReflectionEnvironmentShaders.usf b/Engine/Shaders/Private/ReflectionEnvironmentShaders.usf
--- a/Engine/Shaders/Private/ReflectionEnvironmentShaders.usf
+++ b/Engine/Shaders/Private/ReflectionEnvironmentShaders.usf
@@ -7,6 +7,12 @@
 #include "Common.ush"
 #include "DeferredShadingCommon.ush"
 #include "BRDF.ush"
+
+#if defined(DXC_GROUPSHARED_ALIGNMENT_WORKAROUND) && DXC_GROUPSHARED_ALIGNMENT_WORKAROUND
+// Workaround for shader compiler bug affecting certain platforms
+// Align groupshared SH data if the shader compiler requires it
+#define NEED_SH_VECTOR_PADDING 1
+#endif
 #include "SHCommon.ush"
 #include "ReflectionEnvironmentShared.ush"
 #include "MonteCarlo.ush"
diff --git a/Engine/Shaders/Private/SHCommon.ush b/Engine/Shaders/Private/SHCommon.ush
--- a/Engine/Shaders/Private/SHCommon.ush
+++ b/Engine/Shaders/Private/SHCommon.ush
@@ -40,6 +40,10 @@ struct FThreeBandSHVector
 	half4 V0;
 	half4 V1;
 	half V2;
+#if defined(NEED_SH_VECTOR_PADDING) && NEED_SH_VECTOR_PADDING==1
+	// Add padding to workaround a bug with some shader compilers when storing structs in LDS
+	half3 Pad;
+#endif
 };
 
 struct FThreeBandSHVectorRGB
