[SSS] Fix Separable artifact if there are high frequency non SSS around a SSS pixel.

#tests EngineTest and Internal project.

diff --git a/Engine/Shaders/Private/SeparableSSS.ush b/Engine/Shaders/Private/SeparableSSS.ush
--- a/Engine/Shaders/Private/SeparableSSS.ush
+++ b/Engine/Shaders/Private/SeparableSSS.ush
@@ -339,9 +339,8 @@ float4 SSSSBlurPS(uint2 BufferPos, float2 BufferUV, float2 dir, bool initStencil
 			color.a *= 1 - s;
 #endif
 
-#if SUBSURFACE_HALF_RES
 			color.a *= (LocalSubsurfaceProfileInt != SSS_PROFILE_ID_INVALID) ? 1 : 0;
-#endif
+
 			// approximation, ideally we would reconstruct the mask with GetMaskFromDepthInAlpha() and do manual bilinear filter
 			// needed?
 			color.rgb *= color.a * ColorTint;
@@ -359,11 +358,10 @@ float4 SSSSBlurPS(uint2 BufferPos, float2 BufferUV, float2 dir, bool initStencil
 	// done for each color channel to avoid color shift
 	float3 OutColor = colorAccum / colorInvDiv; 
 
-#if SUBSURFACE_HALF_RES
-	// For half resolution, subsurface is dark if profile id is invalid, which will contribute dark halo to thin geometries surrounded by SSS.
+	// Subsurface is dark if profile id is invalid all around it, which will contribute dark halo to thin geometries surrounded by SSS.
 	// Fallback to the pixel color at the center.
 	OutColor = select(colorInvDiv == initColorInvDiv, colorM.rgb, OutColor);
-#endif
+
 	// alpha stored the SceneDepth (0 if there is no subsurface scattering)
     return float4(OutColor, OutDepth);
 }
