Re-enabling fog volumes on Android:
- Removing the config rule that disables it for GLES and Vulkan
- Removing the device profile config that disables it for the base Android profile.
-  Adding a fix per [REDACTED_HANDLE] for a misc code compilation issue in the local fog volume shader.

#tests Smoke tests and verification on each platform via full standard build and test suite.

#rnx

diff --git a/Engine/Shaders/Private/LocalFogVolumes/LocalFogVolumeTiledCulling.usf b/Engine/Shaders/Private/LocalFogVolumes/LocalFogVolumeTiledCulling.usf
--- a/Engine/Shaders/Private/LocalFogVolumes/LocalFogVolumeTiledCulling.usf
+++ b/Engine/Shaders/Private/LocalFogVolumes/LocalFogVolumeTiledCulling.usf
@@ -142,6 +142,7 @@ void LocalFogVolumeTiledCullingCS(uint3 DispatchThreadId : SV_DispatchThreadID)
 			LFVCount++;
 		}
 	}
+	GroupMemoryBarrier();
 
 	LocalFogVolumeTileDataTextureUAV[uint3(TileCoord, 0)] = LFVCount;
 
