Shader workaround for Adreno 7xx issue with a Local Fog Volumes contribution

diff --git a/Engine/Shaders/Private/LocalFogVolumes/LocalFogVolumeCommon.ush b/Engine/Shaders/Private/LocalFogVolumes/LocalFogVolumeCommon.ush
--- a/Engine/Shaders/Private/LocalFogVolumes/LocalFogVolumeCommon.ush
+++ b/Engine/Shaders/Private/LocalFogVolumes/LocalFogVolumeCommon.ush
@@ -400,22 +400,25 @@ float4 GetLFVContribution(
 	const float3 CamRayWorldDir				= normalize(DepthBufferTranslatedWorldPos - CamRayTranslatedWorldOrigin);
 
 	OutLFVCount = LFVTileDataTex[uint3(TilePos, 0)];
-	const float4 DefaultColorTransmittance = float4(0.0, 0.0, 0.0, 1.0);
-	float4 LFVContribution = DefaultColorTransmittance;
+	
+	float3 LFVLuminanceColor = 0.0;
+	float LFVTransmittance = 1.0;
 	for (int LFVIndex = 0; LFVIndex < OutLFVCount; ++LFVIndex)
 	{
 		uint InstanceIndex = LFVTileDataTex[uint3(TilePos, 1 + LFVIndex)];
 
-		float4 LFVLuminanceColorTransmittance = GetLFVInstanceContribution(
+		float4 InstanceContribution = GetLFVInstanceContribution(
 			InstanceIndex,
 			CamRayTranslatedWorldOrigin,
 			CamRayWorldDir,
 			DepthBufferTranslatedWorldPos,
 			bPlatformSupportsVolumetricFogOntranslucent);
 
-		LFVContribution = float4(LFVContribution.rgb * LFVLuminanceColorTransmittance.a + LFVLuminanceColorTransmittance.rgb, LFVContribution.a * LFVLuminanceColorTransmittance.a);
+		LFVLuminanceColor = LFVLuminanceColor * InstanceContribution.a + InstanceContribution.rgb;
+		LFVTransmittance = LFVTransmittance * InstanceContribution.a;
 	}
-	return LFVContribution;
+
+	return float4(LFVLuminanceColor, LFVTransmittance);
 }
 
 float4 GetLFVContribution(
