[EditorGizmos] Removed occlusion effect

diff --git a/Engine/Shaders/Private/PostProcessCompositePrimitives.usf b/Engine/Shaders/Private/PostProcessCompositePrimitives.usf
--- a/Engine/Shaders/Private/PostProcessCompositePrimitives.usf
+++ b/Engine/Shaders/Private/PostProcessCompositePrimitives.usf
@@ -437,23 +437,31 @@ void MainCompositeEditorPrimitivesPS(
 		EditorPrimitiveColor.rgb *= EditorPrimitiveColor.a;
 	}
 
-	// Generate 2x2 square tiling pattern for foreground primitive that end up behind scene opaque primitives.
-	float PatternMask = ((ColorPixelPos.x/2 + ColorPixelPos.y/2) % 2) * 0.7f;
+#if 1 // Don't modify based on occlusion
+	static const float OccludedColorMultiplier = 1.0f;
+	float DarkenMask = 1.0f;
+#else // Apply pattern mask based on occlusion
+	static const float OccludedColorMultiplier = 0.7f;
 
-	// the contants express the two opacity values we see when the primitive is hidden
-	float LowContrastPatternMask = lerp(0.2, 1, PatternMask);
+	// Generate 2x2 square tiling pattern for foreground primitive that end up behind scene opaque primitives.
+	float PatternMask = ((ColorPixelPos.x / 2 + ColorPixelPos.y / 2) % 2) * OccludedColorMultiplier;
 
+	// These constants express the two opacity values we see when the primitive is hidden
+	float LowContrastPatternMask = lerp(OccludedColorMultiplier, 1, PatternMask);
 	LowContrastPatternMask = saturate(lerp(LowContrastPatternMask, 1, bOpaqueEditorGizmo));
 
-	float HiddenMask = lerp(0.7f, 1.0f, DepthMask);
 	float DarkenMask = lerp(LowContrastPatternMask, 1.0f, DepthMask);
+#endif
 
-	// Blend editor ptimitives with scene color.
-	OutColor.rgb = SceneColor.rgb * (1 - EditorPrimitiveColor.a) + EditorPrimitiveColor.rgb * (DarkenMask * HiddenMask);
+	// Blend editor primitives with scene color.
+	OutColor.rgb = SceneColor.rgb * (1 - EditorPrimitiveColor.a) + EditorPrimitiveColor.rgb * DarkenMask;
 
 	FLATTEN
 	if (bProcessAlpha)
 	{
+		static const float OccludedAlphaMultiplier = OccludedColorMultiplier;
+		float HiddenMask = lerp(OccludedAlphaMultiplier, 1.0f, DepthMask);
+
 		OutColor.a = lerp(SceneColor.a, 1, EditorPrimitiveColor.a * DarkenMask * HiddenMask);
 	}
 	else
