Fix local light shadow on translucent support with Substrate.

This was added in CL 41746615 but only for the legacy path.

diff --git a/Engine/Shaders/Private/ForwardLightingCommon.ush b/Engine/Shaders/Private/ForwardLightingCommon.ush
--- a/Engine/Shaders/Private/ForwardLightingCommon.ush
+++ b/Engine/Shaders/Private/ForwardLightingCommon.ush
@@ -297,7 +297,7 @@ FDeferredLightingSplit GetForwardDirectLightingSplit(
 		float DynamicShadowing = dot(PreviewShadowMapChannelMask, DynamicShadowFactors);
 		float4 LightAttenuation = float4(1, 1, DynamicShadowing.x, DynamicShadowing.x);
 
-		#if VIRTUAL_SHADOW_MAP && SUPPORT_LOCAL_LIGHT_SHADOW_ON_FORWARD_LIT_TRANSLUCENT
+		#if VIRTUAL_SHADOW_MAP && SUPPORT_SHADOWED_LOCAL_LIGHT_ON_FORWARD_LIT_TRANSLUCENT
 		BRANCH
 		if (LocalLight.Internal.VirtualShadowMapId != INDEX_NONE)
 		{
diff --git a/Engine/Shaders/Private/Substrate/SubstrateForwardLighting.ush b/Engine/Shaders/Private/Substrate/SubstrateForwardLighting.ush
--- a/Engine/Shaders/Private/Substrate/SubstrateForwardLighting.ush
+++ b/Engine/Shaders/Private/Substrate/SubstrateForwardLighting.ush
@@ -510,6 +510,47 @@ float3 SubstrateForwardLighting(
 						LightMask = GetLocalLightAttenuation(TranslatedWorldPosition, LightData, ToLight, L);
 					}
 
+					#if VIRTUAL_SHADOW_MAP && SUPPORT_SHADOWED_LOCAL_LIGHT_ON_FORWARD_LIT_TRANSLUCENT
+					BRANCH
+					if (LocalLight.Internal.VirtualShadowMapId != INDEX_NONE)
+					{
+						float DynamicShadowFactor = 1.f;
+						LightData.ShadowedBits = 1;
+						#if SUPPORT_VSM_FOWARD_QUALITY == 1
+						{
+							const FLightShaderParameters LocalLightShaderParameters = ConvertToLightShaderParameters(LightData);
+
+							const float ScreenRayLengthWorld = VirtualShadowMap.ScreenRayLength * View.ClipToView[1][1] * Depth;
+							float SMRTRayOffset = ScreenRayLengthWorld;
+
+							// Use the surface normal instead of the layer's normal, for a single evaluation before runing throught the SubstrateTree
+							const half3 WorldNormal = normalize(MaterialParameters.TangentToWorld[2]);
+
+							FVirtualShadowMapSampleResult VirtualShadowMapSample = TraceLocalLight(
+								LocalLight.VirtualShadowMapId,
+								LocalLightShaderParameters,
+								PixelPos,
+								SceneDepth,
+								TranslatedWorldPosition,
+								SMRTRayOffset,
+								Dither,
+								WorldNormal);
+
+							FilterVirtualShadowMapSampleResult(PixelPos, VirtualShadowMapSample);
+
+							DynamicShadowFactor *= VirtualShadowMapSample.ShadowFactor;
+						}
+						#else
+						{
+							FVirtualShadowMapSampleResult VirtualShadowMapSample = SampleVirtualShadowMapLocal(LocalLight.Internal.VirtualShadowMapId, TranslatedWorldPosition);
+							DynamicShadowFactor *= VirtualShadowMapSample.ShadowFactor;
+						}
+						#endif
+
+						LightAttenuation = DynamicShadowFactor;
+					}
+					#endif
+
 					// Update the Substrate BSDF context accordgin to the new L
 					SubstrateBSDFContext.SubstrateUpdateBSDFContext(ToLight);
 
