Fixed: "implicit truncation of vector type" shader error in a mobile GPUScene implementation

diff --git a/Engine/Shaders/Private/SceneDataMobileLoader.ush b/Engine/Shaders/Private/SceneDataMobileLoader.ush
--- a/Engine/Shaders/Private/SceneDataMobileLoader.ush
+++ b/Engine/Shaders/Private/SceneDataMobileLoader.ush
@@ -161,8 +161,8 @@ FPrimitiveSceneData LoadPrimitiveDataUBO(uint DrawInstanceId)
 	// PreSkinnedLocalBounds, TODO.w
 	// TODO: this is needed only for skinned meshes
 	{
-		PrimitiveData.PreSkinnedLocalBoundsMin = LoadPrimitiveDataElementUBO(DrawInstanceId, DataIdx + 0);
-		PrimitiveData.PreSkinnedLocalBoundsMax = LoadPrimitiveDataElementUBO(DrawInstanceId, DataIdx + 1);
+		PrimitiveData.PreSkinnedLocalBoundsMin = LoadPrimitiveDataElementUBO(DrawInstanceId, DataIdx + 0).xyz;
+		PrimitiveData.PreSkinnedLocalBoundsMax = LoadPrimitiveDataElementUBO(DrawInstanceId, DataIdx + 1).xyz;
 		DataIdx += 2u;
 	}
 
