Fixed shader warning.

diff --git a/Engine/Shaders/Private/LocalFogVolumes/LocalFogVolumeSplat.usf b/Engine/Shaders/Private/LocalFogVolumes/LocalFogVolumeSplat.usf
--- a/Engine/Shaders/Private/LocalFogVolumes/LocalFogVolumeSplat.usf
+++ b/Engine/Shaders/Private/LocalFogVolumes/LocalFogVolumeSplat.usf
@@ -81,9 +81,9 @@ void LocalFogVolumeTiledPS(
 	ResolvedView = ResolveView();
 
 #if LFV_HALFRES_PS
-	const float HalfResScale = ResolvedView.ViewSizeAndInvSize.xy * LFVHalfResTextureSizeAndInvSize.zw;
+	const float2 HalfResScale = ResolvedView.ViewSizeAndInvSize.xy * LFVHalfResTextureSizeAndInvSize.zw;
 #else
-	const float HalfResScale = 1.0f;
+	const float2 HalfResScale = 1.0f;
 #endif
 	const float2 ViewRectSVPos = SVPos.xy * HalfResScale;
 	const float2 SourcePixelPos = ViewRectSVPos + ResolvedView.ViewRectMin.xy;
