- Support receiver mask for local lights (disabled by default, see r.shadow.virtual.usereceivermasklocal)
- Requires using non-greedy mip selection when enabled for SMRT and single sample lookups
- Cvar to enable/disable VSM overflow screen messages (enabled by default)
- Add force invalidate local cvar
- Remove greedy selection clipmap cvar; has been off for a while and increasingly clashes with the dynamic geo optimization strategies (receiver mask, etc)

diff --git a/Engine/Shaders/Private/VirtualShadowMaps/VirtualShadowMapPageMarking.usf b/Engine/Shaders/Private/VirtualShadowMaps/VirtualShadowMapPageMarking.usf
--- a/Engine/Shaders/Private/VirtualShadowMaps/VirtualShadowMapPageMarking.usf
+++ b/Engine/Shaders/Private/VirtualShadowMaps/VirtualShadowMapPageMarking.usf
@@ -39,7 +39,6 @@ float PageDilationBorderSizeLocal;
 uint InputType;
 uint bCullBackfacingPixels;
 uint NumDirectionalLightSmInds;
-uint MipModeLocal;
 
 float DOFBiasStrength;
 
@@ -60,7 +59,7 @@ uint GetMipLevelLocal(FVirtualShadowMapProjectionShaderData ProjectionData, floa
 	float3 ShadowTranslatedWorldPosition = TranslatedWorldPosition + ViewToShadowTranslation;
 
 	float Footprint = VirtualShadowMapCalcPixelFootprintLocal(ProjectionData, ShadowTranslatedWorldPosition, SceneDepth);
-	return GetMipLevelLocal(Footprint, MipModeLocal, ProjectionData.ResolutionLodBias, VirtualShadowMap.GlobalResolutionLodBias, ExtraBias);
+	return GetMipLevelLocal(Footprint, VirtualShadowMap.MipModeLocal, ProjectionData.ResolutionLodBias, VirtualShadowMap.GlobalResolutionLodBias, ExtraBias);
 }
 
 uint GetMipLevelLocal(int VirtualShadowMapId, float3 TranslatedWorldPosition, float SceneDepth, /* >= 0 */ float ExtraBias = 0.0f)
diff --git a/Engine/Shaders/Private/VirtualShadowMaps/VirtualShadowMapProjection.usf b/Engine/Shaders/Private/VirtualShadowMaps/VirtualShadowMapProjection.usf
--- a/Engine/Shaders/Private/VirtualShadowMaps/VirtualShadowMapProjection.usf
+++ b/Engine/Shaders/Private/VirtualShadowMaps/VirtualShadowMapProjection.usf
@@ -353,7 +353,10 @@ FVirtualShadowMapSampleResult ProjectLight(
 				VirtualShadowMapId,
 				TranslatedWorldPosition,
 				SMRTRayOffset,
-				WorldNormal);
+				WorldNormal,
+				true, // OptimalBias
+				SceneDepth
+			);
 		}
 		#endif // DIRECTIONAL_LIGHT
 		
diff --git a/Engine/Shaders/Private/VirtualShadowMaps/VirtualShadowMapProjectionCommon.ush b/Engine/Shaders/Private/VirtualShadowMaps/VirtualShadowMapProjectionCommon.ush
--- a/Engine/Shaders/Private/VirtualShadowMaps/VirtualShadowMapProjectionCommon.ush
+++ b/Engine/Shaders/Private/VirtualShadowMaps/VirtualShadowMapProjectionCommon.ush
@@ -24,7 +24,7 @@ float CalcAbsoluteClipmapLevel(FVirtualShadowMapProjectionShaderData BaseProject
 float CalcBiasedAbsoluteClipmapLevelForSampling(FVirtualShadowMapProjectionShaderData BaseProjectionData, float3 TranslatedWorldPosition)
 {
 	float AbsoluteLevel = CalcAbsoluteClipmapLevel(BaseProjectionData, TranslatedWorldPosition);
-	AbsoluteLevel += VirtualShadowMap.bClipmapGreedyLevelSelection ? 0.0f : BaseProjectionData.ResolutionLodBias;
+	AbsoluteLevel += BaseProjectionData.ResolutionLodBias;
 	return AbsoluteLevel;
 }
 
@@ -54,6 +54,23 @@ float VirtualShadowMapCalcPixelFootprintLocal(
 	return Footprint;
 }
 
+uint VirtualShadowMapGetMipLevelLocalFromReceiver(
+	FVirtualShadowMapProjectionShaderData ProjectionData,
+	float3 ReceiverShadowTranslatedWorld,
+	float ReceiverSceneDepth)
+{
+	uint MipLevel = 0;
+	if (ProjectionData.bUseReceiverMask)
+	{
+		float Footprint = VirtualShadowMapCalcPixelFootprintLocal(ProjectionData, ReceiverShadowTranslatedWorld, ReceiverSceneDepth);
+		return GetMipLevelLocal(Footprint, VirtualShadowMap.MipModeLocal, ProjectionData.ResolutionLodBias, VirtualShadowMap.GlobalResolutionLodBias);
+	}
+	else
+	{
+		return 0;
+	}
+}
+
 struct FVirtualShadowMapSample
 {
 	float Depth;
@@ -428,14 +445,17 @@ FVirtualShadowMapSampleResult SampleVirtualShadowMapDirectional(
 /**
 * Sample virtual shadow map for a local light
 * 
-* Parameter descriptions are as above
+* Parameter descriptions are as above, except:
+* - SceneDepth can be optionally provided to allow computing the appropriate mip level for sampling when
+*	using receiver mask to avoid sampling incomplete data from more detailed mip levels.
 */
 FVirtualShadowMapSampleResult SampleVirtualShadowMapLocal(
 	int VirtualShadowMapId,
 	float3 TranslatedWorldPosition,
 	float RayStartDistance,
 	float3 EstimatedGeoWorldNormal,
-	bool bUseOptimalBias = true)
+	bool bUseOptimalBias = true,
+	float SceneDepth = -1.0f)
 {
 	RayStartDistance = max(RayStartDistance, 0.0f);
 
@@ -453,7 +473,17 @@ FVirtualShadowMapSampleResult SampleVirtualShadowMapLocal(
 	float4 ShadowUVz = mul(float4(ShadowTranslatedWorldPosition, 1.0f), ProjectionData.TranslatedWorldToShadowUVMatrix);
 	ShadowUVz.xyz /= ShadowUVz.w;
 
-	FVirtualShadowMapSample SmSample = SampleVirtualShadowMap(VirtualShadowMapId, ShadowUVz.xy, ProjectionData.MinMipLevel);
+	// When using receiver mask, we also must compute the base mip level to use
+	// We cannot use greedy mip selection in this case since the higher res mips may not have the rendered
+	// geometry that covers the requested mask at this receiver. We must use a similar-or-lower mip than
+	// the one that we marked.
+	uint MipLevel = 0;
+	if (SceneDepth >= 0.0f)
+	{
+		MipLevel = VirtualShadowMapGetMipLevelLocalFromReceiver(ProjectionData, ShadowTranslatedWorldPosition, SceneDepth);
+	}
+
+	FVirtualShadowMapSample SmSample = SampleVirtualShadowMap(VirtualShadowMapId, ShadowUVz.xy, max(MipLevel, ProjectionData.MinMipLevel));
 	if (SmSample.bValid)
 	{
 		FVirtualShadowMapSampleResult Result = InitVirtualShadowMapSampleResult();
diff --git a/Engine/Shaders/Private/VirtualShadowMaps/VirtualShadowMapProjectionSpot.ush b/Engine/Shaders/Private/VirtualShadowMaps/VirtualShadowMapProjectionSpot.ush
--- a/Engine/Shaders/Private/VirtualShadowMaps/VirtualShadowMapProjectionSpot.ush
+++ b/Engine/Shaders/Private/VirtualShadowMaps/VirtualShadowMapProjectionSpot.ush
@@ -69,6 +69,7 @@ FSMRTSingleRayState SMRTSingleRayInitialize(
 	float3 RayEndShadowTranslatedWorld,
 	float ExtrapolateSlope,
 	float DepthBias,
+	uint MipLevel,
 	bool bClampUVs = true)
 {
 	FVirtualShadowMapProjectionShaderData ProjectionData = GetVirtualShadowMapProjectionData(VirtualShadowMapId);
@@ -98,7 +99,7 @@ FSMRTSingleRayState SMRTSingleRayInitialize(
 	Result.RayStartUVz = RayStartUVz.xyz;
 	Result.RayStepUVz = RayStepUVz.xyz;
 	Result.ExtrapolateSlope = ExtrapolateSlope * RayStepUVz.z;
-	Result.MinMipLevel = ProjectionData.MinMipLevel;
+	Result.MinMipLevel = max(MipLevel, ProjectionData.MinMipLevel);
 	Result.MipLevel = 0;
 	Result.VirtualTexelAddress = uint2(0xFFFFFFFF, 0xFFFFFFFF);
 	Result.PhysicalTexelAddress = uint2(0xFFFFFFFF, 0xFFFFFFFF);
@@ -151,11 +152,12 @@ FSMRTTwoCubeFaceRayState SMRTTwoCubeFaceRayInitialize(
 	float3 RayStartShadowTranslatedWorld,
 	float3 RayEndShadowTranslatedWorld,
 	float ExtrapolateSlope,
-	float DepthBias)
+	float DepthBias,
+	uint MipLevel)
 {
 	FSMRTTwoCubeFaceRayState Result;
-	Result.Face0 = SMRTSingleRayInitialize(VirtualShadowMapId0, RayStartShadowTranslatedWorld, RayEndShadowTranslatedWorld, ExtrapolateSlope, DepthBias, false);
-	Result.Face1 = SMRTSingleRayInitialize(VirtualShadowMapId1, RayStartShadowTranslatedWorld, RayEndShadowTranslatedWorld, ExtrapolateSlope, DepthBias, false);
+	Result.Face0 = SMRTSingleRayInitialize(VirtualShadowMapId0, RayStartShadowTranslatedWorld, RayEndShadowTranslatedWorld, ExtrapolateSlope, DepthBias, MipLevel, false);
+	Result.Face1 = SMRTSingleRayInitialize(VirtualShadowMapId1, RayStartShadowTranslatedWorld, RayEndShadowTranslatedWorld, ExtrapolateSlope, DepthBias, MipLevel, false);
 	Result.bSampleInFace1 = true;
 	return Result;
 }
@@ -192,13 +194,14 @@ FSMRTResult ShadowRayCastSpotLight(
 	float3 RayStartShadowTranslatedWorld,
 	float3 RayEndShadowTranslatedWorld,
 	float DepthBias,
+	uint MipLevel,
 	int NumSteps,
 	float StepOffset,
 	float ExtrapolateSlope,
 	// To get debug data out
 	inout FSMRTSingleRayState OutRayState)
 {
-	OutRayState = SMRTSingleRayInitialize(VirtualShadowMapId, RayStartShadowTranslatedWorld, RayEndShadowTranslatedWorld, ExtrapolateSlope, DepthBias);
+	OutRayState = SMRTSingleRayInitialize(VirtualShadowMapId, RayStartShadowTranslatedWorld, RayEndShadowTranslatedWorld, ExtrapolateSlope, DepthBias, MipLevel);
 	return SMRTRayCast(OutRayState, NumSteps, StepOffset);
 }
 
@@ -207,6 +210,7 @@ FSMRTResult ShadowRayCastPointLight(
 	float3 RayStartShadowTranslatedWorld,
 	float3 RayEndShadowTranslatedWorld,
 	float DepthBias,
+	uint MipLevel,
 	int NumSteps,
 	float StepOffset,
 	float ExtrapolateSlope,
@@ -220,7 +224,7 @@ FSMRTResult ShadowRayCastPointLight(
 	if (WaveActiveAnyTrue(RayStartFaceId != RayEndFaceId))
 	{
 		FSMRTTwoCubeFaceRayState RayState = SMRTTwoCubeFaceRayInitialize(RayStartFaceId, RayEndFaceId,
-			RayStartShadowTranslatedWorld, RayEndShadowTranslatedWorld, ExtrapolateSlope, DepthBias);
+			RayStartShadowTranslatedWorld, RayEndShadowTranslatedWorld, ExtrapolateSlope, DepthBias, MipLevel);
 		Result = SMRTRayCast(RayState, NumSteps, StepOffset);
 
 		if (RayState.bSampleInFace1)
@@ -235,7 +239,7 @@ FSMRTResult ShadowRayCastPointLight(
 	else
 	{
 		// Fast path: ray stays on a single cube map face (effectively a spot light)
-		OutRayState = SMRTSingleRayInitialize(RayEndFaceId, RayStartShadowTranslatedWorld, RayEndShadowTranslatedWorld, ExtrapolateSlope, DepthBias);
+		OutRayState = SMRTSingleRayInitialize(RayEndFaceId, RayStartShadowTranslatedWorld, RayEndShadowTranslatedWorld, ExtrapolateSlope, DepthBias, MipLevel);
 		Result = SMRTRayCast(OutRayState, NumSteps, StepOffset);
 	}
 
@@ -342,6 +346,14 @@ FVirtualShadowMapSampleResult TraceLocalLight(
 		SamplesPerRay = 0;
 	}
 
+	// When using receiver mask, we also must compute the base mip level to use
+	// We cannot use greedy mip selection in this case since the higher res mips may not have the rendered
+	// geometry that covers the requested mask at this receiver. We must use a similar-or-lower mip than
+	// the one that we marked.
+	// TODO: Any significant mip changes in GetMipLevelLocal should probably affect the above
+	// bias/texel dither math as well, which currently assumes a roughly proportional pixel:texel mapping
+	uint MipLevel = VirtualShadowMapGetMipLevelLocalFromReceiver(ProjectionData, ShadowTranslatedWorld, SceneDepth);
+
 	uint i = 0;
 	uint RayMissCount = 0;
 	float StepOffset = Noise;
@@ -391,6 +403,7 @@ FVirtualShadowMapSampleResult TraceLocalLight(
 				RayStart,
 				RayEnd,
 				DepthBias,
+				MipLevel,
 				SamplesPerRay,
 				StepOffset,
 				Settings.ExtrapolateMaxSlope,
@@ -403,6 +416,7 @@ FVirtualShadowMapSampleResult TraceLocalLight(
 				RayStart,
 				RayEnd,
 				DepthBias,
+				MipLevel,
 				SamplesPerRay,
 				StepOffset,
 				Settings.ExtrapolateMaxSlope,
diff --git a/Engine/Source/Runtime/Renderer/Private/Shadows/ShadowSceneRenderer.cpp b/Engine/Source/Runtime/Renderer/Private/Shadows/ShadowSceneRenderer.cpp
--- a/Engine/Source/Runtime/Renderer/Private/Shadows/ShadowSceneRenderer.cpp
+++ b/Engine/Source/Runtime/Renderer/Private/Shadows/ShadowSceneRenderer.cpp
@@ -111,6 +111,12 @@ static TAutoConsoleVariable<int32> CVarVirtualShadowMapNaniteAllowMultipassViews
 	ECVF_RenderThreadSafe
 );
 
+static TAutoConsoleVariable<int32>  CVarForceInvalidateLocalVSM(
+	TEXT("r.Shadow.Virtual.Cache.ForceInvalidateLocal"),
+	0,
+	TEXT("Forces local light VSMs (including distant lights) to always invalidate every frame. Generally only used for debugging."),
+	ECVF_RenderThreadSafe);
+
 extern TAutoConsoleVariable<int32> CVarMarkPixelPagesMipModeLocal;
 extern float GMinScreenRadiusForShadowCaster;
 
@@ -431,8 +437,9 @@ TSharedPtr<FVirtualShadowMapPerLightCacheEntry> FShadowSceneRenderer::AddLocalLi
 		LightSceneProxy->GetOrigin(),
 		LightSceneProxy->GetRadius(),
 		bIsDistantLight,
-		CacheManager->IsCacheEnabled(),
-		!bShouldForceTimeSliceDistantUpdate);
+		!CacheManager->IsCacheEnabled() || (CVarForceInvalidateLocalVSM.GetValueOnRenderThread() != 0),
+		!bShouldForceTimeSliceDistantUpdate,
+		IsVirtualShadowMapLocalReceiverMaskEnabled());
 
 	if (bIsDistantLight && PerLightCacheEntry->Prev.ScheduledFrameNumber == Scene.GetFrameNumber())
 	{
@@ -449,6 +456,7 @@ TSharedPtr<FVirtualShadowMapPerLightCacheEntry> FShadowSceneRenderer::AddLocalLi
 	{
 		uint32 PackedCullingViewId = FVirtualShadowMapProjectionShaderData::PackCullingViewId(SceneRenderer.Views[ClosestCullingViewIndex].SceneRendererPrimaryViewId, SceneRenderer.Views[ClosestCullingViewIndex].PersistentViewId);
 		uint32 Flags = PerLightCacheEntry->IsUncached() ? VSM_PROJ_FLAG_UNCACHED : 0U;
+		Flags |= PerLightCacheEntry->ShouldUseReceiverMask() ? VSM_PROJ_FLAG_USE_RECEIVER_MASK : 0U;
 		const FLightSceneProxy* Proxy = ProjectedShadowInfo->GetLightSceneInfo().Proxy;
 
 		// For now just tie this to whether anything has invalidated the light (including movement)
diff --git a/Engine/Source/Runtime/Renderer/Private/VirtualShadowMaps/VirtualShadowMapArray.cpp b/Engine/Source/Runtime/Renderer/Private/VirtualShadowMaps/VirtualShadowMapArray.cpp
--- a/Engine/Source/Runtime/Renderer/Private/VirtualShadowMaps/VirtualShadowMapArray.cpp
+++ b/Engine/Source/Runtime/Renderer/Private/VirtualShadowMaps/VirtualShadowMapArray.cpp
@@ -285,14 +285,6 @@ static TAutoConsoleVariable<int32> CVarCacheAllocateViaLRU(
 	ECVF_RenderThreadSafe
 );
 
-static TAutoConsoleVariable<int32> CVarClipmapGreedyLevelSelection(
-	TEXT("r.Shadow.Virtual.Clipmap.GreedyLevelSelection"),
-	0,
-	TEXT("When enabled, allows greedily sampling more detailed clipmap levels if they happen to be mapped.\n")
-	TEXT("This can increase shadow quality from certain viewing angles, but makes the clipmap boundry less stable which can exacerbate visual artifacts at low shadow resolutions."),
-	ECVF_Scalability | ECVF_RenderThreadSafe
-);
-
 #if !UE_BUILD_SHIPPING
 bool GDumpVSMLightNames = false;
 void DumpVSMLightNames()
@@ -518,6 +510,18 @@ static TAutoConsoleVariable<int32> CVarSMRTRayCountHair(
 	ECVF_Scalability | ECVF_RenderThreadSafe
 );
 
+static TAutoConsoleVariable<bool> CVarVSMUseReceiverMaskLocal(
+	TEXT("r.Shadow.Virtual.UseReceiverMaskLocal"),
+	false,
+	TEXT("Use receiver page masks with local lights. This enables much more effective culling especially at lower resolutions."),	
+	ECVF_RenderThreadSafe
+);
+
+bool IsVirtualShadowMapLocalReceiverMaskEnabled()
+{
+	return CVarVSMUseReceiverMaskLocal.GetValueOnRenderThread() != 0;
+}
+
 namespace Nanite
 {
 	extern bool IsStatFilterActive(const FString& FilterName);
@@ -624,8 +628,8 @@ void FVirtualShadowMapArray::Initialize(
 	UniformParameters.CoarsePagePixelThresholdDynamic = CVarCoarsePagePixelThresholdDynamic.GetValueOnRenderThread();
 	UniformParameters.CoarsePagePixelThresholdStatic = CVarCoarsePagePixelThresholdStatic.GetValueOnRenderThread();
 	UniformParameters.CoarsePagePixelThresholdDynamicNanite = CVarCoarsePagePixelThresholdDynamicNanite.GetValueOnRenderThread();
-	UniformParameters.bClipmapGreedyLevelSelection = CVarClipmapGreedyLevelSelection.GetValueOnRenderThread();
-	UniformParameters.bEnableReceiverMasks = IsDirectionalReceiverMaskEnabled();
+	UniformParameters.bEnableReceiverMasks = IsVirtualShadowMapDirectionalReceiverMaskEnabled() || IsVirtualShadowMapLocalReceiverMaskEnabled();
+	UniformParameters.MipModeLocal = CVarMarkPixelPagesMipModeLocal.GetValueOnRenderThread();
 
 	UniformParameters.SceneFrameNumber = Scene.GetFrameNumberRenderThread();
 	
@@ -4034,6 +4038,7 @@ void FVirtualShadowMapArray::AddRenderViewsLocal(
 		BaseParams.MinBoundsRadius = GMinScreenRadiusForShadowCaster;
 	}
 	BaseParams.Flags |= CacheEntry->IsUncached() ? NANITE_VIEW_FLAG_UNCACHED : 0U;
+	BaseParams.Flags |= CacheEntry->ShouldUseReceiverMask() ? NANITE_VIEW_FLAG_USE_RECEIVER_MASK : 0U;
 
 	Nanite::SetCullingViewOverrides(&Views[ClosestCullingViewIndex], BaseParams);
 	int32 NumMaps = ProjectedShadowInfo->bOnePassPointLightShadow ? 6 : 1;
diff --git a/Engine/Source/Runtime/Renderer/Private/VirtualShadowMaps/VirtualShadowMapArray.h b/Engine/Source/Runtime/Renderer/Private/VirtualShadowMaps/VirtualShadowMapArray.h
--- a/Engine/Source/Runtime/Renderer/Private/VirtualShadowMaps/VirtualShadowMapArray.h
+++ b/Engine/Source/Runtime/Renderer/Private/VirtualShadowMaps/VirtualShadowMapArray.h
@@ -58,6 +58,8 @@ inline bool IsSinglePageVirtualShadowMap(int32 VirtualShadowMapId)
 
 bool DoesVSMWantFroxels(EShaderPlatform ShaderPlatform);
 
+bool IsVirtualShadowMapLocalReceiverMaskEnabled();
+
 class FVirtualShadowMap
 {
 public:
@@ -172,9 +174,9 @@ BEGIN_GLOBAL_SHADER_PARAMETER_STRUCT(FVirtualShadowMapUniformParameters, )
 	// For shadow page age calculations
 	SHADER_PARAMETER(uint32, SceneFrameNumber)
 
-	SHADER_PARAMETER(uint32, bClipmapGreedyLevelSelection)
 	SHADER_PARAMETER(float, GlobalResolutionLodBias)
 	SHADER_PARAMETER(int32, bEnableReceiverMasks)
+	SHADER_PARAMETER(int32, MipModeLocal)
 
 	// SMRT parameters that are sometimes used globally
 	SHADER_PARAMETER(float, ScreenRayLength)
diff --git a/Engine/Source/Runtime/Renderer/Private/VirtualShadowMaps/VirtualShadowMapCacheManager.cpp b/Engine/Source/Runtime/Renderer/Private/VirtualShadowMaps/VirtualShadowMapCacheManager.cpp
--- a/Engine/Source/Runtime/Renderer/Private/VirtualShadowMaps/VirtualShadowMapCacheManager.cpp
+++ b/Engine/Source/Runtime/Renderer/Private/VirtualShadowMaps/VirtualShadowMapCacheManager.cpp
@@ -136,6 +136,14 @@ FAutoConsoleVariableRef CVarVSMLightRadiusCulling(
 	ECVF_RenderThreadSafe
 );
 
+int32 GVSMAllowScreenOverflowMessages = 1;
+FAutoConsoleVariableRef CVarVSMAllowScreenOverflowMessages(
+	TEXT("r.Shadow.Virtual.AllowScreenOverflowMessages"),
+	GVSMAllowScreenOverflowMessages,
+	TEXT("Can be used to disable on-screen VSM-related overflow messages. Even with the messages disabled, overflows will still be logged and can result in visual corruption."),
+	ECVF_RenderThreadSafe
+);
+
 static const char *VirtualShadowMap_StatNames[] = 
 {
 	"REQUESTED_THIS_FRAME_PAGES",
@@ -352,8 +360,9 @@ void FVirtualShadowMapPerLightCacheEntry::UpdateLocal(
 	const FVector& NewLightOrigin,
 	const float NewLightRadius,
 	bool bNewIsDistantLight,
-	bool bCacheEnabled,
-	bool bAllowInvalidation)
+	bool bForceInvalidate,
+	bool bAllowInvalidation,
+	bool bInUseReceiverMask)
 {
 	// TODO: The logic in this function is needlessly convoluted... clean up
 
@@ -369,12 +378,19 @@ void FVirtualShadowMapPerLightCacheEntry::UpdateLocal(
 
 	// Check cache validity based of shadow setup
 	// If it is a distant light, we want to let the time-share perform the invalidation.
-	if (!bCacheEnabled || (bAllowInvalidation && bLightMoved))
+	if (bForceInvalidate || (bAllowInvalidation && bLightMoved))
 	{
 		Prev.RenderedFrameNumber = -1;
 		//UE_LOG(LogRenderer, Display, TEXT("Invalidated!"));
 	}
 
+	// If we swap receiver mask on/off we also may need to invalidate due to incomplete pages
+	if (bInUseReceiverMask != bUseReceiverMask)
+	{
+		Prev.RenderedFrameNumber = -1;
+		bUseReceiverMask = bInUseReceiverMask;
+	}
+
 	// On transition between uncached <-> cached we must invalidate since the static pages may not be initialized
 	bool bNewIsUncached = Prev.RenderedFrameNumber < 0;
 	if (bNewIsUncached != bIsUncached)
@@ -866,14 +882,17 @@ void FVirtualShadowMapArrayCacheManager::InitExtension(FScene& InScene)
 	{
 		float RealTimeSeconds = float(FGameTime::GetTimeSinceAppStart().GetRealTimeSeconds());
 
-		for (uint32 OverflowTypeIndex = 0; OverflowTypeIndex < VSM_STAT_OVERFLOW_FLAG_NUM; OverflowTypeIndex++)
+		if (GVSMAllowScreenOverflowMessages != 0)
 		{
-			// Show for ~10s after last overflow
-			float LastOverflowTime = LastOverflowTimes[OverflowTypeIndex];
-			if (LastOverflowTime >= 0.0f && RealTimeSeconds - LastOverflowTime < 10.0f)
+			for (uint32 OverflowTypeIndex = 0; OverflowTypeIndex < VSM_STAT_OVERFLOW_FLAG_NUM; OverflowTypeIndex++)
 			{
-				FText OverflowMessage = GetOverflowMessage(OverflowTypeIndex);
-				OutMessages.Add(FCoreDelegates::EOnScreenMessageSeverity::Warning, FText::FromString(FString::Printf(TEXT("%s (%0.0f seconds ago)"), *(OverflowMessage.ToString()), RealTimeSeconds - LastOverflowTime)));
+				// Show for ~10s after last overflow
+				float LastOverflowTime = LastOverflowTimes[OverflowTypeIndex];
+				if (LastOverflowTime >= 0.0f && RealTimeSeconds - LastOverflowTime < 10.0f)
+				{
+					FText OverflowMessage = GetOverflowMessage(OverflowTypeIndex);
+					OutMessages.Add(FCoreDelegates::EOnScreenMessageSeverity::Warning, FText::FromString(FString::Printf(TEXT("%s (%0.0f seconds ago). See r.Shadow.Virtual.AllowScreenOverflowMessages."), *(OverflowMessage.ToString()), RealTimeSeconds - LastOverflowTime)));
+				}
 			}
 		}
 
diff --git a/Engine/Source/Runtime/Renderer/Private/VirtualShadowMaps/VirtualShadowMapCacheManager.h b/Engine/Source/Runtime/Renderer/Private/VirtualShadowMaps/VirtualShadowMapCacheManager.h
--- a/Engine/Source/Runtime/Renderer/Private/VirtualShadowMaps/VirtualShadowMapCacheManager.h
+++ b/Engine/Source/Runtime/Renderer/Private/VirtualShadowMaps/VirtualShadowMapCacheManager.h
@@ -134,8 +134,9 @@ public:
 		const FVector& NewLightOrigin,
 		const float NewLightRadius,
 		bool bNewIsDistantLight,
-		bool bCacheEnabled,
-		bool bAllowInvalidation);
+		bool bForceInvalidate,
+		bool bAllowInvalidation,
+		bool bInUseReceiverMask);
 
 	void Invalidate() { Prev.RenderedFrameNumber = -1; }
 
diff --git a/Engine/Source/Runtime/Renderer/Private/VirtualShadowMaps/VirtualShadowMapClipmap.cpp b/Engine/Source/Runtime/Renderer/Private/VirtualShadowMaps/VirtualShadowMapClipmap.cpp
--- a/Engine/Source/Runtime/Renderer/Private/VirtualShadowMaps/VirtualShadowMapClipmap.cpp
+++ b/Engine/Source/Runtime/Renderer/Private/VirtualShadowMaps/VirtualShadowMapClipmap.cpp
@@ -149,7 +149,7 @@ static TAutoConsoleVariable<bool> CVarVSMUseReceiverMaskDirectional(
 	ECVF_RenderThreadSafe
 );
 
-bool IsDirectionalReceiverMaskEnabled()
+bool IsVirtualShadowMapDirectionalReceiverMaskEnabled()
 {
 	return CVarVSMUseReceiverMaskDirectional.GetValueOnRenderThread() != 0;
 }
@@ -177,7 +177,7 @@ FVirtualShadowMapClipmapConfig FVirtualShadowMapClipmapConfig::GetGlobal()
 		.ResolutionLodBiasMoving = CVarVirtualShadowMapResolutionLodBiasDirectionalMoving.GetValueOnRenderThread(),
 		.bForceInvalidate = CVarForceInvalidateDirectionalVSM.GetValueOnRenderThread() != 0,
 		.bCullDynamicTightly = CVarVSMClipmapCullDynamicTightly.GetValueOnRenderThread(),
-		.bUseReceiverMask = IsDirectionalReceiverMaskEnabled()
+		.bUseReceiverMask = IsVirtualShadowMapDirectionalReceiverMaskEnabled()
 	};
 }
 
diff --git a/Engine/Source/Runtime/Renderer/Private/VirtualShadowMaps/VirtualShadowMapClipmap.h b/Engine/Source/Runtime/Renderer/Private/VirtualShadowMaps/VirtualShadowMapClipmap.h
--- a/Engine/Source/Runtime/Renderer/Private/VirtualShadowMaps/VirtualShadowMapClipmap.h
+++ b/Engine/Source/Runtime/Renderer/Private/VirtualShadowMaps/VirtualShadowMapClipmap.h
@@ -19,7 +19,7 @@ class FVirtualShadowMapArray;
 class FVirtualShadowMapArrayCacheManager;
 class FVirtualShadowMapPerLightCacheEntry;
 
-bool IsDirectionalReceiverMaskEnabled();
+bool IsVirtualShadowMapDirectionalReceiverMaskEnabled();
 
 class FVirtualShadowMapClipmapConfig
 {
