Add screen trace support for hair transmittance computation.

This adds finer/higher frequency shadow details.

diff --git a/Engine/Shaders/Private/HairStrands/HairStrandsDeepTransmittanceMask.usf b/Engine/Shaders/Private/HairStrands/HairStrandsDeepTransmittanceMask.usf
--- a/Engine/Shaders/Private/HairStrands/HairStrandsDeepTransmittanceMask.usf
+++ b/Engine/Shaders/Private/HairStrands/HairStrandsDeepTransmittanceMask.usf
@@ -278,9 +278,66 @@ float3 InternalGetHairVoxelJitter(uint2 PixelCoord, uint SampleIndex, uint MaxSa
 	return JitterMode > 0 ? RandomSample : float3(0,0,0);
 }
 
+#include "../SSRT/SSRTRayCast.ush"
+#include "../HZB.ush"
+uint bUseScreenTrace;
+float ScreenTraceLength;
+
+// TODO: Merge with the version in HairStrandsEnvironementLighting.usf
+float GetRayVisibility(float3 TranslatedWorldPosition,float SceneDepth, float3 RayDirection, uint2 PixelCoord, uint PixelRayIndex, uint NumPixelSamples)
+{
+	const float SlopeCompareToleranceScale = 0.5f;
+	const float MaxScreenTraceFraction = ScreenTraceLength * 2.0f / (float)View.ViewSizeAndInvSize.x;
+	float TraceDistance = MaxScreenTraceFraction * 2.0 * GetScreenRayLengthMultiplierForProjectionType(SceneDepth).x;
+	float DepthThresholdScale = View.ProjectionDepthThicknessScale;
+
+	const float PositionBias = 0.1f;
+	const float3 TranslatedWorldStartPosition = TranslatedWorldPosition + RayDirection * PositionBias;
+	
+	FSSRTCastingSettings CastSettings = CreateDefaultCastSettings();
+	CastSettings.bStopWhenUncertain = true;
+
+	bool bHit = false;
+	float Level;
+	float3 HitUVz;
+	bool bRayWasClipped;
+	uint NumSteps = 4;
+	float StartMipLevel = 0;
+	float RayRoughness = .2f;
+	uint2 NoiseCoord = PixelCoord * uint2(NumPixelSamples, 1) + uint2(PixelRayIndex, 0);
+	float StepOffset = BlueNoiseScalar(PixelCoord, View.StateFrameIndex);
+
+	FSSRTRay Ray = InitScreenSpaceRayFromWorldSpace(
+		TranslatedWorldStartPosition, RayDirection,
+		/* WorldTMax = */ TraceDistance,
+		/* SceneDepth = */ SceneDepth,
+		/* SlopeCompareToleranceScale */ SlopeCompareToleranceScale * DepthThresholdScale * (float)NumSteps,
+		/* bExtendRayToScreenBorder = */ false,
+		/* out */ bRayWasClipped);
+
+	bool bUncertain;
+	float3 DebugOutput;
+
+	CastScreenSpaceRay(
+		ClosestHZBTexture, ClosestHZBTextureSampler,
+		StartMipLevel,
+		CastSettings,
+		Ray, RayRoughness, NumSteps, StepOffset - .9f,
+		HZBUvFactorAndInvFactor, false,
+		/* out */ DebugOutput,
+		/* out */ HitUVz,
+		/* out */ Level,
+		/* out */ bHit,
+		/* out */ bUncertain);
+
+	bHit = bHit && !bUncertain;
+	return bHit ? 0.0f : 1.0f;
+}
+
 FHairTransmittanceMask ComputeTransmittanceVoxel(
 	const uint2 PixelCoord,
 	const float PixelRadius,
+	const float SceneDepth,
 	float3 TranslatedWorldPosition,
 	uint MacroGroupId,
 	FTransmittanceSettings Settings,
@@ -335,6 +392,13 @@ FHairTransmittanceMask ComputeTransmittanceVoxel(
 			TraversalSettings);
 		InitHairTraversalResult();
 
+		if (bUseScreenTrace)
+		{
+			const float HairCountPerScreenSpaceIntersection = 1.5f; // Arbitrary value, visually tweaked
+			const float ScreenTraceVis = GetRayVisibility(TranslatedWorldPosition, SceneDepth, Settings.LightDirection, PixelCoord, SampleIt, SampleCount);
+			AvgHairCount += saturate(1.f-ScreenTraceVis) * HairCountPerScreenSpaceIntersection;
+		}
+
 		AvgHairCount += Result.HairCount;
 		AvgHairVisility += Result.Visibility;
 	}
@@ -422,6 +486,7 @@ void MainCS(uint2 DispatchThreadId : SV_DispatchThreadID)
 		const FHairTransmittanceMask Out = ComputeTransmittanceVoxel(
 			PixelCoord, 
 			PixelRadius, 
+			SceneDepth,
 			TranslatedWorldPosition,
 			MacroGroupIndex, 
 			Settings, 
@@ -480,7 +545,7 @@ void MainCS(uint2 DispatchThreadId : SV_DispatchThreadID)
 		Settings.Random = GetHairVoxelJitter(PixelCoord.xy, View.StateFrameIndexMod8, VirtualVoxel.JitterMode);
 
 		const float PixelRadius = ConvertGivenDepthRadiusForProjectionType(VirtualVoxel.HairCoveragePixelRadiusAtDepth1, SceneDepth);
-		const FHairTransmittanceMask Out = ComputeTransmittanceVoxel(PixelCoord, PixelRadius, TranslatedWorldPosition, MacroGroupIndex, Settings, bDebugEnabled);
+		const FHairTransmittanceMask Out = ComputeTransmittanceVoxel(PixelCoord, PixelRadius, SceneDepth, TranslatedWorldPosition, MacroGroupIndex, Settings, bDebugEnabled);
 
 		#if VOXEL_TRAVERSAL_DEBUG
 		if (bDebugEnabled)
diff --git a/Engine/Source/Runtime/Renderer/Private/HZB.cpp b/Engine/Source/Runtime/Renderer/Private/HZB.cpp
--- a/Engine/Source/Runtime/Renderer/Private/HZB.cpp
+++ b/Engine/Source/Runtime/Renderer/Private/HZB.cpp
@@ -50,42 +50,47 @@ static FHZBParameters InitHZBCommonParameter(FRDGBuilder& GraphBuilder, const FV
 	return Out;
 }
 
-FHZBParameters GetHZBParameters(FRDGBuilder& GraphBuilder, const FViewInfo& View, EHZBType InHZBTypes)
+FHZBParameters GetHZBParameters(FRDGBuilder& GraphBuilder, const FViewInfo& View, EHZBType InHZBTypes, FRDGTextureRef InClosestHZB, FRDGTextureRef InFurthestHZB)
 {
 	FIntPoint TextureExtent;
 	if (InHZBTypes & EHZBType::FurthestHZB)
 	{
-		TextureExtent = View.HZB ? View.HZB->Desc.Extent : FIntPoint(1,1);
+		TextureExtent = InFurthestHZB ? InFurthestHZB->Desc.Extent : FIntPoint(1,1);
 	}
 	else
 	{
-		TextureExtent = View.ClosestHZB ? View.ClosestHZB->Desc.Extent : FIntPoint(1,1);
+		TextureExtent = InClosestHZB ? InClosestHZB->Desc.Extent : FIntPoint(1,1);
 	}
 
 	FHZBParameters Out = InitHZBCommonParameter(GraphBuilder, &View, View.ViewRect, TextureExtent);
 
-	if ((InHZBTypes & EHZBType::FurthestHZB) && View.HZB)
+	if ((InHZBTypes & EHZBType::FurthestHZB) && InFurthestHZB)
 	{
 		Out.bIsHZBValid = 1u;
-		Out.HZBTexture = View.HZB;
+		Out.HZBTexture = InFurthestHZB;
 	
-		check(View.HZB && HasBeenProduced(View.HZB));
+		check(InFurthestHZB && HasBeenProduced(InFurthestHZB));
 		Out.bIsFurthestHZBValid = 1u;
-		Out.FurthestHZBTexture = View.HZB;
+		Out.FurthestHZBTexture = InFurthestHZB;
 	}
 
-	if ((InHZBTypes & EHZBType::ClosestHZB) && View.ClosestHZB)
+	if ((InHZBTypes & EHZBType::ClosestHZB) && InClosestHZB)
 	{
-		check(View.ClosestHZB && HasBeenProduced(View.ClosestHZB));
+		check(InClosestHZB && HasBeenProduced(InClosestHZB));
 		Out.bIsClosestHZBValid = 1;
-		Out.ClosestHZBTexture = View.ClosestHZB;
+		Out.ClosestHZBTexture = InClosestHZB;
 	}
 
 	// Sanity check
-	//check(View.ClosestHZB.Desc.Extent == View.HZB.Desc.Extent);
+	//check(ClosestHZB.Desc.Extent == FurthestHZB.Desc.Extent);
 	return Out;
 }
 
+FHZBParameters GetHZBParameters(FRDGBuilder& GraphBuilder, const FViewInfo& View, EHZBType InHZBTypes)
+{
+	return GetHZBParameters(GraphBuilder, View, InHZBTypes, View.ClosestHZB, View.HZB);
+}
+
 FHZBParameters GetHZBParameters(FRDGBuilder& GraphBuilder, const FViewInfo& View, bool bUsePreviousHZBAsFallback)
 {
 	const bool bUseCurrent = View.HZB != nullptr;
diff --git a/Engine/Source/Runtime/Renderer/Private/HZB.h b/Engine/Source/Runtime/Renderer/Private/HZB.h
--- a/Engine/Source/Runtime/Renderer/Private/HZB.h
+++ b/Engine/Source/Runtime/Renderer/Private/HZB.h
@@ -49,6 +49,7 @@ enum class EHZBType : uint8
 	All = ClosestHZB | FurthestHZB
 };
 
+FHZBParameters GetHZBParameters(FRDGBuilder& GraphBuilder, const FViewInfo& View, EHZBType InHZBTypes, FRDGTextureRef InClosestHZB, FRDGTextureRef InFurthestHZB);
 FHZBParameters GetHZBParameters(FRDGBuilder& GraphBuilder, const FViewInfo& View, EHZBType InHZBTypes);
 FHZBParameters GetHZBParameters(FRDGBuilder& GraphBuilder, const FViewInfo& View, bool bUsePreviousHZBAsFallback);
 FHZBParameters GetDummyHZBParameters(FRDGBuilder& GraphBuilder);
diff --git a/Engine/Source/Runtime/Renderer/Private/HairStrands/HairStrandsTransmittance.cpp b/Engine/Source/Runtime/Renderer/Private/HairStrands/HairStrandsTransmittance.cpp
--- a/Engine/Source/Runtime/Renderer/Private/HairStrands/HairStrandsTransmittance.cpp
+++ b/Engine/Source/Runtime/Renderer/Private/HairStrands/HairStrandsTransmittance.cpp
@@ -53,6 +53,12 @@ static FAutoConsoleVariableRef CVarHairStrandsShadowRandomTraversalType(TEXT("r.
 static int32 GHairStrandsShadow_ShadowMaskPassType = 1;
 static FAutoConsoleVariableRef CVarHairStrandsShadow_ShadowMaskPassType(TEXT("r.HairStrands.DeepShadow.ShadowMaskPassType"), GHairStrandsShadow_ShadowMaskPassType, TEXT("Change how shadow mask from hair onto opaque geometry is generated. 0: one pass per hair group, 1: one pass for all groups."), ECVF_Scalability | ECVF_RenderThreadSafe);
 
+static int32 GHairStrandsShadow_UseScreenTrace = 1;
+static FAutoConsoleVariableRef CVarHairStrandsShadow_HZBType(TEXT("r.HairStrands.DeepShadow.UseScreenTrace"), GHairStrandsShadow_UseScreenTrace, TEXT("Use screen trace for adding fine occlusion to hair transmission."), ECVF_Scalability | ECVF_RenderThreadSafe);
+
+static float GHairStrandsShadow_ScreenTraceLength = 10.f;
+static FAutoConsoleVariableRef CVarHairStrandsShadow_ScreenTraceLength(TEXT("r.HairStrands.DeepShadow.ScreenTraceLength"), GHairStrandsShadow_ScreenTraceLength, TEXT("Screen trace length for hair transmission."), ECVF_Scalability | ECVF_RenderThreadSafe);
+
 static float GetDeepShadowDensityScale() { return FMath::Max(0.0f, GDeepShadowDensityScale); }
 static float GetDeepShadowDepthBiasScale() { return FMath::Max(0.0f, GDeepShadowDepthBiasScale); }
 ///////////////////////////////////////////////////////////////////////////////////////////////////
@@ -200,7 +206,10 @@ class FHairStrandsVoxelTransmittanceMaskCS : public FGlobalShader
 		SHADER_PARAMETER_RDG_UNIFORM_BUFFER(FVirtualVoxelParameters, VirtualVoxel)
 		SHADER_PARAMETER_STRUCT_INCLUDE(FVirtualShadowMapSamplingParameters, VirtualShadowMap)
 		SHADER_PARAMETER_STRUCT_REF(FBlueNoise, BlueNoise)
-
+		SHADER_PARAMETER_STRUCT_INCLUDE(FHZBParameters, HZBParameters)
+		SHADER_PARAMETER(float, ScreenTraceNoFallbackThicknessScale)
+		SHADER_PARAMETER(uint32, bUseScreenTrace)
+		SHADER_PARAMETER(float, ScreenTraceLength)
 		SHADER_PARAMETER(FVector4f, TranslatedLightPosition_LightDirection)
 		SHADER_PARAMETER(FVector4f, ShadowChannelMask)
 		SHADER_PARAMETER(uint32, LightChannelMask)
@@ -287,6 +296,13 @@ static FRDGBufferRef AddHairStrandsVoxelTransmittanceMaskPass(
 	const bool bIsSuperSampled = GHairStrandsTransmittanceSuperSampling > 0;
 	const bool bIsMipTraversal = GHairStrandsTransmittanceMaskUseMipTraversal > 0;
 
+	// Screen trace
+	const float ScreenTraceNoFallbackThicknessScale = View.ViewMatrices.GetPerProjectionDepthThicknessScale();
+	Parameters->HZBParameters = GetHZBParameters(GraphBuilder, View, EHZBType::ClosestHZB, View.HairStrandsViewData.VisibilityData.HairOnlyDepthFurthestHZBTexture, View.HairStrandsViewData.VisibilityData.HairOnlyDepthClosestHZBTexture);
+	Parameters->ScreenTraceNoFallbackThicknessScale = ScreenTraceNoFallbackThicknessScale;
+	Parameters->bUseScreenTrace = GHairStrandsShadow_UseScreenTrace > 0? 1u : 0u;
+	Parameters->ScreenTraceLength = FMath::Max(GHairStrandsShadow_ScreenTraceLength, 0.f);
+
 	FHairStrandsVoxelTransmittanceMaskCS::FPermutationDomain PermutationVector;
 	PermutationVector.Set<FHairStrandsVoxelTransmittanceMaskCS::FTransmittanceGroupSize>(NodeGroupSize);
 	PermutationVector.Set<FHairStrandsVoxelTransmittanceMaskCS::FSuperSampling>(bIsSuperSampled ? 1 : 0);
