Fixed: Issues with a ClearCoat shading model and a half precision on mobile platforms

diff --git a/Engine/Shaders/Private/DeferredShadingCommon.ush b/Engine/Shaders/Private/DeferredShadingCommon.ush
--- a/Engine/Shaders/Private/DeferredShadingCommon.ush
+++ b/Engine/Shaders/Private/DeferredShadingCommon.ush
@@ -347,7 +347,7 @@ struct FGBufferData
 	// 0..1
 	half Specular;
 	// 0..1
-	half4 CustomData;
+	float4 CustomData;
 	// AO utility value
 	half GenericAO;
 	// Indirect irradiance luma
diff --git a/Engine/Shaders/Private/ShadingModels.ush b/Engine/Shaders/Private/ShadingModels.ush
--- a/Engine/Shaders/Private/ShadingModels.ush
+++ b/Engine/Shaders/Private/ShadingModels.ush
@@ -390,8 +390,8 @@ BxDFContext RefractClearCoatContext(BxDFContext Context)
 
 FDirectLighting ClearCoatBxDF( FGBufferData GBuffer, half3 N, half3 V, half3 L, float Falloff, half NoL, FAreaLight AreaLight, FShadowTerms Shadow )
 {
-	const half ClearCoat			= GBuffer.CustomData.x;
-	const half ClearCoatRoughness	= max(GBuffer.CustomData.y, 0.02f);
+	const float ClearCoat			= GBuffer.CustomData.x;
+	const float ClearCoatRoughness	= max(GBuffer.CustomData.y, 0.02f);
 
 	FDirectLighting Lighting = {
 		float3(0.0, 0.0, 0.0),
