Add jitter mode support to hair transmittance screen trace so that EngineTest gives stable result.

diff --git a/Engine/Shaders/Private/HairStrands/HairStrandsDeepTransmittanceMask.usf b/Engine/Shaders/Private/HairStrands/HairStrandsDeepTransmittanceMask.usf
--- a/Engine/Shaders/Private/HairStrands/HairStrandsDeepTransmittanceMask.usf
+++ b/Engine/Shaders/Private/HairStrands/HairStrandsDeepTransmittanceMask.usf
@@ -305,7 +305,7 @@ float GetRayVisibility(float3 TranslatedWorldPosition,float SceneDepth, float3 R
 	float StartMipLevel = 0;
 	float RayRoughness = .2f;
 	uint2 NoiseCoord = PixelCoord * uint2(NumPixelSamples, 1) + uint2(PixelRayIndex, 0);
-	float StepOffset = BlueNoiseScalar(PixelCoord, View.StateFrameIndex);
+	float StepOffset = VirtualVoxel.JitterMode > 0 ? BlueNoiseScalar(PixelCoord, View.StateFrameIndex) : 0;
 
 	FSSRTRay Ray = InitScreenSpaceRayFromWorldSpace(
 		TranslatedWorldStartPosition, RayDirection,
