Fix warning:
Warning:  /Engine/Private/PostProcessMaterialShaders.usf:454:11: warning: redefinition of 'CoordinateIndex' shadows declaration in the outer scope; most recent declaration will be used [-Wfor-redefinition]

diff --git a/Engine/Shaders/Private/PostProcessMaterialShaders.usf b/Engine/Shaders/Private/PostProcessMaterialShaders.usf
--- a/Engine/Shaders/Private/PostProcessMaterialShaders.usf
+++ b/Engine/Shaders/Private/PostProcessMaterialShaders.usf
@@ -454,19 +454,23 @@ void MainVS(
 	GetMaterialCustomizedUVs(VertexParameters, CustomizedUVs);
 	GetCustomInterpolators(VertexParameters, CustomizedUVs);
 
-	UNROLL
-	for (int CoordinateIndex = 0; CoordinateIndex < NUM_MATERIAL_TEXCOORDS; CoordinateIndex++)
 	{
-	#if POST_PROCESS_AR_PASSTHROUGH
-		SetUV(Output, CoordinateIndex, OutUV);
-	#else
-		SetUV(Output, CoordinateIndex, InPosition.xy);
-	#endif
+		UNROLL
+		for (int CoordinateIndex = 0; CoordinateIndex < NUM_MATERIAL_TEXCOORDS; CoordinateIndex++)
+		{
+		#if POST_PROCESS_AR_PASSTHROUGH
+			SetUV(Output, CoordinateIndex, OutUV);
+		#else
+			SetUV(Output, CoordinateIndex, InPosition.xy);
+		#endif
+		}
 	}
-	UNROLL
-	for (int CoordinateIndex = NUM_MATERIAL_TEXCOORDS; CoordinateIndex < NUM_TEX_COORD_INTERPOLATORS; CoordinateIndex++)
 	{
-		SetUV(Output, CoordinateIndex, CustomizedUVs[CoordinateIndex]);
+		UNROLL
+		for (int CoordinateIndex = NUM_MATERIAL_TEXCOORDS; CoordinateIndex < NUM_TEX_COORD_INTERPOLATORS; CoordinateIndex++)
+		{
+			SetUV(Output, CoordinateIndex, CustomizedUVs[CoordinateIndex]);
+		}
 	}
 #endif
 }
@@ -486,19 +490,23 @@ void MainPS(
 	float2 BufferUV = ViewportUVToBufferUV(ViewportUV);
 	
 #if NUM_TEX_COORD_INTERPOLATORS
-	UNROLL
-	for (int CoordinateIndex = 0; CoordinateIndex < NUM_MATERIAL_TEXCOORDS; CoordinateIndex++)
 	{
-	#if POST_PROCESS_AR_PASSTHROUGH
-		Parameters.TexCoords[CoordinateIndex] = GetUV(Input, CoordinateIndex);
-	#else
-		Parameters.TexCoords[CoordinateIndex] = ViewportUV;
-	#endif
+		UNROLL
+		for (int CoordinateIndex = 0; CoordinateIndex < NUM_MATERIAL_TEXCOORDS; CoordinateIndex++)
+		{
+		#if POST_PROCESS_AR_PASSTHROUGH
+			Parameters.TexCoords[CoordinateIndex] = GetUV(Input, CoordinateIndex);
+		#else
+			Parameters.TexCoords[CoordinateIndex] = ViewportUV;
+		#endif
+		}
 	}
-	UNROLL
-	for (int CoordinateIndex = NUM_MATERIAL_TEXCOORDS; CoordinateIndex < NUM_TEX_COORD_INTERPOLATORS; CoordinateIndex++)
 	{
-		Parameters.TexCoords[CoordinateIndex] = GetUV(Input, CoordinateIndex);
+		UNROLL
+		for (int CoordinateIndex = NUM_MATERIAL_TEXCOORDS; CoordinateIndex < NUM_TEX_COORD_INTERPOLATORS; CoordinateIndex++)
+		{
+			Parameters.TexCoords[CoordinateIndex] = GetUV(Input, CoordinateIndex);
+		}
 	}
 #endif
 	Parameters.VertexColor = 1;
