VolumetricCloud - workaround for depth leaking in cloud on sky pixel, causing bilateral upsampel filter to flicker.

diff --git a/Engine/Shaders/Private/VolumetricRenderTarget.usf b/Engine/Shaders/Private/VolumetricRenderTarget.usf
--- a/Engine/Shaders/Private/VolumetricRenderTarget.usf
+++ b/Engine/Shaders/Private/VolumetricRenderTarget.usf
@@ -158,6 +158,18 @@ float BOX_NORM_LUMA(float3 Clip, float3 Min, float3 Max)
 #endif
 }
 
+void FixupDepthDataSafe(inout FDepthData DepthData)
+{
+	// When reprojecting depth, it can happen landscape or other front object leaks into the sky pixels, when those pixel have not been identified as bUseNewSample.
+	// This happens rarely but no concrete solution have been found to resolve that situation.
+	// To avoid cloud to flicker (depth difference can be large, causing resonstruction to identify a depth discontinuity)
+	const float MaxDepthKm = 60000.0f;	// Anything beyond that threshold will be considered at maximum depth MaxHalfFloat.
+	DepthData.CloudFrontDepthFromViewKm	= DepthData.CloudFrontDepthFromViewKm	> MaxDepthKm ? MaxHalfFloat : DepthData.CloudFrontDepthFromViewKm;
+	DepthData.SceneDepthFromViewKm		= DepthData.SceneDepthFromViewKm		> MaxDepthKm ? MaxHalfFloat : DepthData.SceneDepthFromViewKm;
+	DepthData.MinMaxViewDepthKm.x		= DepthData.MinMaxViewDepthKm.x			> MaxDepthKm ? MaxHalfFloat : DepthData.MinMaxViewDepthKm.x;
+	DepthData.MinMaxViewDepthKm.y		= DepthData.MinMaxViewDepthKm.y			> MaxDepthKm ? MaxHalfFloat : DepthData.MinMaxViewDepthKm.y;
+}
+
 void ReconstructVolumetricRenderTargetPS(
 	in float4 SVPos : SV_POSITION, 
 #if PERMUTATION_CLOUD_MIN_AND_MAX_DEPTH
@@ -290,7 +302,9 @@ void ReconstructVolumetricRenderTargetPS(
 			FDepthData NewDepthData = SafeLoadTracingVolumetricDepthTexture(CenterSample);
 
 			FDepthData HistoryDepthData = SafeSamplePreviousFrameVolumetricDepthTexture(PrevScreenUVs);
-			const float HistorySceneDepthFromViewCm = HistoryDepthData.SceneDepthFromViewKm * KILOMETER_TO_CENTIMETER;
+
+			FixupDepthDataSafe(NewDepthData);
+			FixupDepthDataSafe(HistoryDepthData);
 
 		#if PERMUTATION_CLOUD_MIN_AND_MAX_DEPTH
 			
@@ -353,7 +367,7 @@ void ReconstructVolumetricRenderTargetPS(
 			{
 				RGBA = SafeSamplePreviousFrameVolumetricTexture(PrevScreenUVs);
 				RGBA2 = SafeSamplePreviousFrameVolumetricSecondaryTexture(PrevScreenUVs);
-				DepthData = SafeSamplePreviousFrameVolumetricDepthTexture(PrevScreenUVs);
+				DepthData = HistoryDepthData;
 
 				if (!(all(IsFinite(RGBA)) && all(IsFinite(float2(DepthData.CloudFrontDepthFromViewKm, DepthData.SceneDepthFromViewKm)))))
 				{
