Workaround: For an issue with a VT sampling materials and a half precision in OpenGL on Adreno devices with a V[REDACTED_HANDLE]  drivers

diff --git a/Engine/Shaders/Private/MobileBasePassPixelShader.usf b/Engine/Shaders/Private/MobileBasePassPixelShader.usf
--- a/Engine/Shaders/Private/MobileBasePassPixelShader.usf
+++ b/Engine/Shaders/Private/MobileBasePassPixelShader.usf
@@ -404,7 +404,7 @@ void Main(
 	
 #if !SUBSTRATE_ENABLED
 	// Store the results in local variables and reuse instead of calling the functions multiple times.
-	half3 BaseColor = GetMaterialBaseColor(PixelMaterialInputs);
+	half3 BaseColor = max(GetMaterialBaseColorRaw(PixelMaterialInputs), 0.0); // UE-226780 Instead of using GetMaterialBaseColor which does saturate, use max() to avoid driver issue on Adreno V@05xx
 	half Metallic = GetMaterialMetallic(PixelMaterialInputs);
 	half Specular = GetMaterialSpecular(PixelMaterialInputs);
 	// The smallest normalized value that can be represented in IEEE 754 (FP16) is 2^-24 = 5.96e-8.
