Added StoreAndIncrementOffset functions to ByteBuffer.ush

#rnx

diff --git a/Engine/Shaders/Private/ByteBuffer.ush b/Engine/Shaders/Private/ByteBuffer.ush
--- a/Engine/Shaders/Private/ByteBuffer.ush
+++ b/Engine/Shaders/Private/ByteBuffer.ush
@@ -61,3 +61,55 @@ void LoadAndIncrementOffset(inout float4x4 Value, ByteAddressBuffer InBuffer, in
 	LoadAndIncrementOffset(Value[2], InBuffer, Offset);
 	LoadAndIncrementOffset(Value[3], InBuffer, Offset);
 }
+
+void StoreAndIncrementOffset(uint Value, RWByteAddressBuffer OutBuffer, inout uint Offset)
+{
+	OutBuffer.Store(Offset, Value);
+	Offset += 4;
+}
+
+void StoreAndIncrementOffset(uint2 Value, RWByteAddressBuffer OutBuffer, inout uint Offset)
+{
+	OutBuffer.Store2(Offset, Value);
+	Offset += 8;
+}
+
+void StoreAndIncrementOffset(uint3 Value, RWByteAddressBuffer OutBuffer, inout uint Offset)
+{
+	OutBuffer.Store3(Offset, Value);
+	Offset += 12;
+}
+
+void StoreAndIncrementOffset(uint4 Value, RWByteAddressBuffer OutBuffer, inout uint Offset)
+{
+	OutBuffer.Store4(Offset, Value);
+	Offset += 16;
+}
+
+void StoreAndIncrementOffset(float Value, RWByteAddressBuffer OutBuffer, inout uint Offset)
+{
+	StoreAndIncrementOffset(asuint(Value), OutBuffer, Offset);
+}
+
+void StoreAndIncrementOffset(float2 Value, RWByteAddressBuffer OutBuffer, inout uint Offset)
+{
+	StoreAndIncrementOffset(asuint(Value), OutBuffer, Offset);
+}
+
+void StoreAndIncrementOffset(float3 Value, RWByteAddressBuffer OutBuffer, inout uint Offset)
+{
+	StoreAndIncrementOffset(asuint(Value), OutBuffer, Offset);
+}
+
+void StoreAndIncrementOffset(float4 Value, RWByteAddressBuffer OutBuffer, inout uint Offset)
+{
+	StoreAndIncrementOffset(asuint(Value), OutBuffer, Offset);
+}
+
+void StoreAndIncrementOffset(float4x4 Value, RWByteAddressBuffer OutBuffer, inout uint Offset)
+{	
+	StoreAndIncrementOffset(Value[0], OutBuffer, Offset);
+	StoreAndIncrementOffset(Value[1], OutBuffer, Offset);
+	StoreAndIncrementOffset(Value[2], OutBuffer, Offset);
+	StoreAndIncrementOffset(Value[3], OutBuffer, Offset);
+}
