Fixed Quantization parameter order.

diff --git a/Engine/Shaders/Private/Quantization.ush b/Engine/Shaders/Private/Quantization.ush
--- a/Engine/Shaders/Private/Quantization.ush
+++ b/Engine/Shaders/Private/Quantization.ush
@@ -40,5 +40,5 @@ float3 QuantizeFloatColor(in float3 Color, in float3 QuantizationError, in float
 		E = InterleavedGradientNoise(Pos2D, View.StateFrameIndexMod8);
 	}
 
-	return QuantizeFloatColor(Color, E, QuantizationError);
+	return QuantizeFloatColor(Color, QuantizationError, E);
 }
