VolumetricCloud - added MinimumDistanceKmToEnableDisoclusion.

diff --git a/Engine/Shaders/Private/VolumetricRenderTarget.usf b/Engine/Shaders/Private/VolumetricRenderTarget.usf
--- a/Engine/Shaders/Private/VolumetricRenderTarget.usf
+++ b/Engine/Shaders/Private/VolumetricRenderTarget.usf
@@ -27,6 +27,7 @@ Texture2D<float4> TracingVolumetricDepthTexture;
 uint4 TracingVolumetricTextureValidCoordRect;
 float4 TracingVolumetricTextureValidUvRect;
 float MinimumDistanceKmToEnableReprojection;
+float MinimumDistanceKmToDisableDisoclusion;
 float2 TracingVolumetricTextureUVScale;
 
 struct FDepthData
@@ -308,53 +309,59 @@ void ReconstructVolumetricRenderTargetPS(
 
 		#if PERMUTATION_CLOUD_MIN_AND_MAX_DEPTH
 			
-			const bool bLowScreenVelocity = length(ScreenVelocity) < 0.001;
-
-			const float ThresholdToNewSampleKm = 2.0; // Arbitrary
-
-			if (abs(NewDepthData.MinMaxViewDepthKm.x - NewDepthData.MinMaxViewDepthKm.y) > ThresholdToNewSampleKm 
-				&& (abs(NewDepthData.MinMaxViewDepthKm.x - NewDepthData.SceneDepthFromViewKm) < abs(NewDepthData.MinMaxViewDepthKm.y - NewDepthData.SceneDepthFromViewKm)) )
+			// Do not apply disoclusion if all the traced and reprojected cloud depth are larger. In this case, cloud information will be like a layer blended on top without upsampling.
+			// Otherwise we might be hitting an edge/disoclusion and we will need upsampling.
+			const bool bApplyDisoclusion = any(NewDepthData.MinMaxViewDepthKm < MinimumDistanceKmToDisableDisoclusion) || any(HistoryDepthData.MinMaxViewDepthKm < MinimumDistanceKmToDisableDisoclusion)
+										|| NewDepthData.SceneDepthFromViewKm < MinimumDistanceKmToDisableDisoclusion || HistoryDepthData.SceneDepthFromViewKm < MinimumDistanceKmToDisableDisoclusion;
+			if (bApplyDisoclusion)
 			{
-				// If there is a huge delta of depth for a pixel, use the new sample but only if we are on the closest depth part.
-				// This helps a lot removing cloud over trees and small details
-				bUseNewSample = true;
-			}
-			else 
-				if (HistoryDepthData.SceneDepthFromViewKm < (NewDepthData.SceneDepthFromViewKm - ThresholdToNewSampleKm))
-			{
-				// History is closer than the near cloud tracing this frame. This means an object had move and an new disocluded area is discovered.
-				// So we simply use the new data from this frame according to the new depth
-				bUseNewSample = true;
-			}
-			else 
-				if (HistoryDepthData.SceneDepthFromViewKm > (NewDepthData.SceneDepthFromViewKm + ThresholdToNewSampleKm))
-			{
-				// An area that just go covered (history is invalid because occluded)
-				bUseNewSample = true;
+				const float ThresholdToNewSampleKm = 2.0; // Arbitrary
+
+				if (abs(NewDepthData.MinMaxViewDepthKm.x - NewDepthData.MinMaxViewDepthKm.y) > ThresholdToNewSampleKm
+					&& (abs(NewDepthData.MinMaxViewDepthKm.x - NewDepthData.SceneDepthFromViewKm) < abs(NewDepthData.MinMaxViewDepthKm.y - NewDepthData.SceneDepthFromViewKm)))
+				{
+					// If there is a huge delta of depth for a pixel, use the new sample but only if we are on the closest depth part.
+					// This helps a lot removing cloud over trees and small details
+					bUseNewSample = true;
+				}
+				else
+					if (HistoryDepthData.SceneDepthFromViewKm < (NewDepthData.SceneDepthFromViewKm - ThresholdToNewSampleKm))
+					{
+						// History is closer than the near cloud tracing this frame. This means an object had move and an new disocluded area is discovered.
+						// So we simply use the new data from this frame according to the new depth
+						bUseNewSample = true;
+					}
+				else
+					if (HistoryDepthData.SceneDepthFromViewKm > (NewDepthData.SceneDepthFromViewKm + ThresholdToNewSampleKm))
+					{
+						// An area that just go covered (history is invalid because occluded)
+						bUseNewSample = true;
+					}
+			//	else
+			//		if (NewDepthData.SceneDepthFromViewKm > 1.5)
+			//	{
+			//		if (bUseNewSample || !bValidPreviousUVs)
+			//		{
+			//			RGBA		= SafeLoadTracingVolumetricTexture(int2(SVPos.xy) / DownSampleFactor);
+			//			RGBA2		= SafeLoadSecondaryTracingVolumetricTexture(int2(SVPos.xy) / DownSampleFactor);
+			//			DepthData	= SafeLoadTracingVolumetricDepthTexture(int2(SVPos.xy) / DownSampleFactor);
+			//		}
+			//		else
+			//		{
+			//			RGBA		= SafeSamplePreviousFrameVolumetricTexture(PrevScreenUVs);
+			//			RGBA2		= SafeSamplePreviousFrameVolumetricSecondaryTexture(PrevScreenUVs);
+			//			DepthData	= SafeSamplePreviousFrameVolumetricDepthTexture(PrevScreenUVs);
+			//		}
+			//		OutputRt0 = RGBA;// *float4(1, 0, 0, 1);
+			//		OutputRt1 = RGBA2;// *float4(1, 0, 0, 1);
+			//		OutputRt2 = float4(
+			//			DepthData.CloudFrontDepthFromViewKm, 
+			//			DepthData.SceneDepthFromViewKm, 
+			//			DepthData.MinMaxViewDepthKm);
+			//		return;	
+			//	}
 			}
-		//	else
-		//		if (NewDepthData.SceneDepthFromViewKm > 1.5)
-		//	{
-		//		if (bUseNewSample || !bValidPreviousUVs)
-		//		{
-		//			RGBA		= SafeLoadTracingVolumetricTexture(int2(SVPos.xy) / DownSampleFactor);
-		//			RGBA2		= SafeLoadSecondaryTracingVolumetricTexture(int2(SVPos.xy) / DownSampleFactor);
-		//			DepthData	= SafeLoadTracingVolumetricDepthTexture(int2(SVPos.xy) / DownSampleFactor);
-		//		}
-		//		else
-		//		{
-		//			RGBA		= SafeSamplePreviousFrameVolumetricTexture(PrevScreenUVs);
-		//			RGBA2		= SafeSamplePreviousFrameVolumetricSecondaryTexture(PrevScreenUVs);
-		//			DepthData	= SafeSamplePreviousFrameVolumetricDepthTexture(PrevScreenUVs);
-		//		}
-		//		OutputRt0 = RGBA;// *float4(1, 0, 0, 1);
-		//		OutputRt1 = RGBA2;// *float4(1, 0, 0, 1);
-		//		OutputRt2 = float4(
-		//			DepthData.CloudFrontDepthFromViewKm, 
-		//			DepthData.SceneDepthFromViewKm, 
-		//			DepthData.MinMaxViewDepthKm);
-		//		return;	
-		//	}
+			
 
 			if (bUseNewSample)
 			{
@@ -387,6 +394,12 @@ void ReconstructVolumetricRenderTargetPS(
 				DepthData = SafeSampleTracingVolumetricDepthTexture(ScreenUV);
 			}
 
+			//if (bApplyDisoclusion)
+			//{
+			//	RGBA *= float4(1, 1, 0, 1);
+			//	RGBA2 *= float4(1, 1, 0, 1);
+			//}
+
 		#else // PERMUTATION_CLOUD_MIN_AND_MAX_DEPTH
 
 			if (bUseNewSample)
@@ -632,6 +645,7 @@ float4 SafeSampleVolumetricDepthTexture(float2 UV)
 
 float4 VolumetricTextureSizeAndInvSize;
 float UvOffsetSampleAcceptanceWeight;
+float MinimumDistanceKmToDisableDisoclusion;
 float2 FullResolutionToVolumetricBufferResolutionScale;
 
 float4 ApplyFogToCloudSampleFromVPosAndDepth(float2 SVPos, float DepthKilometer, float4 CloudRGBT)
@@ -1281,11 +1295,15 @@ void ComposeVolumetricRTOverScenePS(
 			//const bool bAllCloudToBilinear = all(float4(VolumeFrontDepth0.x, VolumeFrontDepth1.x, VolumeFrontDepth2.x, VolumeFrontDepth3.x) > 3)
 			//	&& all(float4(VolumeFrontDepth0.x, VolumeFrontDepth1.x, VolumeFrontDepth2.x, VolumeFrontDepth3.x) < float4(VolumeFrontDepth0.z, VolumeFrontDepth1.z, VolumeFrontDepth2.z, VolumeFrontDepth3.z));
 
-			const bool bAllCloudToBilinear = all(float4(VolumeFrontDepth0.x, VolumeFrontDepth1.x, VolumeFrontDepth2.x, VolumeFrontDepth3.x) <= float4(VolumeFrontDepth0.z, VolumeFrontDepth1.z, VolumeFrontDepth2.z, VolumeFrontDepth3.z));
+			const bool bAllCloudToBilinear = 
+				// All cloud samples are in front of the depth buffer
+					all(float4(VolumeFrontDepth0.x, VolumeFrontDepth1.x, VolumeFrontDepth2.x, VolumeFrontDepth3.x) <= float4(VolumeFrontDepth0.z, VolumeFrontDepth1.z, VolumeFrontDepth2.z, VolumeFrontDepth3.z))
+				// All depth values are further than thte distance at which we disable upsampling.
+				 || (all(VolumeFrontDepth0.yzw >= MinimumDistanceKmToDisableDisoclusion) && all(VolumeFrontDepth1.yzw >= MinimumDistanceKmToDisableDisoclusion) && all(VolumeFrontDepth2.yzw >= MinimumDistanceKmToDisableDisoclusion) && all(VolumeFrontDepth3.yzw >= MinimumDistanceKmToDisableDisoclusion));
 
 			if ((all(float2(MinDepth, MaxDepth) > DistanceThresholdToBilinear) && MinMaxDepthRange < DepthDifferenceToBilinear) || bAllCloudToBilinear)
 			{
-				// Bilateral upsampling for a soft look.
+				// Bilinear upsampling for a soft look.
 				float2 SampleUVs = (float2(VolumeCoord0) + 0.25 + (float2(OffsetX, OffsetY) * 0.5 + 0.5) * 0.5) * VolumetricTextureSizeAndInvSize.zw;
 				DataAcc = SafeSampleVolumetricTexture(SampleUVs);
 
@@ -1322,7 +1340,7 @@ void ComposeVolumetricRTOverScenePS(
 		}
 		else
 		{
-			// Bilateral upsampling for a soft look.
+			// Bilinear upsampling for a soft look.
 			ValidSampleCount = 1.0;
 			float2 SampleUVs = (float2(VolumeCoord0) + 0.25 + (float2(OffsetX, OffsetY) * 0.5 + 0.5) * 0.5) * VolumetricTextureSizeAndInvSize.zw;
 			DataAcc = SafeSampleVolumetricTexture(SampleUVs);
diff --git a/Engine/Source/Runtime/Renderer/Private/VolumetricRenderTarget.cpp b/Engine/Source/Runtime/Renderer/Private/VolumetricRenderTarget.cpp
--- a/Engine/Source/Runtime/Renderer/Private/VolumetricRenderTarget.cpp
+++ b/Engine/Source/Runtime/Renderer/Private/VolumetricRenderTarget.cpp
@@ -62,6 +62,11 @@ static TAutoConsoleVariable<float> CVarVolumetricRenderTargetMinimumDistanceKmTo
 	TEXT("This is the distance in kilometer at which the `cloud surface` must be before we enable reprojection of the previous frame data. One could start with a value of 4km. This helps hide reprojection issues due to imperfect approximation of cloud depth as a single front surface, especially visible when flying through the cloud layer. It is not perfect but will help in lots of cases. The problem when using this method: clouds will look noisier when closer to that distance."),
 	ECVF_RenderThreadSafe | ECVF_Scalability);
 
+static TAutoConsoleVariable<float> CVarVolumetricRenderTargetMinimumDistanceKmToDisableDisoclusion(
+	TEXT("r.VolumetricRenderTarget.MinimumDistanceKmToDisableDisoclusion"), 5.0f,
+	TEXT("This is the distance in kilometer at which we stop applying disocclusion, if all the traced and reprojected cloud depth are larger. Otherwise we might be hitting an edge. In this case, cloud information will be like a layer blended on top without upsampling."),
+	ECVF_RenderThreadSafe | ECVF_Scalability);
+
 static TAutoConsoleVariable<int32> CVarVolumetricRenderTargetSimulateNullResolution(
 	TEXT("r.VolumetricRenderTarget.SimulateNullResolution"), 0,
 	TEXT("Simulate a view port resolution of 0x0 for debugging."),
@@ -72,6 +77,11 @@ static float GetUvNoiseSampleAcceptanceWeight()
 	return FMath::Max(0.0f, CVarVolumetricRenderTargetUvNoiseSampleAcceptanceWeight.GetValueOnRenderThread());
 }
 
+static float GetMinimumDistanceKmToDisableDisoclusion()
+{
+	return FMath::Max(0.0f, CVarVolumetricRenderTargetMinimumDistanceKmToDisableDisoclusion.GetValueOnRenderThread());;
+}
+
 static bool ShouldPipelineCompileVolumetricRenderTargetShaders(EShaderPlatform ShaderPlatform)
 {
 	return GetMaxSupportedFeatureLevel(ShaderPlatform) >= ERHIFeatureLevel::SM5;
@@ -719,6 +729,7 @@ class FReconstructVolumetricRenderTargetPS : public FGlobalShader
 		SHADER_PARAMETER(FUintVector4, PreviousFrameVolumetricTextureValidCoordRect)
 		SHADER_PARAMETER(FVector4f, PreviousFrameVolumetricTextureValidUvRect)
 		SHADER_PARAMETER(float, MinimumDistanceKmToEnableReprojection)
+		SHADER_PARAMETER(float, MinimumDistanceKmToDisableDisoclusion)
 		SHADER_PARAMETER(float, HistoryPreExposureCorrection)
 		SHADER_PARAMETER(FVector2f, PreviousFrameVolumetricTextureUVScale)
 	END_SHADER_PARAMETER_STRUCT()
@@ -830,6 +841,7 @@ void ReconstructVolumetricRenderTarget(
 		PassParameters->VolumetricRenderTargetMode = VolumetricCloudRT.GetMode();
 		PassParameters->HalfResDepthTexture = (VolumetricCloudRT.GetMode() == 0 || VolumetricCloudRT.GetMode() == 3) ? HalfResolutionDepthCheckerboardMinMaxTexture : SceneDepthTexture;
 		PassParameters->MinimumDistanceKmToEnableReprojection = FMath::Max(0.0f, CVarVolumetricRenderTargetMinimumDistanceKmToEnableReprojection.GetValueOnRenderThread());
+		PassParameters->MinimumDistanceKmToDisableDisoclusion = GetMinimumDistanceKmToDisableDisoclusion();
 		PassParameters->HistoryPreExposureCorrection = ViewInfo.PreExposure / VolumetricCloudRT.GetPrevViewExposure();
 
 		const bool bVisualizeConservativeDensity = ShouldViewVisualizeVolumetricCloudConservativeDensity(ViewInfo, ViewInfo.Family->EngineShowFlags);
@@ -892,6 +904,7 @@ class FComposeVolumetricRTOverScenePS : public FGlobalShader
 		SHADER_PARAMETER(uint32, OutputAlphaHoldout)
 		SHADER_PARAMETER(float, VolumeTracingStartDistanceFromCamera)
 		SHADER_PARAMETER(float, UvOffsetSampleAcceptanceWeight)
+		SHADER_PARAMETER(float, MinimumDistanceKmToDisableDisoclusion)
 		SHADER_PARAMETER(FVector4f, VolumetricTextureSizeAndInvSize)
 		SHADER_PARAMETER(FVector2f, FullResolutionToVolumetricBufferResolutionScale)
 		SHADER_PARAMETER(FVector2f, FullResolutionToWaterBufferScale)
@@ -1055,6 +1068,7 @@ void ComposeVolumetricRenderTargetOverScene(
 		PassParameters->VolumetricDepthTexture = VolumetricDepthTexture;
 		PassParameters->LinearTextureSampler = TStaticSamplerState<SF_Bilinear>::GetRHI();
 		PassParameters->UvOffsetSampleAcceptanceWeight = GetUvNoiseSampleAcceptanceWeight();
+		PassParameters->MinimumDistanceKmToDisableDisoclusion = GetMinimumDistanceKmToDisableDisoclusion();
 		PassParameters->FullResolutionToVolumetricBufferResolutionScale = GetCompositionFullResolutionToVolumetricBufferResolutionScale(VRTMode) * GetVolumetricBufferResolutionScale(VRTMode);
 		PassParameters->SceneTextures = SceneTextures.UniformBuffer;
 		PassParameters->Substrate = Substrate::BindSubstrateGlobalUniformParameters(ViewInfo);
@@ -1175,6 +1189,7 @@ void ComposeVolumetricRenderTargetOverSceneUnderWater(
 		PassParameters->LinearTextureSampler = TStaticSamplerState<SF_Bilinear>::GetRHI();
 		PassParameters->WaterLinearDepthSampler = TStaticSamplerState<SF_Point>::GetRHI();
 		PassParameters->UvOffsetSampleAcceptanceWeight = GetUvNoiseSampleAcceptanceWeight();
+		PassParameters->MinimumDistanceKmToDisableDisoclusion = GetMinimumDistanceKmToDisableDisoclusion();
 		PassParameters->FullResolutionToVolumetricBufferResolutionScale = GetCompositionFullResolutionToVolumetricBufferResolutionScale(VRTMode) * GetVolumetricBufferResolutionScale(VRTMode);
 		PassParameters->FullResolutionToWaterBufferScale = FVector2f(1.0f / WaterPassData.RefractionDownsampleFactor, WaterPassData.RefractionDownsampleFactor);
 		PassParameters->SceneWithoutSingleLayerWaterViewRect = FVector4f(WaterPassViewData.ViewRect.Min.X, WaterPassViewData.ViewRect.Min.Y,
@@ -1261,6 +1276,7 @@ void ComposeVolumetricRenderTargetOverSceneForVisualization(
 		PassParameters->VolumetricDepthTexture = VolumetricDepthTexture;
 		PassParameters->LinearTextureSampler = TStaticSamplerState<SF_Bilinear>::GetRHI();
 		PassParameters->UvOffsetSampleAcceptanceWeight = GetUvNoiseSampleAcceptanceWeight();
+		PassParameters->MinimumDistanceKmToDisableDisoclusion = GetMinimumDistanceKmToDisableDisoclusion();
 		PassParameters->FullResolutionToVolumetricBufferResolutionScale = GetCompositionFullResolutionToVolumetricBufferResolutionScale(VRTMode) * GetVolumetricBufferResolutionScale(VRTMode);
 		PassParameters->SceneTextures = SceneTextures.UniformBuffer;
 		PassParameters->Substrate = Substrate::BindSubstrateGlobalUniformParameters(ViewInfo);
