Refactor material translation of CustomData0/CustomData1.
Previous version filled the translated Getter function with a copy of TranslatedCodeChunkDefinitions code.
Often the contents of TranslatedCodeChunkDefinitions wasn't needed (it's there to fulfil potential dependencies to the TranslatedCodeChunks). But we rely on the final shader compilation to strip the dead code.
When the contents contain a virtual texture sample, the shader compiler can't strip that out, because the sample also has a feedback UAV write as a side effect.
This leads to extra sampling code, and the added cost of writing a UAV.
In addition the TranslatedCodeChunkDefinitions wasn't added to the translated code with the correct ECompiledPartialDerivativeVariation setting. This could lead to runtime virtual texture samples triggering mip 0 feedback, which blows out the virtual texture residency.
The fix is to treat CustomData0 and CustomData1 as normal members of PixelMaterialInputs so that their translated generation code can be inlined with the rest of the PixelMaterialInput translated generation code and doesn't need full copying into a function.

diff --git a/Engine/Shaders/Private/BasePassPixelShader.usf b/Engine/Shaders/Private/BasePassPixelShader.usf
--- a/Engine/Shaders/Private/BasePassPixelShader.usf
+++ b/Engine/Shaders/Private/BasePassPixelShader.usf
@@ -1076,6 +1076,7 @@ void FPixelShaderInOut_MainPS(
 	SetGBufferForShadingModel(
 		GBuffer,
 		MaterialParameters,
+		PixelMaterialInputs,
 		Opacity,
 		BaseColor,
 		Metallic,
@@ -1247,7 +1248,7 @@ void FPixelShaderInOut_MainPS(
 		#if MATERIAL_SHADINGMODEL_CLOTH
 		if (GBuffer.ShadingModelID == SHADINGMODELID_CLOTH)
 		{
-			DiffuseColorForIndirect += SubsurfaceColor * saturate(GetMaterialCustomData0(MaterialParameters));
+			DiffuseColorForIndirect += SubsurfaceColor * saturate(GetMaterialCustomData0(PixelMaterialInputs));
 		}
 		#endif
 
@@ -1808,6 +1809,7 @@ void FPixelShaderInOut_MainPS(
 		SetGBufferForShadingModel(
 			GBuffer,
 			MaterialParameters,
+			PixelMaterialInputs,
 			Export.Coverage,		// Opacity
 			Export.BaseColor,
 			Export.Metallic,
diff --git a/Engine/Shaders/Private/DebugViewModePixelShader.usf b/Engine/Shaders/Private/DebugViewModePixelShader.usf
--- a/Engine/Shaders/Private/DebugViewModePixelShader.usf
+++ b/Engine/Shaders/Private/DebugViewModePixelShader.usf
@@ -260,8 +260,8 @@ float4 VisualizeMaterialTexCoordScales(in FDebugPSIn DebugInputs, in bool bIsFro
 		}
 #endif
 		half4 SSData = GetMaterialSubsurfaceDataRaw(PixelMaterialInputs);
-		float Custom0 = GetMaterialCustomData0(MaterialParameters);
-		float Custom1 = GetMaterialCustomData1(MaterialParameters);
+		float Custom0 = GetMaterialCustomData0(PixelMaterialInputs);
+		float Custom1 = GetMaterialCustomData1(PixelMaterialInputs);
 		float MaterialAO = GetMaterialAmbientOcclusionRaw(PixelMaterialInputs);
 		float PixelDepthOffset = GetMaterialPixelDepthOffset(PixelMaterialInputs);
 #if CLEAR_COAT_BOTTOM_NORMAL && NUM_MATERIAL_OUTPUTS_CLEARCOATBOTTOMNORMAL > 0
@@ -384,8 +384,8 @@ float4 VisualizeRequiredTextureResolution(in FDebugPSIn DebugInputs, in bool bIs
 		clip(GetMaterialMask(PixelMaterialInputs));
 #endif
 		half4 SSData = GetMaterialSubsurfaceDataRaw(PixelMaterialInputs);
-		float Custom0 = GetMaterialCustomData0(MaterialParameters);
-		float Custom1 = GetMaterialCustomData1(MaterialParameters);
+		float Custom0 = GetMaterialCustomData0(PixelMaterialInputs);
+		float Custom1 = GetMaterialCustomData1(PixelMaterialInputs);
 		float MaterialAO = GetMaterialAmbientOcclusionRaw(PixelMaterialInputs);
 		float PixelDepthOffset = GetMaterialPixelDepthOffset(PixelMaterialInputs);
 #if CLEAR_COAT_BOTTOM_NORMAL && NUM_MATERIAL_OUTPUTS_CLEARCOATBOTTOMNORMAL > 0
diff --git a/Engine/Shaders/Private/HairStrands/HairStrandsMaterialCommon.ush b/Engine/Shaders/Private/HairStrands/HairStrandsMaterialCommon.ush
--- a/Engine/Shaders/Private/HairStrands/HairStrandsMaterialCommon.ush
+++ b/Engine/Shaders/Private/HairStrands/HairStrandsMaterialCommon.ush
@@ -256,7 +256,7 @@ FEvaluateOutput Evaluate(
 		OutSample.BaseColor = GetMaterialBaseColor(PixelMaterialInputs);
 		OutSample.Roughness = GetMaterialRoughness(PixelMaterialInputs);
 		OutSample.Specular	= GetMaterialSpecular(PixelMaterialInputs);
-		OutSample.Backlit	= saturate(GetMaterialCustomData0(MaterialParameters));
+		OutSample.Backlit	= saturate(GetMaterialCustomData0(PixelMaterialInputs));
 		OutSample.Emissive	= GetMaterialEmissive(PixelMaterialInputs);
 		float3 Tangent		= normalize(GetMaterialNormalRaw(PixelMaterialInputs));
 
diff --git a/Engine/Shaders/Private/MaterialTemplate.ush b/Engine/Shaders/Private/MaterialTemplate.ush
--- a/Engine/Shaders/Private/MaterialTemplate.ush
+++ b/Engine/Shaders/Private/MaterialTemplate.ush
@@ -3973,14 +3973,14 @@ MaterialFloat4 GetMaterialSubsurfaceData(FPixelMaterialInputs PixelMaterialInput
 SHADER_PUSH_WARNINGS_STATE
 SHADER_DISABLE_WARNINGS
 
-half GetMaterialCustomData0(in out FMaterialPixelParameters Parameters)
+half GetMaterialCustomData0(FPixelMaterialInputs PixelMaterialInputs)
 {
-%{get_material_custom_data0};
+	return PixelMaterialInputs.CustomData0;
 }
 
-half GetMaterialCustomData1(in out FMaterialPixelParameters Parameters)
+half GetMaterialCustomData1(FPixelMaterialInputs PixelMaterialInputs)
 {
-%{get_material_custom_data1};
+	return PixelMaterialInputs.CustomData1;
 }
 
 SHADER_POP_WARNINGS_STATE
diff --git a/Engine/Shaders/Private/MobileBasePassPixelShader.usf b/Engine/Shaders/Private/MobileBasePassPixelShader.usf
--- a/Engine/Shaders/Private/MobileBasePassPixelShader.usf
+++ b/Engine/Shaders/Private/MobileBasePassPixelShader.usf
@@ -518,6 +518,7 @@ void Main(
 	SetGBufferForShadingModel(
 		GBuffer,
 		MaterialParameters,
+		PixelMaterialInputs,
 		Opacity,
 		BaseColor,
 		Metallic,
@@ -721,7 +722,7 @@ void Main(
 			SurfaceWorldNormal,
 			MaterialParameters.WorldPosition_CamRelative,
 		#if SHADING_PATH_MOBILE && MATERIAL_SUBSURFACE_PROFILE_USE_CURVATURE
-			GetMaterialCustomData0(MaterialParameters)
+			GetMaterialCustomData0(PixelMaterialInputs)
 		#else
 			0.0f
 		#endif
@@ -741,6 +742,7 @@ void Main(
 		SetGBufferForShadingModel(
 			GBuffer,
 			MaterialParameters,
+			PixelMaterialInputs,
 			Export.Coverage,		// Opacity
 			Export.BaseColor,
 			Export.Metallic,
@@ -842,7 +844,7 @@ void Main(
 	#if MATERIAL_SHADINGMODEL_CLOTH
 		if (GBuffer.ShadingModelID == SHADINGMODELID_CLOTH)
 		{
-			DiffuseColorForIndirect += SubsurfaceColor * saturate(GetMaterialCustomData0(MaterialParameters));
+			DiffuseColorForIndirect += SubsurfaceColor * saturate(GetMaterialCustomData0(PixelMaterialInputs));
 		}
 	#endif
 
@@ -1026,7 +1028,7 @@ void Main(
 #if MATERIAL_SHADINGMODEL_CLOTH
 	if (GBuffer.ShadingModelID == SHADINGMODELID_CLOTH)
 	{
-		DiffuseColorForIndirect += SubsurfaceColor * saturate(GetMaterialCustomData0(MaterialParameters));
+		DiffuseColorForIndirect += SubsurfaceColor * saturate(GetMaterialCustomData0(PixelMaterialInputs));
 	}
 #endif
 
diff --git a/Engine/Shaders/Private/PathTracing/PathTracingMaterialHitShader.usf b/Engine/Shaders/Private/PathTracing/PathTracingMaterialHitShader.usf
--- a/Engine/Shaders/Private/PathTracing/PathTracingMaterialHitShader.usf
+++ b/Engine/Shaders/Private/PathTracing/PathTracingMaterialHitShader.usf
@@ -905,8 +905,8 @@ RAY_TRACING_ENTRY_CLOSEST_HIT(PathTracingMaterialCHS,
 #if MATERIAL_SHADINGMODEL_CLEAR_COAT
 	if (Payload.ShadingModelID == SHADINGMODELID_CLEAR_COAT)
 	{
-		const float ClearCoat = GetMaterialCustomData0(MaterialParameters);
-		const float ClearCoatRoughness = GetMaterialCustomData1(MaterialParameters);
+		const float ClearCoat = GetMaterialCustomData0(PixelMaterialInputs);
+		const float ClearCoatRoughness = GetMaterialCustomData1(PixelMaterialInputs);
 		Payload.SetClearCoat(ClearCoat);
 		Payload.SetClearCoatRoughness(AdjustMaterialRoughness(ClearCoatRoughness, PackedPayload.GetPathRoughness()));
 #if CLEAR_COAT_BOTTOM_NORMAL
@@ -935,7 +935,7 @@ RAY_TRACING_ENTRY_CLOSEST_HIT(PathTracingMaterialCHS,
 	if (Payload.ShadingModelID == SHADINGMODELID_CLOTH)
 	{
 		const float3 ClothColor = GetMaterialSubsurfaceData(PixelMaterialInputs).rgb;
-		const float Fuzz = saturate(GetMaterialCustomData0(MaterialParameters));
+		const float Fuzz = saturate(GetMaterialCustomData0(PixelMaterialInputs));
 		Payload.SetClothColor(ClothColor);
 		Payload.SetClothAmount(Fuzz);
 	}
@@ -1113,8 +1113,8 @@ RAY_TRACING_ENTRY_CLOSEST_HIT(PathTracingMaterialCHS,
 	if (Payload.ShadingModelID == SHADINGMODELID_EYE)
 	{
 		// This is all based on logic from ShadingModelsMaterial.ush
-		const float IrisMask     = saturate(GetMaterialCustomData0(MaterialParameters));
-		const float IrisDistance = saturate(GetMaterialCustomData1(MaterialParameters));
+		const float IrisMask     = saturate(GetMaterialCustomData0(PixelMaterialInputs));
+		const float IrisDistance = saturate(GetMaterialCustomData1(PixelMaterialInputs));
 		Payload.SetEyeIrisMask(IrisMask);
 		Payload.Metallic = 0.0;
 #if NUM_MATERIAL_OUTPUTS_GETTANGENTOUTPUT > 0
diff --git a/Engine/Shaders/Private/RayTracing/RayTracingMaterialHitShaders.usf b/Engine/Shaders/Private/RayTracing/RayTracingMaterialHitShaders.usf
--- a/Engine/Shaders/Private/RayTracing/RayTracingMaterialHitShaders.usf
+++ b/Engine/Shaders/Private/RayTracing/RayTracingMaterialHitShaders.usf
@@ -429,7 +429,7 @@ void GetMaterialPayload(
 	half3 EmissiveColor = GetMaterialEmissive(PixelMaterialInputs);
 	half3 WorldNormal = normalize(MaterialParameters.WorldNormal);
 	half  Opacity = GetMaterialOpacity(PixelMaterialInputs);
-	half4 CustomData = half4(GetMaterialCustomData0(MaterialParameters), GetMaterialCustomData1(MaterialParameters), 0, 0);
+	half4 CustomData = half4(GetMaterialCustomData0(PixelMaterialInputs), GetMaterialCustomData1(PixelMaterialInputs), 0, 0);
 	half4 SubsurfaceData = GetMaterialSubsurfaceData(PixelMaterialInputs);
 	half4 ClothData = float4(SubsurfaceData.rgb, CustomData.x);
 
diff --git a/Engine/Shaders/Private/ShadingModelsMaterial.ush b/Engine/Shaders/Private/ShadingModelsMaterial.ush
--- a/Engine/Shaders/Private/ShadingModelsMaterial.ush
+++ b/Engine/Shaders/Private/ShadingModelsMaterial.ush
@@ -12,6 +12,7 @@
 void SetGBufferForShadingModel(
 	in out FGBufferData GBuffer, 
 	in out FMaterialPixelParameters MaterialParameters,
+	FPixelMaterialInputs PixelMaterialInputs,
 	const float Opacity,
 	const half3 BaseColor,
 	const half  Metallic,
@@ -65,7 +66,7 @@ void SetGBufferForShadingModel(
 
 		#if SHADING_PATH_MOBILE
 			#if MATERIAL_SUBSURFACE_PROFILE_USE_CURVATURE
-				GBuffer.Curvature = GetMaterialCustomData0(MaterialParameters);
+				GBuffer.Curvature = GetMaterialCustomData0(PixelMaterialInputs);
 			#else
 				GBuffer.Curvature = CalculateCurvature(GBuffer.WorldNormal, MaterialParameters.WorldPosition_CamRelative);
 			#endif
@@ -85,8 +86,8 @@ void SetGBufferForShadingModel(
 #if MATERIAL_SHADINGMODEL_CLEAR_COAT
 	else if (ShadingModel == SHADINGMODELID_CLEAR_COAT)
 	{
-		float ClearCoat				= saturate( GetMaterialCustomData0(MaterialParameters) );
-		float ClearCoatRoughness	= saturate( GetMaterialCustomData1(MaterialParameters) );
+		float ClearCoat				= saturate( GetMaterialCustomData0(PixelMaterialInputs) );
+		float ClearCoatRoughness	= saturate( GetMaterialCustomData1(PixelMaterialInputs) );
 		GBuffer.CustomData.x = ClearCoat;
 		GBuffer.CustomData.y = ClearCoatRoughness;
 
@@ -127,22 +128,22 @@ void SetGBufferForShadingModel(
 	else if (ShadingModel == SHADINGMODELID_HAIR)
 	{
 		GBuffer.CustomData.xy = UnitVectorToOctahedron( MaterialParameters.WorldNormal ) * 0.5 + 0.5;
-		GBuffer.CustomData.z = saturate( GetMaterialCustomData0(MaterialParameters) );	// Backlit
+		GBuffer.CustomData.z = saturate( GetMaterialCustomData0(PixelMaterialInputs) );	// Backlit
 	}
 #endif
 #if MATERIAL_SHADINGMODEL_CLOTH
 	else if (ShadingModel == SHADINGMODELID_CLOTH)
 	{
 		GBuffer.CustomData.rgb = EncodeSubsurfaceColor(SubsurfaceColor);
-		GBuffer.CustomData.a = saturate( GetMaterialCustomData0(MaterialParameters) );	// Cloth
+		GBuffer.CustomData.a = saturate( GetMaterialCustomData0(PixelMaterialInputs) );	// Cloth
 		GBuffer.IndirectIrradiance *= 1 - GBuffer.CustomData.a;
 	}
 #endif
 #if MATERIAL_SHADINGMODEL_EYE
 	else if (ShadingModel == SHADINGMODELID_EYE)
 	{
-		const float IrisMask     = saturate(GetMaterialCustomData0(MaterialParameters));
-		const float IrisDistance = saturate(GetMaterialCustomData1(MaterialParameters));
+		const float IrisMask     = saturate(GetMaterialCustomData0(PixelMaterialInputs));
+		const float IrisDistance = saturate(GetMaterialCustomData1(PixelMaterialInputs));
 
 		GBuffer.CustomData.x = EncodeSubsurfaceProfile(SubsurfaceProfile).x;
 		GBuffer.CustomData.w = 1.0f - IrisMask;	// Opacity
diff --git a/Engine/Source/Runtime/Core/Private/UObject/FortniteMainBranchObjectVersion.cpp b/Engine/Source/Runtime/Core/Private/UObject/FortniteMainBranchObjectVersion.cpp
--- a/Engine/Source/Runtime/Core/Private/UObject/FortniteMainBranchObjectVersion.cpp
+++ b/Engine/Source/Runtime/Core/Private/UObject/FortniteMainBranchObjectVersion.cpp
@@ -14,6 +14,6 @@ TMap<FGuid, FGuid> FFortniteMainBranchObjectVersion::GetSystemGuids()
 	SystemGuids.Add(DevGuids.Niagara_LatestScriptCompileVersion, FGuid("081BBC4CD785874B8E77E9AEAEBD0537"));
 	SystemGuids.Add(DevGuids.SkeletalMeshDerivedDataVersion, FGuid("716424BCEC314E969ACFE36D2C2ECE49"));
 	SystemGuids.Add(DevGuids.STATICMESH_DERIVEDDATA_VER, FGuid("3ABF32149F9D4D838750368AB95573BE"));
-	SystemGuids.Add(DevGuids.MaterialTranslationDDCVersion, FGuid("84CF0B79D69D4A4BBE8CAFD4B3D48AB7"));
+	SystemGuids.Add(DevGuids.MaterialTranslationDDCVersion, FGuid("AAAFC530186D4149A06F027D4A060F5E"));
 	return SystemGuids;
 }
diff --git a/Engine/Source/Runtime/Engine/Private/Materials/HLSLMaterialTranslator.cpp b/Engine/Source/Runtime/Engine/Private/Materials/HLSLMaterialTranslator.cpp
--- a/Engine/Source/Runtime/Engine/Private/Materials/HLSLMaterialTranslator.cpp
+++ b/Engine/Source/Runtime/Engine/Private/Materials/HLSLMaterialTranslator.cpp
@@ -764,6 +764,8 @@ FHLSLMaterialTranslator::FHLSLMaterialTranslator(FMaterial* InMaterial,
 	SharedPixelProperties[MP_SurfaceThickness] = true;
 	SharedPixelProperties[MP_FrontMaterial] = true;
 	SharedPixelProperties[MP_Displacement] = true;
+	SharedPixelProperties[MP_CustomData0] = true;
+	SharedPixelProperties[MP_CustomData1] = true;
 
 	for (int32 Frequency = 0; Frequency < SF_NumFrequencies; ++Frequency)
 	{
@@ -16805,8 +16807,6 @@ void FHLSLMaterialTranslator::PrepareMaterialSourceStringParameters()
 
 	MaterialSourceTemplateParams.Add({ TEXT("get_material_world_position_offset_raw"), *GenerateFunctionCode(MP_WorldPositionOffset, BaseDerivativeVariation) });
 	MaterialSourceTemplateParams.Add({ TEXT("get_material_previous_world_position_offset_raw"), *GenerateFunctionCode(CompiledMP_PrevWorldPositionOffset, BaseDerivativeVariation) });
-	MaterialSourceTemplateParams.Add({ TEXT("get_material_custom_data0"), *GenerateFunctionCode(MP_CustomData0, BaseDerivativeVariation) });
-	MaterialSourceTemplateParams.Add({ TEXT("get_material_custom_data1"), *GenerateFunctionCode(MP_CustomData1, BaseDerivativeVariation) });
 
 	// Print custom texture coordinate assignments, should be fine with regular derivatives
 	FString CustomUVAssignments;
diff --git a/Engine/Source/Runtime/Engine/Private/Materials/MaterialAttributeDefinitionMap.cpp b/Engine/Source/Runtime/Engine/Private/Materials/MaterialAttributeDefinitionMap.cpp
--- a/Engine/Source/Runtime/Engine/Private/Materials/MaterialAttributeDefinitionMap.cpp
+++ b/Engine/Source/Runtime/Engine/Private/Materials/MaterialAttributeDefinitionMap.cpp
@@ -344,8 +344,8 @@ void FMaterialAttributeDefinitionMap::InitializeAttributeMap()
 	Add(FGuid(0xF905F895, 0xD5814314, 0x916D2434, 0x8C40CE9E), TEXT("WorldPositionOffset"),		MP_WorldPositionOffset,		MCT_Float3,	FVector4(0,0,0,0),	SF_Vertex);
 	Add(FGuid(0x199A7166, 0xC67041DC, 0xA68EAD0D, 0x7017D0AD), TEXT("Displacement"),			MP_Displacement,			MCT_Float,	FVector4(0.5,0,0,0), SF_Pixel);
 	Add(FGuid(0x5B8FC679, 0x51CE4082, 0x9D777BEE, 0xF4F72C44), TEXT("SubsurfaceColor"),			MP_SubsurfaceColor,			MCT_Float3,	FVector4(1,1,1,0),	SF_Pixel);
-	Add(FGuid(0x9E502E69, 0x3C8F48FA, 0x94645CFD, 0x28E5428D), TEXT("ClearCoat"),				MP_CustomData0,				MCT_Float,	FVector4(1,0,0,0),	SF_Pixel);
-	Add(FGuid(0xBE4F2FFD, 0x12FC4296, 0xB0124EEA, 0x12C28D92), TEXT("ClearCoatRoughness"),		MP_CustomData1,				MCT_Float,	FVector4(.1,0,0,0),	SF_Pixel);
+	Add(FGuid(0x9E502E69, 0x3C8F48FA, 0x94645CFD, 0x28E5428D), TEXT("CustomData0"),				MP_CustomData0,				MCT_Float,	FVector4(1,0,0,0),	SF_Pixel);
+	Add(FGuid(0xBE4F2FFD, 0x12FC4296, 0xB0124EEA, 0x12C28D92), TEXT("CustomData1"),				MP_CustomData1,				MCT_Float,	FVector4(.1,0,0,0),	SF_Pixel);
 	Add(FGuid(0xE8EBD0AD, 0xB1654CBE, 0xB079C3A8, 0xB39B9F15), TEXT("AmbientOcclusion"),		MP_AmbientOcclusion,		MCT_Float,	FVector4(1,0,0,0),	SF_Pixel);
 	Add(FGuid(0xD0B0FA03, 0x14D74455, 0xA851BAC5, 0x81A0788B), TEXT("Refraction"),				MP_Refraction,				MCT_Float3,	FVector4(1,0,0,0),	SF_Pixel);
 	Add(FGuid(0x0AC97EC3, 0xE3D047BA, 0xB610167D, 0xC4D919FF), TEXT("PixelDepthOffset"),		MP_PixelDepthOffset,		MCT_Float,	FVector4(0,0,0,0),	SF_Pixel);
diff --git a/Engine/Source/Runtime/Engine/Private/Materials/MaterialIRInternal.cpp b/Engine/Source/Runtime/Engine/Private/Materials/MaterialIRInternal.cpp
--- a/Engine/Source/Runtime/Engine/Private/Materials/MaterialIRInternal.cpp
+++ b/Engine/Source/Runtime/Engine/Private/Materials/MaterialIRInternal.cpp
@@ -33,6 +33,8 @@ bool IsMaterialPropertyEnabled(EMaterialProperty InProperty)
 		case MP_FrontMaterial:
 		case MP_Displacement:
 		case MP_WorldPositionOffset:
+		case MP_CustomData0:
+		case MP_CustomData1:
 			return true;
 		default:
 			return false;
diff --git a/Engine/Source/Runtime/Engine/Private/Materials/MaterialIRToHLSLTranslator.cpp b/Engine/Source/Runtime/Engine/Private/Materials/MaterialIRToHLSLTranslator.cpp
--- a/Engine/Source/Runtime/Engine/Private/Materials/MaterialIRToHLSLTranslator.cpp
+++ b/Engine/Source/Runtime/Engine/Private/Materials/MaterialIRToHLSLTranslator.cpp
@@ -1192,10 +1192,6 @@ struct FTranslator : FMaterialIRToHLSLTranslation
 		Params.Add(TEXT("get_material_world_position_offset_raw"), WorldPositionOffsetHLSL);
 		Params.Add(TEXT("get_material_previous_world_position_offset_raw"), WorldPositionOffsetHLSL);
 	
-		// CustomData0/1 are named ClearCoat/ClearCoatRoughness
-		Params.Add(TEXT("get_material_custom_data0"), TEXT(TAB "return 1.0f; // todo"));
-		Params.Add(TEXT("get_material_custom_data1"), TEXT(TAB "return 0.1f; // todo"));
-
 		FString EvaluateMaterialDeclaration;
 		EvaluateMaterialDeclaration.Append(TEXT("void EvaluateVertexMaterialAttributes(in out FMaterialVertexParameters Parameters)\n{\n"));
 		for (int32 CustomUVIndex = 0; CustomUVIndex < ModuleStatistics.NumPixelTexCoords; CustomUVIndex++)
