VolumetricCloud - We now always bilinear upsample cloud in front of min depth when there is an edge detected, or when further than the distance threashold.

diff --git a/Engine/Shaders/Private/VolumetricRenderTarget.usf b/Engine/Shaders/Private/VolumetricRenderTarget.usf
--- a/Engine/Shaders/Private/VolumetricRenderTarget.usf
+++ b/Engine/Shaders/Private/VolumetricRenderTarget.usf
@@ -1288,20 +1288,24 @@ void ComposeVolumetricRTOverScenePS(
 
 
 		// TODO parameterise from cvar
-		const float DistanceThresholdToBilinear	= 2.0;// Km. Could be parameterised or based on distance from view, max(PixelDistanceFromViewKm, 2.0f) * 0.1;
-		const float DepthDifferenceToBilinear	= 1.0;// IDEM
+		const float DepthDifferenceToBilinearKm	= 1.0;
 		if (!bIsWaterPixel)
 		{
-			//const bool bAllCloudToBilinear = all(float4(VolumeFrontDepth0.x, VolumeFrontDepth1.x, VolumeFrontDepth2.x, VolumeFrontDepth3.x) > 3)
-			//	&& all(float4(VolumeFrontDepth0.x, VolumeFrontDepth1.x, VolumeFrontDepth2.x, VolumeFrontDepth3.x) < float4(VolumeFrontDepth0.z, VolumeFrontDepth1.z, VolumeFrontDepth2.z, VolumeFrontDepth3.z));
+			// All cloud samples are in front of the min depth buffer
+			const bool bAllCloudInFrontOfMinDepth = all(float4(VolumeFrontDepth0.x, VolumeFrontDepth1.x, VolumeFrontDepth2.x, VolumeFrontDepth3.x) <= float4(VolumeFrontDepth0.z, VolumeFrontDepth1.z, VolumeFrontDepth2.z, VolumeFrontDepth3.z));
 
 			const bool bAllCloudToBilinear = 
-				// All cloud samples are in front of the depth buffer
-					all(float4(VolumeFrontDepth0.x, VolumeFrontDepth1.x, VolumeFrontDepth2.x, VolumeFrontDepth3.x) <= float4(VolumeFrontDepth0.z, VolumeFrontDepth1.z, VolumeFrontDepth2.z, VolumeFrontDepth3.z))
+						bAllCloudInFrontOfMinDepth
 				// All depth values are further than thte distance at which we disable upsampling.
-				 || (all(VolumeFrontDepth0.yzw >= MinimumDistanceKmToDisableDisoclusion) && all(VolumeFrontDepth1.yzw >= MinimumDistanceKmToDisableDisoclusion) && all(VolumeFrontDepth2.yzw >= MinimumDistanceKmToDisableDisoclusion) && all(VolumeFrontDepth3.yzw >= MinimumDistanceKmToDisableDisoclusion));
-
-			if ((all(float2(MinDepth, MaxDepth) > DistanceThresholdToBilinear) && MinMaxDepthRange < DepthDifferenceToBilinear) || bAllCloudToBilinear)
+					&& (all(VolumeFrontDepth0.yzw >= MinimumDistanceKmToDisableDisoclusion) && all(VolumeFrontDepth1.yzw >= MinimumDistanceKmToDisableDisoclusion) 
+						&& all(VolumeFrontDepth2.yzw >= MinimumDistanceKmToDisableDisoclusion) && all(VolumeFrontDepth3.yzw >= MinimumDistanceKmToDisableDisoclusion));
+			// => Note: We could not enable bilinear sampling when all front cloud depths are in front of scene depth ONLY, 
+			//			This is because if a mountain is visible through and intersecting a cloud not too far from the camera, low resolution texels might become visible.
+			//			So unfortunately this does not work for an isolated cloud in front of a mountain.
+			//			That is why we use the && operator. 
+			// => The case of mountain visible inside a cloud is handled by the other test below. 
+
+			if ((bAllCloudInFrontOfMinDepth && (MinMaxDepthRange < DepthDifferenceToBilinearKm)) || bAllCloudToBilinear)
 			{
 				// Bilinear upsampling for a soft look.
 				float2 SampleUVs = (float2(VolumeCoord0) + 0.25 + (float2(OffsetX, OffsetY) * 0.5 + 0.5) * 0.5) * VolumetricTextureSizeAndInvSize.zw;
