[Auto Exposure] Updated desktop and mobile auto exposure to use the same logic when applying r.EyeAdaptation.BlackHistogramBucketInfluence.
- also apply the logic when using Basic method (not just Histogram).

diff --git a/Engine/Shaders/Private/PostProcessEyeAdaptation.usf b/Engine/Shaders/Private/PostProcessEyeAdaptation.usf
--- a/Engine/Shaders/Private/PostProcessEyeAdaptation.usf
+++ b/Engine/Shaders/Private/PostProcessEyeAdaptation.usf
@@ -101,11 +101,23 @@ float2 ComputeWeightedTextureAverageAlphaSubRegion(
 
 			float Weight = max(AdaptationWeightTexture(ScreenUV),0.05f);
 
-			RowWeight += Weight;
-
 			// Accumulate values from alpha channel.
 			float Sample = Texture.Load(int3(i, j, 0)).w;
+
+			{
+				// apply EyeAdaptation_BlackHistogramBucketInfluence using similar logic as Histogram method
+				float fBucket = saturate(Sample) * (HISTOGRAM_SIZE - 1);
+
+				// Find two discrete buckets that straddle the continuous histogram position.
+				uint Bucket0 = (uint)(fBucket);
+				if (Bucket0 == 0)
+				{
+					Weight *= EyeAdaptation_BlackHistogramBucketInfluence;
+				}
+			}
+
 			RowValue += Weight * Sample;
+			RowWeight += Weight;
 		}
 
 		Value += RowValue;
diff --git a/Engine/Shaders/Private/PostProcessMobile.usf b/Engine/Shaders/Private/PostProcessMobile.usf
--- a/Engine/Shaders/Private/PostProcessMobile.usf
+++ b/Engine/Shaders/Private/PostProcessMobile.usf
@@ -456,13 +456,7 @@ void BloomPS_Mobile(
 #endif
 
 #if MOBILE_USEEYEADAPTATION
-	float Intensity = dot(AverageColor * View.OneOverPreExposure, float3(1.0f, 1.0f, 1.0f) / 3.0f);
-	if (Intensity == 0.0)
-	{
-		// Override intensity for black pixels to MidGray. This behaviour is controlled by the cvar "r.EyeAdaptation.BlackHistogramBucketInfluence"
-		Intensity = (0.18 - EyeAdaptation_BlackHistogramBucketInfluence);
-	}
-	Intensity = max(Intensity, EyeAdaptation_LuminanceMin);
+	const float Intensity = CalculateEyeAdaptationLuminance(AverageColor * View.OneOverPreExposure);
 	const float LogIntensity = clamp(log2(Intensity), -10.0f, 20.0f);
 	// Store log intensity in the alpha channel: scale to 0,1 range.
 	OutEyeAdaptation = EyeAdaptation_HistogramScale * LogIntensity + EyeAdaptation_HistogramBias;
@@ -1321,6 +1315,7 @@ void AverageLuminance_MainCS(
 	uint2 Tile, TexelPos;
 	half2 BufferUV;
 	float Weight;
+	float Value;
 
 	float2 GroupLuminance = float2(0.0f, 0.0f);
 
@@ -1338,7 +1333,21 @@ void AverageLuminance_MainCS(
 
 				Weight = max(AdaptationWeightTexture(BufferUV), 0.05f);
 
-				GroupLuminance.x += SampleInput(BufferUV, 0).x * Weight;
+				Value = SampleInput(BufferUV, 0).x;
+
+				{
+					// apply EyeAdaptation_BlackHistogramBucketInfluence using similar logic as Histogram method
+					float fBucket = saturate(Value) * (HISTOGRAM_SIZE - 1);
+
+					// Find two discrete buckets that straddle the continuous histogram position.
+					uint Bucket0 = (uint)(fBucket);
+					if (Bucket0 == 0)
+					{
+						Weight *= EyeAdaptation_BlackHistogramBucketInfluence;
+					}
+				}
+
+				GroupLuminance.x += Value * Weight;
 				GroupLuminance.y += Weight;
 			}
 		}
@@ -1515,11 +1524,10 @@ void Histogram_MainCS(
 				// case is so the black background pixels in the editor have no effect. But if we have cases where
 				// pixel values can actually be black, we want to set EyeAdaptation_LastHistogramBucketInfluence=1.0.
 				// This value is controlled by the cvar "r.EyeAdaptation.BlackHistogramBucketInfluence"
-				// THIS IS DONE in a Bloom setup, black pixels will set intencity to 0.18
-				//if (Bucket0 == 0)
-				//{
-				//	Weight0 *= EyeAdaptation_BlackHistogramBucketInfluence;
-				//}
+				if (Bucket0 == 0)
+				{
+					Weight0 *= EyeAdaptation_BlackHistogramBucketInfluence;
+				}
 
 				// Accumulate the weight to the nearby history buckets.
 #if IOS // The IOS A8 and lower devices seems don't support using a float4 array as a two dimension array, separate the operations to two steps.
