Fixed improper blend of landscape edit layers with custom height/weight flags

diff --git a/Engine/Shaders/Private/Landscape/LandscapeCommon.ush b/Engine/Shaders/Private/Landscape/LandscapeCommon.ush
--- a/Engine/Shaders/Private/Landscape/LandscapeCommon.ush
+++ b/Engine/Shaders/Private/Landscape/LandscapeCommon.ush
@@ -41,16 +41,16 @@ float UnpackHeight(float2 InPackedHeight)
 // Height alpha is 14 bits-wide + 2 bits are left for flags. InAlpha is expected in the [0,1] range
 float2 PackHeightAlpha(float InAlpha, uint InAlphaFlags)
 {
-	uint PackedAlpha = (uint)clamp(InAlpha * 65532.0f, 0.0f, 65532.0f); // First 14 bits are used to store alpha (0xFFFC in float)
+	uint PackedAlpha = (uint)clamp(InAlpha * 65535.0f, 0.0f, 65535.0f); // First 14 bits are used to store alpha (0xFFFC in float)
 	uint PackedFlags = InAlphaFlags & 0x3; // Last 2 bits are used to store flags
-	return float2((float)(PackedAlpha >> 8) / 255.0f, (float)((PackedAlpha & 0xFF) | PackedFlags) / 255.0);
+	return float2((float)(PackedAlpha >> 8) / 255.0f, (float)((PackedAlpha & 0xFC) | PackedFlags) / 255.0);
 }
 
 void UnpackHeightAlpha(float2 InPackedAlpha, out float OutAlpha, out uint OutFlags)
 {
 	uint AlphaWithFlags = ((uint)round(InPackedAlpha.r * 255.0f) << 8) | (uint)round(InPackedAlpha.g * 255.0f);
-	uint Alpha = AlphaWithFlags & 0xFFFC; 
-	OutAlpha = float(Alpha) / 65532.0f; // First 14 bits are used to store alpha (0xFFFC in float)
+	uint Alpha = AlphaWithFlags & 0xFFFC; // First 14 bits are used to store alpha (0xFFFC in float)
+	OutAlpha = float(Alpha) / 65535.0f; 
 	OutFlags = AlphaWithFlags & 0x3; // Last 2 bits are used to store flags
 }
 
@@ -70,16 +70,16 @@ float UnpackWeight(float2 InPackedWeight)
 // Weight Alpha is 14 bits-wide + 2 bits are left for flags. InWeight is expected in the [0,1] range
 float2 PackWeightAlpha(float InAlpha, uint InAlphaFlags)
 {
-	uint PackedAlpha = (uint)clamp(InAlpha * 65532.0f, 0.0f, 65532.0f); // First 14 bits are used to store alpha (0xFFFC in float)
+	uint PackedAlpha = (uint)clamp(InAlpha * 65535.0f, 0.0f, 65535.0f); // First 14 bits are used to store alpha (0xFFFC in float)
 	uint PackedFlags = InAlphaFlags & 0x3; // Last 2 bits are used to store flags
-	return float2((float)(PackedAlpha >> 8) / 255.0, (float)((PackedAlpha & 0xFF) | PackedFlags) / 255.0);
+	return float2((float)(PackedAlpha >> 8) / 255.0, (float)((PackedAlpha & 0xFC) | PackedFlags) / 255.0);
 }
 
 void UnpackWeightAlpha(float2 InPackedAlpha, out float OutAlpha, out uint OutFlags)
 {
 	uint AlphaWithFlags = ((uint)round(InPackedAlpha.r * 255.0f) << 8) | (uint)round(InPackedAlpha.g * 255.0f);
-	uint Alpha = AlphaWithFlags & 0xFFFC; 
-	OutAlpha = float(Alpha) / 65532.0f; // First 14 bits are used to store alpha (0xFFFC in float)
+	uint Alpha = AlphaWithFlags & 0xFFFC; // First 14 bits are used to store alpha (0xFFFC in float)
+	OutAlpha = float(Alpha) / 65535.0f; 
 	OutFlags = AlphaWithFlags & 0x3; // Last 2 bits are used to store flags
 }
 
