Undo change to GetTranslatedWorldPositionForLighting(...) in TranslucentLightInjectionShaders.usf

diff --git a/Engine/Shaders/Private/TranslucentLightInjectionShaders.usf b/Engine/Shaders/Private/TranslucentLightInjectionShaders.usf
--- a/Engine/Shaders/Private/TranslucentLightInjectionShaders.usf
+++ b/Engine/Shaders/Private/TranslucentLightInjectionShaders.usf
@@ -112,8 +112,7 @@ float3 GetTranslatedWorldPositionForLighting(float3 TranslatedWorldPosition, flo
 	// Push out the position for lighting in the oposite direction of the L vector 
 	// to make sure we always have half a voxel distance from the light to avoid overblown lighting intensities.
 	float3 LightingOffset = -NormalizedLightVector * 0.5f * VoxelSize;
-	//return TranslatedWorldPosition + LightingOffset;
-	return TranslatedWorldPosition;
+	return TranslatedWorldPosition + LightingOffset;
 }
 
 /** Pixel shader that calculates direct lighting for a simple light (unshadowed point light) for all the affected voxels of a volume texture. */
