TSR: Fix for holdout translucency conflicting with bHasPixelAnimation.

Previously fixed CL 34691971, regression CL 35865103.

When bHasPixelAnimation is active, opaque colors take the translucency path for shading rejection (OriginalOpaqueInput -> HasAnimationBlurInput -> Translucency (ComposeTranslucency)),  while the SharpInput -> Color (ComposeTranslucency) remains transparent.

However once bPassthroughAlpha is enabled in addition to bHasPixelAnimation, ComposeTranslucencyForRejection would incorrectly return the constant transparent alpha, ignoring changes in scene color alpha.

This yielded incorrect behavior such as scene color alpha accumulation while toggling holdout on an opaque meshes with bHasPixelAnimation materials.

diff --git a/Engine/Shaders/Private/TemporalSuperResolution/TSRShadingAnalysis.ush b/Engine/Shaders/Private/TemporalSuperResolution/TSRShadingAnalysis.ush
--- a/Engine/Shaders/Private/TemporalSuperResolution/TSRShadingAnalysis.ush
+++ b/Engine/Shaders/Private/TemporalSuperResolution/TSRShadingAnalysis.ush
@@ -31,12 +31,12 @@ uint bPassthroughAlpha;
 
 /** Compose translucency in linear color space. */
 CALL_SITE_DEBUGLOC
-tsr_tensor_halfC ComposeTranslucency(tsr_tensor_halfC Color, tsr_tensor_half4 Translucency)
+tsr_tensor_halfC ComposeTranslucency(tsr_tensor_halfC Color, tsr_tensor_half4 Translucency, tsr_tensor_bool bAllowPassthroughAlpha = tsr_tensor_bool::Const(true))
 {
 	tsr_tensor_halfC ComposedColor;
 	ComposedColor = Color * tsr_tensor_halfC::Vectorize(Translucency[3]) + ResizeChannels<CONFIG_CHANNEL_COUNT>(Translucency);
 	#if CONFIG_SCENE_COLOR_ALPHA
-		ComposedColor.SetComponent(3, select(tsr_tensor_bool::Const(bPassthroughAlpha > 0), Color[3], Color[3] * Translucency[3]));
+		ComposedColor.SetComponent(3, select(and(tsr_tensor_bool::Const(bPassthroughAlpha > 0), bAllowPassthroughAlpha), Color[3], Color[3] * Translucency[3]));
 	#endif
 	#if CONFIG_SCENE_COLOR_ALPHA
 		ComposedColor = min(ComposedColor, tsr_tensor_halfC::Const(LargestSceneColorRGBA));
@@ -46,7 +46,7 @@ tsr_tensor_halfC ComposeTranslucency(tsr_tensor_halfC Color, tsr_tensor_half4 Tr
 	return ComposedColor;
 }
 
-/** Compose blury translucency for tighther rejection in linear color space. */
+/** Compose blurry translucency for tighter rejection in linear color space. */
 tsr_tensor_halfC ComposeTranslucencyForRejection(tsr_tensor_halfC OriginalOpaqueInput, tsr_tensor_half4 OriginalTranslucencyInput, tsr_tensor_bool bHasPixelAnimation)
 {
 	tsr_tensor_halfC SharpInput = OriginalOpaqueInput;
@@ -63,11 +63,18 @@ tsr_tensor_halfC ComposeTranslucencyForRejection(tsr_tensor_halfC OriginalOpaque
 			
 		tsr_tensor_half4 HasAnimationBlurInput = ResizeChannels<4>(CenterFinalSceneColor);
 			
-		BlurInput  = select(tsr_tensor_bool4::Vectorize(bHasPixelAnimation), HasAnimationBlurInput, BlurInput);
 		SharpInput = select(tsr_tensor_boolC::Vectorize(bHasPixelAnimation), tsr_tensor_halfC::Const(TransparentSharpInput), SharpInput);
+		BlurInput  = select(tsr_tensor_bool4::Vectorize(bHasPixelAnimation), HasAnimationBlurInput, BlurInput);
 	}
-		
+
+#if CONFIG_SCENE_COLOR_ALPHA
+	// Disallow alpha passthrough on opaque materials with pixel animation to retain scene color alpha values (i.e. dynamic holdout state).
+	const tsr_tensor_bool bAllowPassthroughAlpha = not(bHasPixelAnimation);
+
+	return ComposeTranslucency(SharpInput, Blur3x3(BlurInput), bAllowPassthroughAlpha);
+#else
 	return ComposeTranslucency(SharpInput, Blur3x3(BlurInput));
+#endif
 }
 
 
