[Motion blur visualization] Add two additional motion blur visualization mode:
* r.motionblur.visualize 1: HSV view. Encode direction in H, and magnitude in V.
* r.motionblur.visualize 2: Plasma view.  Encode motion vector magnitude only

diff --git a/Engine/Shaders/Private/ColorSpace.ush b/Engine/Shaders/Private/ColorSpace.ush
--- a/Engine/Shaders/Private/ColorSpace.ush
+++ b/Engine/Shaders/Private/ColorSpace.ush
@@ -255,6 +255,14 @@ float3 LinearRGB_2_HSV(in float3 RGB)
 	return float3(HCV.x, s, HCV.z);
 }
 
+// adapted from https://www.shadertoy.com/view/MsS3Wc by Inigo Quilez
+float3 HSV_2_LinearRGB_Smooth(float3 HSV)
+{
+	float3 RGB = clamp(abs((HSV.x * 6.0 + float3(0.0, 4.0, 2.0)) % 6.0 - 3.0) - 1.0, 0.0, 1.0);
+	RGB = RGB * RGB * (3.0 - 2.0 * RGB); // cubic smoothing	
+	return HSV.z * lerp(float3(1.0, 1.0, 1.0), RGB, HSV.y);
+}
+
 //------------------------------------------------------------------
 // LCH
 //------------------------------------------------------------------
diff --git a/Engine/Shaders/Private/MotionBlur/MotionBlurVisualize.usf b/Engine/Shaders/Private/MotionBlur/MotionBlurVisualize.usf
--- a/Engine/Shaders/Private/MotionBlur/MotionBlurVisualize.usf
+++ b/Engine/Shaders/Private/MotionBlur/MotionBlurVisualize.usf
@@ -4,13 +4,17 @@
 
 #include "MotionBlurCommon.ush"
 #include "../LensDistortion.ush"
-
+#include "../ColorMap.ush"
 
 //------------------------------------------------------- CONSTANTS
 
 // 0:off / 1:on, useful to debug the motionblur algorithm
 #define MOTIONBLUR_TESTCHART 0
 
+#define VISUALIZE_COLOR_ENCODING 1
+#define TILE_SIZE 48.0f
+#define COLOR_METER_LENGTH TILE_SIZE
+#define COLOR_METER_OFFSET 5
 
 //------------------------------------------------------- PARAMETERS
 
@@ -29,6 +33,7 @@ FScreenTransform SvPositionToScreenPos;
 FScreenTransform ScreenPosToColorUV;
 FScreenTransform ScreenPosToVelocityUV;
 int CheckerboardEnabled;
+int VisualizeMode;
 
 #if USE_POST_MOTION_BLUR_TRANSLUCENCY
 
@@ -124,14 +129,13 @@ float Compute3DCheckerBoard(float3 Pos)
 	return (float)((uint)dot(float3(1,1,1), TiledWorldPos) % 2);
 }
 
-
 //------------------------------------------------------- ENTRY POINT
 
 void MainPS(
 	float4 SvPosition : SV_POSITION,
 	out float4 OutColor : SV_Target0)
 {
-	const float TileSize = 48;
+	const float TileSize = TILE_SIZE;
 
 	float2 DistortedScreenPos = ApplyScreenTransform(SvPosition.xy, SvPositionToScreenPos);
 	float2 UndistortedScreenPos = ApplyLensDistortionOnScreenPos(UndistortingDisplacementTexture, UndistortingDisplacementSampler, DistortedScreenPos);
@@ -194,6 +198,75 @@ void MainPS(
 		OutColor.rgb = lerp(OutColor.rgb, TintColor, 0.45f);
 	}
 
+#if VISUALIZE_COLOR_ENCODING
+	{
+		float2 Velocity = DecodeVelocityFromTexture(VelocityTexture.SampleLevel(SharedVelocitySampler, VelocityUV, 0)).xy;
+		{
+			// reconstruct from camera motion if we don't have velocity data
+			if (!bSelectorOpaque)
+			{
+				float DeviceZ = DepthTexture.SampleLevel(SharedDepthSampler, VelocityUV, 0).r;
+
+				// UV is per tile
+				float4 ThisClip = float4(UndistortedScreenPos, DeviceZ, 1);
+				float4 PrevClip = mul(ThisClip, View.ClipToPrevClip);
+				float2 PrevScreen = PrevClip.xy / PrevClip.w;
+
+				// points into the movement direction
+				Velocity = UndistortedScreenPos - PrevScreen;
+			}
+		}
+
+		Velocity = Velocity * View.ViewSizeAndInvSize.xy;
+
+		float Gamma = 1/2.2f;
+		float Magnitude = clamp(length(Velocity)/COLOR_METER_LENGTH, 0.0f, 1.0f);
+		float AdjustedMagnitude = pow(Magnitude, Gamma);
+		float Satuation = 1.0f;
+		int2 TexCoord = SvPosition.xy;
+
+		if (VisualizeMode == 1)
+		{
+			float h = atan2(Velocity.y, Velocity.x) / 3.1415927f;
+			h = h < 0 ? ((h + 2) / 2) : h / 2.0;
+			float s = Satuation;
+			float v = AdjustedMagnitude;
+
+			float3 ColorEncoding = HSV_2_LinearRGB_Smooth(float3(h, s, v));
+
+			// Use per pixel encoding if the velocity is available, otherwise fallback to the yellow checkerbaord.
+			OutColor.rgb = lerp(OutColor.rgb, ColorEncoding, length(Velocity) > 0.0f);
+
+			if (length(TexCoord.xy - (COLOR_METER_LENGTH + COLOR_METER_OFFSET)) < COLOR_METER_LENGTH)
+			{
+				float2 Vector = TexCoord.xy - COLOR_METER_LENGTH;
+
+				h = atan2(-Vector.y, Vector.x) / 3.1415927f;
+				h = h < 0 ? ((h + 2) / 2) : h / 2.0;
+				s = Satuation;
+				v = pow(length(float3(Vector, 0)) / COLOR_METER_LENGTH, Gamma);
+				OutColor = float4(HSV_2_LinearRGB_Smooth(float3(h, s, v)), 1.0f);
+			}
+		}
+		else if (VisualizeMode == 2)
+		{
+			float3 ColorEncoding = ColorMapPlasma(AdjustedMagnitude);
+
+			// Use per pixel encoding if the velocity is available, otherwise fallback to the yellow checkerbaord.
+			OutColor.rgb = lerp(OutColor.rgb, ColorEncoding, length(Velocity) > 0.0f);
+
+			uint2 ColorMeterPosition = TexCoord.xy - (COLOR_METER_OFFSET);
+			if (all(ColorMeterPosition < uint2(COLOR_METER_LENGTH, 20)) &&
+				all(ColorMeterPosition >= uint2(0, 0)))
+			{
+				Magnitude = (TexCoord.x - COLOR_METER_OFFSET);
+				AdjustedMagnitude = pow( Magnitude / COLOR_METER_LENGTH, Gamma);
+				OutColor = float4(ColorMapPlasma(AdjustedMagnitude), 1.0f);
+			}
+		}
+	}
+#endif 
+
 	// tile center
 	{
 		float2 Delta = PixelPos - PixelPosAtTileCenter;
diff --git a/Engine/Source/Runtime/Renderer/Private/PostProcess/PostProcessMotionBlur.cpp b/Engine/Source/Runtime/Renderer/Private/PostProcess/PostProcessMotionBlur.cpp
--- a/Engine/Source/Runtime/Renderer/Private/PostProcess/PostProcessMotionBlur.cpp
+++ b/Engine/Source/Runtime/Renderer/Private/PostProcess/PostProcessMotionBlur.cpp
@@ -75,6 +75,14 @@ namespace
 		ECVF_RenderThreadSafe
 	);
 
+	TAutoConsoleVariable<int32> CVarMotionBlurVisualize(
+		TEXT("vr.MotionBlur.Visualize"),
+		0,
+		TEXT("0: Default view.\n")
+		TEXT("1: HSV view. Encode direction in H, and magnitude in V.\n")
+		TEXT("2: Plasma view. Encode motion vector magnitude only.")
+	);
+
 	TAutoConsoleVariable<int32> CVarOrthoUsePreviousMotionVelocityFlattenPass(
 		TEXT("r.Ortho.UsePreviousMotionVelocityFlattenPass"),
 		0,
@@ -546,6 +554,7 @@ class FMotionBlurVisualizePS : public FMotionBlurShader
 		SHADER_PARAMETER(FScreenTransform, SvPositionToScreenPos)
 		SHADER_PARAMETER(FScreenTransform, ScreenPosToVelocityUV)
 		SHADER_PARAMETER(int, CheckerboardEnabled)
+		SHADER_PARAMETER(int, VisualizeMode)
 
 		RENDER_TARGET_BINDING_SLOTS()
 	END_SHADER_PARAMETER_STRUCT()
@@ -1220,6 +1229,7 @@ FScreenPassTextureSlice AddVisualizeMotionBlurPass(FRDGBuilder& GraphBuilder, co
 	const FMotionBlurViewports Viewports(FScreenPassTextureViewport(Inputs.SceneColor), FScreenPassTextureViewport(Inputs.SceneDepth));
 
 	const bool bVisualizeDebugInfo = CVarVisualizeMotionBlurEnableDebugInformation.GetValueOnRenderThread();
+	const int32 VisualizeMode = CVarMotionBlurVisualize.GetValueOnRenderThread();
 
 	{
 		FMotionBlurVisualizePS::FParameters* PassParameters = GraphBuilder.AllocParameters<FMotionBlurVisualizePS::FParameters>();
@@ -1247,6 +1257,7 @@ FScreenPassTextureSlice AddVisualizeMotionBlurPass(FRDGBuilder& GraphBuilder, co
 			Inputs.SceneDepth, FScreenTransform::ETextureBasis::ScreenPosition, FScreenTransform::ETextureBasis::TextureUV);
 
 		PassParameters->CheckerboardEnabled = bVisualizeDebugInfo;
+		PassParameters->VisualizeMode = VisualizeMode;
 
 		TShaderMapRef<FMotionBlurVisualizePS> PixelShader(View.ShaderMap);
 		FPixelShaderUtils::AddFullscreenPass(
