Fix for skylight bottom color to work with preexposure.

#rnx
#tests px editor dx12

diff --git a/Engine/Shaders/Private/ReflectionEnvironmentShaders.usf b/Engine/Shaders/Private/ReflectionEnvironmentShaders.usf
--- a/Engine/Shaders/Private/ReflectionEnvironmentShaders.usf
+++ b/Engine/Shaders/Private/ReflectionEnvironmentShaders.usf
@@ -413,6 +413,8 @@ void ApplyLowerHemisphereColorPS(
 	in float4 SVPos : SV_POSITION,
 	out float4 OutLuminanceAlpha : SV_Target0)
 {
+	ResolvedView = ResolveView();
+
 	float2 UV = SVPos.xy * SvPositionToUVScale;
 	float2 ScaledUVs = UV * 2 - 1;
 
@@ -431,7 +433,8 @@ void ApplyLowerHemisphereColorPS(
 		float Coverage = 0.0f;
 		Coverage = pow(saturate(LowerHemisphereSolidColor.a), 2.2);
 		Coverage = saturate(LowerHemisphereSolidColor.a);
-		OutLuminanceAlpha = float4(lerp(OutLuminanceAlpha.rgb, LowerHemisphereSolidColor.rgb ,Coverage), OutLuminanceAlpha.a * (1.0 - Coverage));
+		const float Exposure = ResolvedView.RealTimeReflectionCapture ? ResolvedView.RealTimeReflectionCapturePreExposure : ResolvedView.PreExposure;
+		OutLuminanceAlpha = float4(lerp(OutLuminanceAlpha.rgb, LowerHemisphereSolidColor.rgb * Exposure, Coverage), OutLuminanceAlpha.a * (1.0 - Coverage));
 	}
 }
 
