Path Tracer: Add option to hide planet ground surface

Modify the analytic calculation of transmittance to account for rays starting "inside" the planet. This adds a segment going through homogeneous density which can be computed analytically.

diff --git a/Engine/Shaders/Private/PathTracing/PathTracingCore.ush b/Engine/Shaders/Private/PathTracing/PathTracingCore.ush
--- a/Engine/Shaders/Private/PathTracing/PathTracingCore.ush
+++ b/Engine/Shaders/Private/PathTracing/PathTracingCore.ush
@@ -652,7 +652,9 @@ FPackedPathTracingPayload TraceTransparentRay(inout FPathState PathState, int Bo
 		}
 #endif
 		FVolumeIntersectionList VolumeIsectList = VolumeIntersect(PathState.Ray.Origin, PathState.Ray.Direction, PathState.Ray.TMin, PathState.Ray.TMax, PathState.GetPathRoughness(), VolumeFlags);
-		if (VolumeIsectList.HitBlocker() && (PackedPayload.IsMiss() || VolumeIsectList.BlockerHitT < PackedPayload.HitT))
+		if (VolumeIsectList.HitBlocker() &&
+			(PathState.GetFirstScatterType() == PATHTRACER_SCATTER_VOLUME ||  (VolumeFlags & PATH_TRACER_VOLUME_SHOW_PLANET_GROUND) != 0) &&
+			(PackedPayload.IsMiss() || VolumeIsectList.BlockerHitT < PackedPayload.HitT))
 		{
 			// we didn't hit any real geometry, but we did hit the volume's blocker geometry
 			// create a virtual hit with it here
diff --git a/Engine/Shaders/Private/PathTracing/Volume/PathTracingAtmosphere.ush b/Engine/Shaders/Private/PathTracing/Volume/PathTracingAtmosphere.ush
--- a/Engine/Shaders/Private/PathTracing/Volume/PathTracingAtmosphere.ush
+++ b/Engine/Shaders/Private/PathTracing/Volume/PathTracingAtmosphere.ush
@@ -61,10 +61,7 @@ float3 RayPlanetIntersectPrecise(float3 Ro, float3 Rd)
 			q = DFSubtract(b, hp);
 		const float t0 = DFDemote(q);
 		const float t1 = DFFastDivideDemote(DFSubtract(DFLengthSqr(Co), PR2), q);
-			Result.x = min(t0, t1);
-		// clip volume layer against the back of the planet in case we started inside
-		if (Result.x < 0.0)
-			Result.y = max(t0, t1);
+		Result.x = min(t0, t1);
 	}
 	return Result;
 }
@@ -90,22 +87,19 @@ float3 RayPlanetIntersectPrecise(float3 Ro, float3 Rd)
 	{
 		ha = sqrt(ha);
 		float q = b > 0.0 ? b + ha : b - ha;
-		float ta = q;
-		float tb = (length2(co) - AtmoRadius * AtmoRadius) / q;
-		Result.y = min(ta, tb);
-		Result.z = max(ta, tb);
+		float t0 = q;
+		float t1 = (length2(co) - AtmoRadius * AtmoRadius) / q;
+		Result.y = min(t0, t1);
+		Result.z = max(t0, t1);
 	}
 	float hp = PlanetRadius * PlanetRadius - hd;
 	if (hp > 0)
 	{
 		hp = sqrt(hp);
 		float q = b > 0.0 ? b + hp : b - hp;
-		float ta = q;
-		float tb = (length2(co) - PlanetRadius * PlanetRadius) / q;
-		Result.x = min(ta, tb);
-		// clip volume layer against the back of the planet in case we started inside
-		if (Result.x < 0.0)
-			Result.y = max(ta, tb);
+		float t0 = q;
+		float t1 = (length2(co) - PlanetRadius * PlanetRadius) / q;
+		Result.x = min(t0, t1);
 	}
 
 	return Result;
@@ -157,7 +151,7 @@ struct FSphericalPlanetParameters {
 	{
 		// work out the ViewHeight / CosAngle of the specified point along the ray (in SkyUnits)
 		const float TOverH = TKm / HeightKm;
-		const float RelativeHeight = sqrt(1 + TOverH * TOverH + 2 * TOverH * ZenithCosAngle);
+		const float RelativeHeight = sqrt(mad(mad(2.0f, ZenithCosAngle, TOverH), TOverH, 1.0f));
 		FSphericalPlanetParameters Result;
 		Result.HeightKm = HeightKm * RelativeHeight;
 		Result.ZenithCosAngle = (TOverH + ZenithCosAngle) / RelativeHeight;
@@ -270,8 +264,27 @@ FVolumeShadedResult AtmosphereGetDensity(float3 TranslatedWorldPos)
 
 float3 AtmosphereGetOpticalDensity(FSphericalPlanetParameters Params)
 {
+	float3 Tau = 0;
+	if (Params.HeightKm < Atmosphere.BottomRadiusKm)
+	{
+		// Starting below ground level, need to compute transmittance through homogeneous region first
+		// TODO: merge this into the LUT somehow (requires a new parameterization)
+
+		// Ground level extinction, applies to entire planet core
+		// We keep this value in reciprocal sky units because we measure the distance in sky units as well (canceling out to unitless optical depth)
+		float3 SigmaT = SampleAtmosphereMediumRGB(float3(0, 0, 0)).Extinction;
+
+		const float C = Params.ZenithCosAngle;
+		const float H = Params.HeightKm;
+		const float B = Atmosphere.BottomRadiusKm;
+		const float DistanceToGround = -C * H + sqrt(max(B * B + H * H * (C * C - 1.0f), 0.0f));
+		Tau = DistanceToGround * SigmaT;
+
+		// Get remaining density from the point on the ground
+		Params = Params.GetAt(DistanceToGround);
+	}
 	float2 UV = ToAtmosphereOpticalDepthLUTUV(Params.HeightKm, Params.ZenithCosAngle, Atmosphere.BottomRadiusKm, Atmosphere.TopRadiusKm);
-	return AtmosphereOpticalDepthLUT.SampleLevel(AtmosphereOpticalDepthLUTSampler, UV, 0);
+	return Tau + AtmosphereOpticalDepthLUT.SampleLevel(AtmosphereOpticalDepthLUTSampler, UV, 0);
 }
 
 float3 AtmosphereGetTransmittance(float3 Origin, float3 Direction, float TMin, float TMax)
diff --git a/Engine/Shaders/Shared/PathTracingDefinitions.h b/Engine/Shaders/Shared/PathTracingDefinitions.h
--- a/Engine/Shaders/Shared/PathTracingDefinitions.h
+++ b/Engine/Shaders/Shared/PathTracingDefinitions.h
@@ -75,6 +75,8 @@
 #define PATH_TRACER_VOLUME_CALLABLE_FLAGS_BOUNCE_MASK       (2048u * 31) // What bounce are we rendering? (5 bits, so bounces>31 are clamped)
 #define PATH_TRACER_VOLUME_CALLABLE_FLAGS_BOUNCE_SHIFT      (11)
 
+// Flags beyond this point are not visible to the cloud callable shader
+#define PATH_TRACER_VOLUME_SHOW_PLANET_GROUND               (65536u)
 
 // Constants related to debugging
 #define PATH_TRACER_DEBUG_VIZ_RADIANCE					0 
diff --git a/Engine/Source/Runtime/Renderer/Private/PathTracing.cpp b/Engine/Source/Runtime/Renderer/Private/PathTracing.cpp
--- a/Engine/Source/Runtime/Renderer/Private/PathTracing.cpp
+++ b/Engine/Source/Runtime/Renderer/Private/PathTracing.cpp
@@ -328,6 +328,14 @@ TAutoConsoleVariable<int32> CVarPathTracingAtmosphereOpticalDepthLutNumSamples(
 	ECVF_RenderThreadSafe
 );
 
+TAutoConsoleVariable<bool> CVarPathTracingEnableAtmosphereGround(
+	TEXT("r.PathTracing.EnableAtmosphereGround"),
+	false,
+	TEXT("Should the planet ground surface of the atmosphere model be visible by the path tracer?\n")
+	TEXT("The planet ground is always visible to volume bounces to influence the color of the atmosphere\n"),
+	ECVF_RenderThreadSafe
+);
+
 TAutoConsoleVariable<int32> CVarPathTracingCloudAccelerationMapResolution(
 	TEXT("r.PathTracing.CloudAccelerationMap.Resolution"),
 	512,
@@ -932,6 +940,7 @@ static void PreparePathTracingData(const FScene* Scene, const FViewInfo& View, F
 		PathTracingData.VolumeFlags |= Scene->ExponentialFogs.Num() > 0 && IsExponentialFogHoldout(View.CachedViewUniformShaderParameters->EnvironmentComponentsFlags) ? PATH_TRACER_VOLUME_HOLDOUT_FOG : 0;
 	}
 	PathTracingData.VolumeFlags |= CVarPathTracingUseAnalyticTransmittance.GetValueOnRenderThread() ? PATH_TRACER_VOLUME_USE_ANALYTIC_TRANSMITTANCE : 0;
+	PathTracingData.VolumeFlags |= CVarPathTracingEnableAtmosphereGround.GetValueOnRenderThread() ? PATH_TRACER_VOLUME_SHOW_PLANET_GROUND : 0;
 
 	PathTracingData.CloudRoughnessCutoff = bVolumeCloudsVisible ? CVarPathTracingCloudRoughnessCutoff.GetValueOnRenderThread() : -1.0f;
 	if (bVolumeCloudsVisible &&
