- Fix mismatch between using NUM_MATERIAL_TEXCOORDS in VS vs NUM_MATERIAL_TEXCOORDS_VERTEX

#rnx

diff --git a/Engine/Plugins/FX/Niagara/Shaders/Private/NiagaraSpriteVertexFactory.ush b/Engine/Plugins/FX/Niagara/Shaders/Private/NiagaraSpriteVertexFactory.ush
--- a/Engine/Plugins/FX/Niagara/Shaders/Private/NiagaraSpriteVertexFactory.ush
+++ b/Engine/Plugins/FX/Niagara/Shaders/Private/NiagaraSpriteVertexFactory.ush
@@ -196,7 +196,7 @@ struct FVertexFactoryIntermediates
 
 	/** The texture coordinate at this vertex. */
 	float4 TexCoord;
-#if NUM_MATERIAL_TEXCOORDS >= 3 //Unflipped UVs are placed in UVs 2-3
+#if NUM_MATERIAL_TEXCOORDS_VERTEX >= 3 //Unflipped UVs are placed in UVs 2-3
 	/** A second UV set. Always non-UV flipped. Allows use of UV flipping with Normal maps etc. */
 	float4 TexCoord_Unflipped;
 #endif
@@ -1057,7 +1057,7 @@ FVertexFactoryIntermediates GetVertexFactoryIntermediates(FVertexFactoryInput In
 	float SubImageBV = floor( SubImageB * NiagaraSpriteVF.SubImageSize.z );
 	Intermediates.TexCoord.xy = (float2( SubImageAH, SubImageAV ) + UVForTexturing) * NiagaraSpriteVF.SubImageSize.zw;
 	Intermediates.TexCoord.zw = (float2(SubImageBH, SubImageBV) + UVForTexturing) * NiagaraSpriteVF.SubImageSize.zw;
-#if NUM_MATERIAL_TEXCOORDS >= 3
+#if NUM_MATERIAL_TEXCOORDS_VERTEX >= 3
 	Intermediates.TexCoord_Unflipped.xy = (float2(SubImageAH, SubImageAV) + UVForTexturingUnflipped) * NiagaraSpriteVF.SubImageSize.zw;
 	Intermediates.TexCoord_Unflipped.zw = (float2(SubImageBH, SubImageBV) + UVForTexturingUnflipped) * NiagaraSpriteVF.SubImageSize.zw;
 #endif
diff --git a/Engine/Shaders/Private/ParticleGPUSpriteVertexFactory.ush b/Engine/Shaders/Private/ParticleGPUSpriteVertexFactory.ush
--- a/Engine/Shaders/Private/ParticleGPUSpriteVertexFactory.ush
+++ b/Engine/Shaders/Private/ParticleGPUSpriteVertexFactory.ush
@@ -127,7 +127,7 @@ struct FVertexFactoryIntermediates
 
 	/** The texture coordinate at this vertex. */
 	float4 TexCoord;
-#if NUM_MATERIAL_TEXCOORDS >= 3 //Unflipped UVs in texcoords 3 and 4
+#if NUM_MATERIAL_TEXCOORDS_VERTEX >= 3 //Unflipped UVs in texcoords 3 and 4
 	float4 TexCoord_Unflipped;
 #endif
 	/** The sprite tangent in world space (+V). */
@@ -312,7 +312,7 @@ FVertexFactoryIntermediates GetVertexFactoryIntermediates( FVertexFactoryInput I
 	TexCoord.xy = (float2( SubImageAH, SubImageAV ) + FlippedUV) * EmitterUniforms.SubImageSize.zw;
 	TexCoord.zw = (float2(SubImageBH, SubImageBV) + FlippedUV) * EmitterUniforms.SubImageSize.zw;
 
-#if NUM_MATERIAL_TEXCOORDS >= 3 //Unflipped UVs in texcoords 3 and 4
+#if NUM_MATERIAL_TEXCOORDS_VERTEX >= 3 //Unflipped UVs in texcoords 3 and 4
 	float4 TexCoord_Unflipped;
 	TexCoord_Unflipped.xy = (float2(SubImageAH, SubImageAV) + Input.TexCoord.xy) * EmitterUniforms.SubImageSize.zw;
 	TexCoord_Unflipped.zw = (float2(SubImageBH, SubImageBV) + Input.TexCoord.xy) * EmitterUniforms.SubImageSize.zw;
@@ -368,7 +368,7 @@ FVertexFactoryIntermediates GetVertexFactoryIntermediates( FVertexFactoryInput I
 #endif
 
 	Intermediates.TexCoord = TexCoord;
-#if NUM_MATERIAL_TEXCOORDS >= 3 //Unflipped UVs in texcoords 3 and 4
+#if NUM_MATERIAL_TEXCOORDS_VERTEX >= 3 //Unflipped UVs in texcoords 3 and 4
 	Intermediates.TexCoord_Unflipped = TexCoord_Unflipped; 
 #endif
 	Intermediates.TangentUp = TangentUp;
diff --git a/Engine/Shaders/Private/ParticleSpriteVertexFactory.ush b/Engine/Shaders/Private/ParticleSpriteVertexFactory.ush
--- a/Engine/Shaders/Private/ParticleSpriteVertexFactory.ush
+++ b/Engine/Shaders/Private/ParticleSpriteVertexFactory.ush
@@ -107,7 +107,7 @@ struct FVertexFactoryIntermediates
 
 	/** The texture coordinate at this vertex. */
 	float4 TexCoord;
-#if NUM_MATERIAL_TEXCOORDS >= 3 //Unflipped UVs are placed in UVs 2-3
+#if NUM_MATERIAL_TEXCOORDS_VERTEX >= 3 //Unflipped UVs are placed in UVs 2-3
 	/** A second UV set. Always non-UV flipped. Allows use of UV flipping with Normal maps etc. */
 	float4 TexCoord_Unflipped;
 #endif
@@ -556,7 +556,7 @@ FVertexFactoryIntermediates GetVertexFactoryIntermediates(FVertexFactoryInput In
 	float SubImageBV = floor( SubImageB * SpriteVF.SubImageSize.z );
 	Intermediates.TexCoord.xy = (float2( SubImageAH, SubImageAV ) + UVForTexturing) * SpriteVF.SubImageSize.zw ;
 	Intermediates.TexCoord.zw = (float2(SubImageBH, SubImageBV) + UVForTexturing) * SpriteVF.SubImageSize.zw;
-#if NUM_MATERIAL_TEXCOORDS >= 3
+#if NUM_MATERIAL_TEXCOORDS_VERTEX >= 3
 	Intermediates.TexCoord_Unflipped.xy = (float2(SubImageAH, SubImageAV) + UVForTexturingUnflipped) * SpriteVF.SubImageSize.zw;
 	Intermediates.TexCoord_Unflipped.zw = (float2(SubImageBH, SubImageBV) + UVForTexturingUnflipped) * SpriteVF.SubImageSize.zw;
 #endif
