Using correct PS input signature in FSlatePostProcessColorDeficiencyPS.

diff --git a/Engine/Shaders/Private/SlatePostProcessColorDeficiencyPixelShader.usf b/Engine/Shaders/Private/SlatePostProcessColorDeficiencyPixelShader.usf
--- a/Engine/Shaders/Private/SlatePostProcessColorDeficiencyPixelShader.usf
+++ b/Engine/Shaders/Private/SlatePostProcessColorDeficiencyPixelShader.usf
@@ -14,9 +14,9 @@ float ColorVisionDeficiencySeverity;
 float bCorrectDeficiency;
 float bSimulateCorrectionWithDeficiency;
 
-float4 ColorDeficiencyMain(FScreenVertexOutput Input) : SV_Target0
+float4 ColorDeficiencyMain(noperspective float4 UVAndScreenPos : TEXCOORD0) : SV_Target0
 {
-	float4 SourceColor = Texture2DSample(ElementTexture, ElementTextureSampler, Input.UV);
+	float4 SourceColor = Texture2DSample(ElementTexture, ElementTextureSampler, UVAndScreenPos.xy);
 	float3 OutRGB = ColorDeficiency(SourceColor.rgb, ColorVisionDeficiencyType, ColorVisionDeficiencySeverity, bCorrectDeficiency, bSimulateCorrectionWithDeficiency);
 	return float4(OutRGB, SourceColor.a);
 }
\ No newline at end of file
