[VSM] Accumulate throttling bias in projection data, rather than in page marking. This makes it behave more like global bias during sampling, e.g. for dithering. Also account for DOF bias during projection to avoid artifacting when receivermask is on.

diff --git a/Engine/Shaders/Private/VirtualShadowMaps/VirtualShadowMapPageMarking.usf b/Engine/Shaders/Private/VirtualShadowMaps/VirtualShadowMapPageMarking.usf
--- a/Engine/Shaders/Private/VirtualShadowMaps/VirtualShadowMapPageMarking.usf
+++ b/Engine/Shaders/Private/VirtualShadowMaps/VirtualShadowMapPageMarking.usf
@@ -4,6 +4,8 @@
 	VirtualShadowMapPageMarking.usf: 
 =============================================================================*/
 
+#define VSM_WITH_DOF_BIAS 1
+
 #include "/Engine/Shared/SingleLayerWaterDefinitions.h"
 #include "../HairStrands/HairStrandsVisibilityCommonStruct.ush"
 #include "../Common.ush"
@@ -22,7 +24,6 @@
 #include "../ShaderPrint.ush"
 #include "../Froxel/Froxel.ush"
 #include "../Nanite/NaniteHZBCull.ush"
-#include "../DiaphragmDOF/DOFCommon.ush"
 
 // Type of input data consume by the page allocation (i.e., data read from the source buffer: Gbuffer, HairStrands data, ...)
 #define INPUT_TYPE_GBUFFER 0
@@ -42,50 +43,30 @@ uint InputType;
 uint bCullBackfacingPixels;
 uint NumDirectionalLightSmInds;
 
-float DOFBiasStrength;
-float ThrottleMaxBiasDirectional;
-float ThrottleMaxBiasLocal;
-
 #ifdef GeneratePageFlagsFromFroxelsCS
 int PassId;
 #endif
 
-float GetSustainedLoadThrottlingResolutionBiasDirectional(FVirtualShadowMapProjectionShaderData BaseProjectionData, float BiasedLevel)
-{
-	float ResolutionBias = 0.0f;
-
-#if PERMUTATION_THROTTLING
-	const FVirtualShadowMapHandle VSMHandle = BaseProjectionData.VirtualShadowMapHandle.MakeOffset(BiasedLevel - BaseProjectionData.ClipmapLevel);
-
-	float ThrottleIntensity = asfloat(ThrottleBuffer[VSM_TB_SIZEOF_HEADER + VSM_TB_SIZEOF_ENTRY * VSMHandle.GetDataIndex() + VSM_TB_ENTRY_THROTTLE]);
-	
-	ResolutionBias = ThrottleIntensity * ThrottleMaxBiasDirectional;
-#endif
-
-	return ResolutionBias;
-}
 
-float GetSustainedLoadThrottlingResolutionBiasLocal(FVirtualShadowMapProjectionShaderData ProjectionData)
+int GetBiasedClipmapLevel(FVirtualShadowMapProjectionShaderData BaseProjectionData, float3 TranslatedWorldPosition, float ExtraBias)
 {
-	float ResolutionBias = 0.0f;
-
 #if PERMUTATION_THROTTLING
-	const FVirtualShadowMapHandle VSMHandle = ProjectionData.VirtualShadowMapHandle;
-
-	float ThrottleIntensity = asfloat(ThrottleBuffer[VSM_TB_SIZEOF_HEADER + VSM_TB_SIZEOF_ENTRY * VSMHandle.GetDataIndex() + VSM_TB_ENTRY_THROTTLE]);
-	
-	ResolutionBias = ThrottleIntensity * ThrottleMaxBiasLocal;
-#endif
+	float BiasedLevel = CalcAbsoluteClipmapLevel(BaseProjectionData, TranslatedWorldPosition) + VirtualShadowMap.GlobalResolutionLodBias + ExtraBias;
 
-	return ResolutionBias;
-}
-
-int GetBiasedClipmapLevel(FVirtualShadowMapProjectionShaderData BaseProjectionData, float3 TranslatedWorldPosition, float ExtraBias)
-{
+	int ClipmapIndex = max(0, BiasedLevel - BaseProjectionData.ClipmapLevel);
+	if (ClipmapIndex < BaseProjectionData.ClipmapLevelCountRemaining)
+	{
+		const FVirtualShadowMapHandle VSMHandle = BaseProjectionData.VirtualShadowMapHandle.MakeOffset(ClipmapIndex);
+		float PerVSMBias = GetVirtualShadowMapProjectionData(VSMHandle).ResolutionLodBias;
+		BiasedLevel += PerVSMBias;
+	}
+	else
+	{
+		BiasedLevel += BaseProjectionData.ResolutionLodBias;
+	}
+#else
 	float BiasedLevel = CalcAbsoluteClipmapLevel(BaseProjectionData, TranslatedWorldPosition) + BaseProjectionData.ResolutionLodBias + VirtualShadowMap.GlobalResolutionLodBias + ExtraBias;
-
-	float ThrottlingBias = GetSustainedLoadThrottlingResolutionBiasDirectional(BaseProjectionData, BiasedLevel);
-	BiasedLevel += ThrottlingBias;
+#endif
 
 	return int(floor(BiasedLevel));
 }
@@ -96,8 +77,6 @@ uint GetMipLevelLocal(FVirtualShadowMapProjectionShaderData ProjectionData, floa
 	float3 ViewToShadowTranslation = DFFastLocalSubtractDemote(ProjectionData.PreViewTranslation, PrimaryView.PreViewTranslation);
 	float3 ShadowTranslatedWorldPosition = TranslatedWorldPosition + ViewToShadowTranslation;
 
-	ExtraBias += GetSustainedLoadThrottlingResolutionBiasLocal(ProjectionData);
-
 	float Footprint = VirtualShadowMapCalcPixelFootprintLocal(ProjectionData, ShadowTranslatedWorldPosition, SceneDepth);
 	return GetMipLevelLocal(Footprint, VirtualShadowMap.MipModeLocal, ProjectionData.ResolutionLodBias, VirtualShadowMap.GlobalResolutionLodBias, ExtraBias);
 }
@@ -461,22 +440,6 @@ FPositionData InitPositionData(uint2 PixelPos, float DeviceZ, bool bIsValid)
 	return Out;
 }
 
-// Areas that are out of focus get a lower VSM resolution, as the resulting image will get blurred anyway
-float CalculateResolutionBiasFromDepthOfField(float SceneDepth)
-{
-	bool bApplyBiasFromDOFCoCRadius = DOFBiasStrength != 0.0f;
-	if (bApplyBiasFromDOFCoCRadius)
-	{
-		float AbsCocRadius = abs(SceneDepthToCocRadius(SceneDepth));
-		const float BiasRampUpDistance = 8.0f;
-		return min(DOFBiasStrength, lerp(0.0f, DOFBiasStrength, AbsCocRadius / BiasRampUpDistance));
-	}
-	else
-	{
-		return 0.0f;
-	}
-}
-
 // Bias applied to pixels originating from first-person geometry, if negative marking is skipped for these
 float FirstPersonPixelRequestBias;
 int FirstPersonPixelRequestLevelClamp;
diff --git a/Engine/Shaders/Private/VirtualShadowMaps/VirtualShadowMapProjection.usf b/Engine/Shaders/Private/VirtualShadowMaps/VirtualShadowMapProjection.usf
--- a/Engine/Shaders/Private/VirtualShadowMaps/VirtualShadowMapProjection.usf
+++ b/Engine/Shaders/Private/VirtualShadowMaps/VirtualShadowMapProjection.usf
@@ -285,7 +285,8 @@ FVirtualShadowMapSampleResult ProjectLight(
 				VirtualShadowMapId,
 				TranslatedWorldPosition,
 				SMRTRayOffset,
-				WorldNormal);
+				WorldNormal,
+				SceneDepth);
 		}
 #if FIRST_PERSON_SHADOW
 		// If there is an extra FP shadow map, it represents the shadow cast from the 3rd person mesh of the player onto the world only, thus it should only be sampled if
@@ -315,7 +316,8 @@ FVirtualShadowMapSampleResult ProjectLight(
 					FirstPersonVirtualShadowMapId,
 					TranslatedWorldPosition,
 					SMRTRayOffset,
-					WorldNormal);
+					WorldNormal,
+					SceneDepth);
 			}
 			if (FPResult.bValid)
 			{
diff --git a/Engine/Shaders/Private/VirtualShadowMaps/VirtualShadowMapProjectionCommon.ush b/Engine/Shaders/Private/VirtualShadowMaps/VirtualShadowMapProjectionCommon.ush
--- a/Engine/Shaders/Private/VirtualShadowMaps/VirtualShadowMapProjectionCommon.ush
+++ b/Engine/Shaders/Private/VirtualShadowMaps/VirtualShadowMapProjectionCommon.ush
@@ -12,6 +12,28 @@ VirtualShadowMapProjectionCommon.ush:
 
 #define VIRTUAL_SHADOW_MAP_MIN_OCCLUDER_DISTANCE 1e-6f
 
+#if VSM_WITH_DOF_BIAS
+#include "../DiaphragmDOF/DOFCommon.ush"
+
+float DOFBiasStrength;
+
+// Areas that are out of focus get a lower VSM resolution, as the resulting image will get blurred anyway
+float CalculateResolutionBiasFromDepthOfField(float SceneDepth)
+{
+	bool bApplyBiasFromDOFCoCRadius = DOFBiasStrength != 0.0f;
+	if (bApplyBiasFromDOFCoCRadius)
+	{
+		float AbsCocRadius = abs(SceneDepthToCocRadius(SceneDepth));
+		const float BiasRampUpDistance = 8.0f;
+		return min(DOFBiasStrength, lerp(0.0f, DOFBiasStrength, AbsCocRadius / BiasRampUpDistance));
+	}
+	else
+	{
+		return 0.0f;
+	}
+}
+#endif
+
 float CalcAbsoluteClipmapLevel(FVirtualShadowMapProjectionShaderData BaseProjectionData, float3 TranslatedWorldPosition)
 {
 	float3 ViewToShadowTranslation = DFFastLocalSubtractDemote(BaseProjectionData.PreViewTranslation, PrimaryView.PreViewTranslation);
@@ -21,11 +43,33 @@ float CalcAbsoluteClipmapLevel(FVirtualShadowMapProjectionShaderData BaseProject
 }
 
 // Optionally apply the resolution bias (rather than greedy sampling highest res data) depending on the global setting
-float CalcBiasedAbsoluteClipmapLevelForSampling(FVirtualShadowMapProjectionShaderData BaseProjectionData, float3 TranslatedWorldPosition)
+float CalcBiasedAbsoluteClipmapLevelForSampling(FVirtualShadowMapProjectionShaderData BaseProjectionData, float3 TranslatedWorldPosition, float SceneDepth)
 {
 	float AbsoluteLevel = CalcAbsoluteClipmapLevel(BaseProjectionData, TranslatedWorldPosition);
-	AbsoluteLevel += BaseProjectionData.ResolutionLodBias;
-	return AbsoluteLevel;
+
+	float BiasedLevel = AbsoluteLevel; //todo: what about other bias sources, like VirtualShadowMap.GlobalResolutionLodBias
+
+	if (SceneDepth >= 0)
+	{
+#if VSM_WITH_DOF_BIAS
+		float DOFResolutionBias = CalculateResolutionBiasFromDepthOfField(SceneDepth);
+		BiasedLevel += DOFResolutionBias;
+#endif
+	}
+
+	int ClipmapIndex = max(0, BiasedLevel - BaseProjectionData.ClipmapLevel);
+	if (ClipmapIndex < BaseProjectionData.ClipmapLevelCountRemaining)
+	{
+		const FVirtualShadowMapHandle VSMHandle = BaseProjectionData.VirtualShadowMapHandle.MakeOffset(ClipmapIndex);
+		float PerVSMBias = GetVirtualShadowMapProjectionData(VSMHandle).ResolutionLodBias;
+		BiasedLevel += PerVSMBias;
+	}
+	else
+	{
+		BiasedLevel += BaseProjectionData.ResolutionLodBias;
+	}
+
+	return BiasedLevel;
 }
 
 // Calculate the projected pixel footprint onto a local light, scaled in shadow map texels
@@ -371,23 +415,25 @@ FVirtualShadowMapSampleResult InitVirtualShadowMapSampleResult()
 *   - Should be zero or positive
 *   - This offset is useful in that it does *not* affect the selection of clipmap level, unlike offsetting the WorldPosition itself
 *   - OccluderDistance will still be relative to the original sample position
-* * EstimatedGeoWorldNormal is ideally the geometric (flat) normal of the sample point in world space
+* EstimatedGeoWorldNormal is ideally the geometric (flat) normal of the sample point in world space
 *   - If bUseOptimalBias is true, this is used to compute a receiver-plane-based bias for the sample point
 *   - The shading normal can be used if the geometric normal is not available, but divergence from the geometric normal can cause biasing issues.
+* SceneDepth is the distance from the sample point to the main view, or -1 if not available.
 */
 FVirtualShadowMapSampleResult SampleVirtualShadowMapDirectional(
 	FVirtualShadowMapHandle VirtualShadowMapHandle,
 	float3 TranslatedWorldPosition,
 	float RayStartDistance,
 	float3 EstimatedGeoWorldNormal,
-	bool bUseOptimalBias = true)
+	bool bUseOptimalBias = true,
+	float SceneDepth = -1.0f)
 {
 	RayStartDistance = max(RayStartDistance, 0.0f);
 
 	FVirtualShadowMapProjectionShaderData BaseProjectionData = GetVirtualShadowMapProjectionData(VirtualShadowMapHandle);
 	checkSlow(BaseProjectionData.LightType == LIGHT_TYPE_DIRECTIONAL);
 	
-	const int ClipmapLevel = int(floor(CalcBiasedAbsoluteClipmapLevelForSampling(BaseProjectionData, TranslatedWorldPosition)));
+	const int ClipmapLevel = int(floor(CalcBiasedAbsoluteClipmapLevelForSampling(BaseProjectionData, TranslatedWorldPosition, SceneDepth)));
 	int ClipmapIndex = max(0, ClipmapLevel - BaseProjectionData.ClipmapLevel);
 
 	// Check if sample is within the clipmap range (from camera)
@@ -446,14 +492,16 @@ FVirtualShadowMapSampleResult SampleVirtualShadowMapDirectional(
 	float3 TranslatedWorldPosition,
 	float RayStartDistance,
 	float3 EstimatedGeoWorldNormal,
-	bool bUseOptimalBias = true)
+	bool bUseOptimalBias = true,
+	float SceneDepth = -1.0f)
 {
 	return SampleVirtualShadowMapDirectional(
 		FVirtualShadowMapHandle::MakeFromIdDirectional(VirtualShadowMapId),
 		TranslatedWorldPosition,
 		RayStartDistance,
 		EstimatedGeoWorldNormal,
-		bUseOptimalBias);
+		bUseOptimalBias,
+		SceneDepth);
 }
 
 /**
@@ -558,21 +606,21 @@ FVirtualShadowMapSampleResult SampleVirtualShadowMapLocal(
 // Some common variants/convenience functions
 
 // Sample in PrimaryView translated world space with no optimal bias
-FVirtualShadowMapSampleResult SampleVirtualShadowMapDirectional(int VirtualShadowMapId, float3 TranslatedWorldPosition, float RayStartDistance = 0.0f)
+FVirtualShadowMapSampleResult SampleVirtualShadowMapDirectional(int VirtualShadowMapId, float3 TranslatedWorldPosition, float RayStartDistance = 0.0f, float SceneDepth = -1.0f)
 {
-	return SampleVirtualShadowMapDirectional(VirtualShadowMapId, TranslatedWorldPosition, RayStartDistance, float3(0, 0, 0), false);
+	return SampleVirtualShadowMapDirectional(VirtualShadowMapId, TranslatedWorldPosition, RayStartDistance, float3(0, 0, 0), false, SceneDepth);
 }
-FVirtualShadowMapSampleResult SampleVirtualShadowMapDirectional(FVirtualShadowMapHandle VirtualShadowMapHandle, float3 TranslatedWorldPosition, float RayStartDistance = 0.0f)
+FVirtualShadowMapSampleResult SampleVirtualShadowMapDirectional(FVirtualShadowMapHandle VirtualShadowMapHandle, float3 TranslatedWorldPosition, float RayStartDistance = 0.0f, float SceneDepth = -1.0f)
 {
-	return SampleVirtualShadowMapDirectional(VirtualShadowMapHandle, TranslatedWorldPosition, RayStartDistance, float3(0, 0, 0), false);
+	return SampleVirtualShadowMapDirectional(VirtualShadowMapHandle, TranslatedWorldPosition, RayStartDistance, float3(0, 0, 0), false, SceneDepth);
 }
-FVirtualShadowMapSampleResult SampleVirtualShadowMapLocal(int VirtualShadowMapId, float3 TranslatedWorldPosition, float RayStartDistance = 0.0f)
+FVirtualShadowMapSampleResult SampleVirtualShadowMapLocal(int VirtualShadowMapId, float3 TranslatedWorldPosition, float RayStartDistance = 0.0f, float SceneDepth = -1.0f)
 {
-	return SampleVirtualShadowMapLocal(VirtualShadowMapId, TranslatedWorldPosition, RayStartDistance, float3(0, 0, 0), false);
+	return SampleVirtualShadowMapLocal(VirtualShadowMapId, TranslatedWorldPosition, RayStartDistance, float3(0, 0, 0), false, SceneDepth);
 }
-FVirtualShadowMapSampleResult SampleVirtualShadowMapLocal(FVirtualShadowMapHandle VirtualShadowMapHandle, float3 TranslatedWorldPosition, float RayStartDistance = 0.0f)
+FVirtualShadowMapSampleResult SampleVirtualShadowMapLocal(FVirtualShadowMapHandle VirtualShadowMapHandle, float3 TranslatedWorldPosition, float RayStartDistance = 0.0f, float SceneDepth = -1.0f)
 {
-	return SampleVirtualShadowMapLocal(VirtualShadowMapHandle, TranslatedWorldPosition, RayStartDistance, float3(0, 0, 0), false);
+	return SampleVirtualShadowMapLocal(VirtualShadowMapHandle, TranslatedWorldPosition, RayStartDistance, float3(0, 0, 0), false, SceneDepth);
 }
 
 
diff --git a/Engine/Shaders/Private/VirtualShadowMaps/VirtualShadowMapProjectionDirectional.ush b/Engine/Shaders/Private/VirtualShadowMaps/VirtualShadowMapProjectionDirectional.ush
--- a/Engine/Shaders/Private/VirtualShadowMaps/VirtualShadowMapProjectionDirectional.ush
+++ b/Engine/Shaders/Private/VirtualShadowMaps/VirtualShadowMapProjectionDirectional.ush
@@ -35,11 +35,11 @@ float ComputeOptimalSlopeBiasDirectional(
 	return 2.0f * max(0.0f, dot(DepthSlopeUV, OffsetUV));
 }
 
-FVirtualShadowMapHandle GetMappedClipmap(FVirtualShadowMapHandle VirtualShadowMapHandle, float3 RayOriginTranslatedWorld)
+FVirtualShadowMapHandle GetMappedClipmap(FVirtualShadowMapHandle VirtualShadowMapHandle, float3 RayOriginTranslatedWorld, float SceneDepth)
 {
 	// This is the (unbiased) most detailed clipmap level that we are guaranteed to have a page table entry for
 	const FVirtualShadowMapProjectionShaderData BaseProjectionData = GetVirtualShadowMapProjectionData(VirtualShadowMapHandle);		
-	float ClipmapLevelFloat = CalcBiasedAbsoluteClipmapLevelForSampling(BaseProjectionData, RayOriginTranslatedWorld);
+	float ClipmapLevelFloat = CalcBiasedAbsoluteClipmapLevelForSampling(BaseProjectionData, RayOriginTranslatedWorld, SceneDepth);
 	int ClipmapIndex = max(0, int(floor(ClipmapLevelFloat)) - BaseProjectionData.ClipmapLevel);
 
 	if (ClipmapIndex >= BaseProjectionData.ClipmapLevelCountRemaining)
@@ -192,7 +192,7 @@ FVirtualShadowMapSampleResult TraceDirectional(
 	Result.ShadowFactor = 1.0f;
 
 	// Find the best resolution mapped clipmap at the ray origin location
-	const FVirtualShadowMapHandle ClipmapHandle = GetMappedClipmap(VirtualShadowMapHandle, TranslatedWorldPosition);
+	const FVirtualShadowMapHandle ClipmapHandle = GetMappedClipmap(VirtualShadowMapHandle, TranslatedWorldPosition, SceneDepth);
 	if (!ClipmapHandle.IsValid())
 	{
 		// TODO: False for valid probably? But make sure that doesn't do anything weird as we have no fallback here
diff --git a/Engine/Shaders/Private/VirtualShadowMaps/VirtualShadowMapProjectionStructs.ush b/Engine/Shaders/Private/VirtualShadowMaps/VirtualShadowMapProjectionStructs.ush
--- a/Engine/Shaders/Private/VirtualShadowMaps/VirtualShadowMapProjectionStructs.ush
+++ b/Engine/Shaders/Private/VirtualShadowMaps/VirtualShadowMapProjectionStructs.ush
@@ -7,6 +7,9 @@
 #include "/Engine/Shared/VirtualShadowMapDefinitions.h"
 #include "VirtualShadowMapHandle.ush"
 
+#define VSM_PSD_OFFSET_RESOLUTION_LOD_BIAS 236 // todo: should use structuredbuffer or improvised offsetof instead
+#define VSM_PSD_OFFSET_LIGHT_TYPE 204
+
 // Rougly the same structure as in VirtualShadowMapArray.h, with some logic around LWC values
 struct FVirtualShadowMapProjectionShaderData
 {
@@ -50,9 +53,8 @@ FVirtualShadowMapProjectionShaderData DecodeVirtualShadowMapProjectionData(ByteA
 	FVirtualShadowMapProjectionShaderData Result;
 	Result.VirtualShadowMapHandle = VirtualShadowMapHandle;
 
-	// Seems the FMatrix forces 16-byte alignment, so ensure padding matches the C++ structure size
-	const uint Stride = 16 * 18;
-	uint CurrentLoadOffset = VirtualShadowMapHandle.GetDataIndex() * Stride;
+	const uint VSMOffset = VirtualShadowMapHandle.GetDataIndex() * VSM_PSD_STRIDE;
+	uint CurrentLoadOffset = VSMOffset;
 	LoadAndIncrementOffset(Result.ShadowViewToClipMatrix, ProjectionData, CurrentLoadOffset);
 	LoadAndIncrementOffset(Result.TranslatedWorldToShadowUVMatrix, ProjectionData, CurrentLoadOffset);
 	LoadAndIncrementOffset(Result.TranslatedWorldToShadowUVNormalMatrix, ProjectionData, CurrentLoadOffset);
@@ -61,6 +63,7 @@ FVirtualShadowMapProjectionShaderData DecodeVirtualShadowMapProjectionData(ByteA
 	// of the data in the calling code and we want to ensure the compiler has simple DCE options.
 	LoadAndIncrementOffset(Result.LightDirection, ProjectionData, CurrentLoadOffset);
 	LoadAndIncrementOffset(Result.LightType, ProjectionData, CurrentLoadOffset);
+	checkSlow(VSM_PSD_OFFSET_LIGHT_TYPE == (CurrentLoadOffset - VSMOffset));
 		 
 	float3 PreViewTranslationHigh;
 	LoadAndIncrementOffset(PreViewTranslationHigh, ProjectionData, CurrentLoadOffset);
@@ -68,6 +71,7 @@ FVirtualShadowMapProjectionShaderData DecodeVirtualShadowMapProjectionData(ByteA
 
 	float3 PreViewTranslationLow;
 	LoadAndIncrementOffset(PreViewTranslationLow, ProjectionData, CurrentLoadOffset);
+	checkSlow(VSM_PSD_OFFSET_RESOLUTION_LOD_BIAS == (CurrentLoadOffset - VSMOffset));
 	LoadAndIncrementOffset(Result.ResolutionLodBias, ProjectionData, CurrentLoadOffset);
 
 	float3 NegativeClipmapWorldOriginOffset;
diff --git a/Engine/Shaders/Private/VirtualShadowMaps/VirtualShadowMapThrottle.usf b/Engine/Shaders/Private/VirtualShadowMaps/VirtualShadowMapThrottle.usf
--- a/Engine/Shaders/Private/VirtualShadowMaps/VirtualShadowMapThrottle.usf
+++ b/Engine/Shaders/Private/VirtualShadowMaps/VirtualShadowMapThrottle.usf
@@ -2,7 +2,9 @@
 
 #include "../Common.ush"
 #include "VirtualShadowMapHandle.ush"
+#include "VirtualShadowMapProjectionStructs.ush"
 #include "/Engine/Shared/VirtualShadowMapDefinitions.h"
+#include "/Engine/Shared/LightDefinitions.h"
 
 StructuredBuffer<uint> PrevThrottleBuffer;
 StructuredBuffer<uint> PrevNanitePerformanceFeedback;
@@ -13,10 +15,15 @@ uint NextVirtualShadowMapDataCount;
 uint NumThrottleBufferEntries;
 RWStructuredBuffer<uint> OutThrottleBuffer;
 
+RWByteAddressBuffer InOutProjectionData;
+
 float DynResThrottleStrength;
 float ThrottleLoadBudget;
 float ThrottleEMAHistoryWeight;
 
+float ThrottleMaxBiasDirectional;
+float ThrottleMaxBiasLocal;
+
 [numthreads(64, 1, 1)]
 void ProcessPrevFramePerfDataCS(uint DispatchThreadId : SV_DispatchThreadID)
 {
@@ -58,6 +65,27 @@ float CalculateCostHeuristic(const float NaniteClustersHW, const float NaniteClu
 	return 0.00012176f * NaniteClustersHW + 0.00002700f * NaniteClustersSW;
 }
 
+void ApplyResolutionLODBias(RWByteAddressBuffer ProjectionData, FVirtualShadowMapHandle VirtualShadowMapHandle, float ExtraBias)
+{
+	FVirtualShadowMapProjectionShaderData Data;
+	
+	// Load
+	uint Offset = VirtualShadowMapHandle.GetDataIndex() * VSM_PSD_STRIDE + VSM_PSD_OFFSET_RESOLUTION_LOD_BIAS;
+	Data.ResolutionLodBias = asfloat(ProjectionData.Load(Offset));
+
+	// Modify
+	Data.ResolutionLodBias += ExtraBias;
+
+	// Store
+	ProjectionData.Store(Offset, asuint(Data.ResolutionLodBias));
+}
+
+uint GetLightType(RWByteAddressBuffer ProjectionData, FVirtualShadowMapHandle VirtualShadowMapHandle)
+{
+	uint Offset = VirtualShadowMapHandle.GetDataIndex() * VSM_PSD_STRIDE + VSM_PSD_OFFSET_LIGHT_TYPE;
+	return ProjectionData.Load(Offset);
+}
+
 [numthreads(64, 1, 1)]
 void UpdateThrottleParametersCS(uint DispatchThreadId : SV_DispatchThreadID)
 {
@@ -122,6 +150,12 @@ void UpdateThrottleParametersCS(uint DispatchThreadId : SV_DispatchThreadID)
 			ThrottleIntensity = ThrottleIntensity * (1.0 - ThrottleEMAHistoryWeight) + PrevThrottleIntensity * ThrottleEMAHistoryWeight;
 		}
 
+		uint LightType = GetLightType(InOutProjectionData, VSMHandle);
+		float MaxBias = (LightType == LIGHT_TYPE_DIRECTIONAL) ? ThrottleMaxBiasDirectional : ThrottleMaxBiasLocal;
+		float ExtraBias = ThrottleIntensity * MaxBias;
+		
+		ApplyResolutionLODBias(InOutProjectionData, VSMHandle, ExtraBias);
+
 		OutThrottleBuffer[ThrottleBufferOffset + VSM_TB_ENTRY_THROTTLE] = asuint(ThrottleIntensity);
 	}
 }
diff --git a/Engine/Shaders/Shared/VirtualShadowMapDefinitions.h b/Engine/Shaders/Shared/VirtualShadowMapDefinitions.h
--- a/Engine/Shaders/Shared/VirtualShadowMapDefinitions.h
+++ b/Engine/Shaders/Shared/VirtualShadowMapDefinitions.h
@@ -97,6 +97,9 @@ VirtualShadowMapDefinitions.h: used in virtual shadow map shaders and C++ code t
 #define VSM_STATUS_MSG_PAGE_MANAGEMENT					0
 #define VSM_STATUS_MSG_OVERFLOW							1
 
+// Projection Shader Data
+#define VSM_PSD_STRIDE (16 * 18)
+
 // Nanite Performance Feedback
 #define VSM_NPF_HEADER_TOTAL_HW_CLUSTERS				0
 #define VSM_NPF_HEADER_TOTAL_SW_CLUSTERS				1
diff --git a/Engine/Source/Runtime/Renderer/Private/VirtualShadowMaps/VirtualShadowMapArray.cpp b/Engine/Source/Runtime/Renderer/Private/VirtualShadowMaps/VirtualShadowMapArray.cpp
--- a/Engine/Source/Runtime/Renderer/Private/VirtualShadowMaps/VirtualShadowMapArray.cpp
+++ b/Engine/Source/Runtime/Renderer/Private/VirtualShadowMaps/VirtualShadowMapArray.cpp
@@ -140,7 +140,7 @@ static TAutoConsoleVariable<int32> CVarFirstPersonPixelRequestLevelClamp(
 	ECVF_Scalability | ECVF_RenderThreadSafe
 );
 
-static TAutoConsoleVariable<float> CVarMaxDOFResolutionBias(
+TAutoConsoleVariable<float> CVarMaxDOFResolutionBias(
 	TEXT("r.Shadow.Virtual.MaxDOFResolutionBias"),
 	1.0f,
 	TEXT("Determine which pixels are out of focus, and request a lower VSM resolution in those areas.\n")
@@ -997,6 +997,7 @@ public:
 
 		SHADER_PARAMETER(uint32, NumThrottleBufferEntries)
 		SHADER_PARAMETER_RDG_BUFFER_UAV(RWBuffer<uint>, OutThrottleBuffer)
+		SHADER_PARAMETER_RDG_BUFFER_UAV(RWByteAddressBuffer, InOutProjectionData)
 	END_SHADER_PARAMETER_STRUCT()
 };
 IMPLEMENT_GLOBAL_SHADER(FUpdateThrottleParametersCS, "/Engine/Private/VirtualShadowMaps/VirtualShadowMapThrottle.usf", "UpdateThrottleParametersCS", SF_Compute);
@@ -1023,8 +1024,6 @@ BEGIN_SHADER_PARAMETER_STRUCT(FBaseGeneratePageFlagsParameters, )
 	SHADER_PARAMETER(uint32, FirstPersonPixelRequestLevelClamp)	
 	SHADER_PARAMETER(float, DOFBiasStrength)
 	SHADER_PARAMETER_STRUCT_INCLUDE(DiaphragmDOF::FDOFCocModelShaderParameters, CocModel)
-	SHADER_PARAMETER_STRUCT_INCLUDE(FThrottlingParameters, ThrottlingParameters)
-	SHADER_PARAMETER_RDG_BUFFER_SRV(StructuredBuffer<uint>, ThrottleBuffer)
 END_SHADER_PARAMETER_STRUCT()
 
 class FGeneratePageFlagsFromPixelsCS : public FVirtualShadowMapPageManagementShader
@@ -2236,6 +2235,8 @@ void FVirtualShadowMapArray::BuildPageAllocations(
 		PassParameters->NumThrottleBufferEntries = GetNumShadowMapSlots();
 		PassParameters->OutThrottleBuffer = GraphBuilder.CreateUAV(ThrottleBufferRDG);
 
+		PassParameters->InOutProjectionData = GraphBuilder.CreateUAV(ProjectionDataRDG);
+
 		auto ComputeShader = GetGlobalShaderMap(Scene.GetFeatureLevel())->GetShader<FUpdateThrottleParametersCS>();
 		ClearUnusedGraphResources(ComputeShader, PassParameters);
 
@@ -2377,8 +2378,6 @@ void FVirtualShadowMapArray::BuildPageAllocations(
 				BaseParameters.MipModeLocal = CVarMarkPixelPagesMipModeLocal.GetValueOnRenderThread();
 				BaseParameters.FirstPersonPixelRequestBias = CVarFirstPersonPixelRequestBias.GetValueOnRenderThread();
 				BaseParameters.FirstPersonPixelRequestLevelClamp = CVarFirstPersonPixelRequestLevelClamp.GetValueOnRenderThread();
-				BaseParameters.ThrottleBuffer = GraphBuilder.CreateSRV(ThrottleBufferRDG);
-				BaseParameters.ThrottlingParameters = ThrottlingParameters;
 
 				DiaphragmDOF::FPhysicalCocModel CocModel;
 				CocModel.Compile(View);
diff --git a/Engine/Source/Runtime/Renderer/Private/VirtualShadowMaps/VirtualShadowMapArray.h b/Engine/Source/Runtime/Renderer/Private/VirtualShadowMaps/VirtualShadowMapArray.h
--- a/Engine/Source/Runtime/Renderer/Private/VirtualShadowMaps/VirtualShadowMapArray.h
+++ b/Engine/Source/Runtime/Renderer/Private/VirtualShadowMaps/VirtualShadowMapArray.h
@@ -132,7 +132,7 @@ struct FVirtualShadowMapProjectionShaderData
 
 	static uint32 PackCullingViewId(int32 SceneRendererPrimaryViewId, const FPersistentViewId& PersistentViewId);
 };
-static_assert(sizeof(FVirtualShadowMapProjectionShaderData) == (16*18), "FVirtualShadowMapProjectionShaderData does not match size in shader. See VirtualShadowMapProjectionStructs.ush.");
+static_assert(sizeof(FVirtualShadowMapProjectionShaderData) == VSM_PSD_STRIDE, "FVirtualShadowMapProjectionShaderData does not match size in shader. See VirtualShadowMapProjectionStructs.ush.");
 
 BEGIN_SHADER_PARAMETER_STRUCT(FVirtualShadowMapPerViewParameters, )
 	// Light grid with only the lights that have VSMs present
diff --git a/Engine/Source/Runtime/Renderer/Private/VirtualShadowMaps/VirtualShadowMapProjection.cpp b/Engine/Source/Runtime/Renderer/Private/VirtualShadowMaps/VirtualShadowMapProjection.cpp
--- a/Engine/Source/Runtime/Renderer/Private/VirtualShadowMaps/VirtualShadowMapProjection.cpp
+++ b/Engine/Source/Runtime/Renderer/Private/VirtualShadowMaps/VirtualShadowMapProjection.cpp
@@ -28,6 +28,7 @@
 #include "BasePassRendering.h"
 #include "BlueNoise.h"
 #include "ShaderPermutationUtils.h"
+#include "PostProcess/DiaphragmDOF.h"
 
 #define MAX_TEST_PERMUTATION 0
 
@@ -55,6 +56,7 @@ static TAutoConsoleVariable<int32> CVarSubsurfaceShadowMinSourceAngle(
 	ECVF_RenderThreadSafe
 );
 
+extern TAutoConsoleVariable<float> CVarMaxDOFResolutionBias;
 
 #if MAX_TEST_PERMUTATION > 0
 static TAutoConsoleVariable<int32> CVarTestPermutation(
@@ -155,6 +157,9 @@ class FVirtualShadowMapProjectionCS : public FGlobalShader
 		SHADER_PARAMETER_RDG_BUFFER_SRV(Buffer, TileListData)
 		RDG_BUFFER_ACCESS(IndirectDispatchArgs, ERHIAccess::IndirectArgs)
 		SHADER_PARAMETER_RDG_BUFFER_UAV(RWStructuredBuffer<uint>, OutStatsBuffer)
+		// DoF
+		SHADER_PARAMETER(float, DOFBiasStrength)
+		SHADER_PARAMETER_STRUCT_INCLUDE(DiaphragmDOF::FDOFCocModelShaderParameters, CocModel)
 	END_SHADER_PARAMETER_STRUCT()
 
 	static void ModifyCompilationEnvironment( const FGlobalShaderPermutationParameters& Parameters, FShaderCompilerEnvironment& OutEnvironment )
@@ -175,6 +180,8 @@ class FVirtualShadowMapProjectionCS : public FGlobalShader
 		}
 
 		OutEnvironment.SetDefine(TEXT("WORK_TILE_SIZE"), VSMProjectionWorkTileSize);
+
+		OutEnvironment.SetDefine(TEXT("VSM_WITH_DOF_BIAS"), 1);
 	}
 
 	static bool ShouldCompilePermutation(const FGlobalShaderPermutationParameters& Parameters)
@@ -283,6 +290,11 @@ static void RenderVirtualShadowMapProjectionCommon(
 		PassParameters->LightUniformVirtualShadowMapId = VirtualShadowMapId;
 		PassParameters->OutShadowFactor = GraphBuilder.CreateUAV( OutputTexture );
 	}
+
+	DiaphragmDOF::FPhysicalCocModel CocModel;
+	CocModel.Compile(View);
+	DiaphragmDOF::SetCocModelParameters(GraphBuilder, &PassParameters->CocModel, CocModel, View.ViewRect.Size().X);
+	PassParameters->DOFBiasStrength = FMath::Max(0.0f, DiaphragmDOF::IsEnabled(View) ? CVarMaxDOFResolutionBias.GetValueOnRenderThread() : 0.0f);
  	
 	bool bDebugOutput = false;
 #if !UE_BUILD_SHIPPING
