Lumen Screen Probe Gather - fixed tile artifacts on indirect specular on cloth pixels in the same tiles as hair pixels

diff --git a/Engine/Shaders/Private/ShadingModelsSampling.ush b/Engine/Shaders/Private/ShadingModelsSampling.ush
--- a/Engine/Shaders/Private/ShadingModelsSampling.ush
+++ b/Engine/Shaders/Private/ShadingModelsSampling.ush
@@ -213,6 +213,7 @@ FBxDFSample SampleBxDF(const uint TermMask, FGBufferData GBuffer, float3 V, floa
 		case SHADINGMODELID_PREINTEGRATED_SKIN:
 		case SHADINGMODELID_CLEAR_COAT:
 		case SHADINGMODELID_TWOSIDED_FOLIAGE:
+		case SHADINGMODELID_CLOTH:
 		case SHADINGMODELID_EYE:
 			return SampleDefaultLitBxDF(TermMask, GBuffer, V, E);
 		case SHADINGMODELID_HAIR:
@@ -233,6 +234,7 @@ bool SupportsSampleBxDF(uint ShadingModelID)
 	case SHADINGMODELID_PREINTEGRATED_SKIN:
 	case SHADINGMODELID_CLEAR_COAT:
 	case SHADINGMODELID_TWOSIDED_FOLIAGE:
+	case SHADINGMODELID_CLOTH:
 	case SHADINGMODELID_EYE:
 	case SHADINGMODELID_HAIR:
 		return true;
