Fix shader compile warning in GammaCorrectionCommon.ush

#rnx

diff --git a/Engine/Shaders/Private/GammaCorrectionCommon.ush b/Engine/Shaders/Private/GammaCorrectionCommon.ush
--- a/Engine/Shaders/Private/GammaCorrectionCommon.ush
+++ b/Engine/Shaders/Private/GammaCorrectionCommon.ush
@@ -9,7 +9,8 @@ half3 LinearTo709Branchless(half3 lin)
 
 half3 Rec709ToLinear(half3 Color)
 {
-	return select(Color > 0.081, pow((Color + 0.099) / 1.099, 1.0 / 0.45), Color / 4.5);
+	// abs() added to silence compiler warning about computing pow() of negative number (both sides of the select are evaluated).
+	return select(Color > 0.081, pow((abs(Color) + 0.099) / 1.099, 1.0 / 0.45), Color / 4.5);
 }
 
 half3 LinearToSrgbBranchless(half3 lin) 
@@ -46,7 +47,8 @@ half3 LinearToSrgb(half3 lin)
 
 half3 sRGBToLinear( half3 Color ) 
 {
-	return select(Color > 0.04045, pow( Color * (1.0 / 1.055) + 0.0521327, 2.4 ), Color * (1.0 / 12.92));
+	// abs() added to silence compiler warning about computing pow() of negative number (both sides of the select are evaluated).
+	return select(Color > 0.04045, pow( abs(Color) * (1.0 / 1.055) + 0.0521327, 2.4 ), Color * (1.0 / 12.92));
 }
 
 /**
