- Use GroupMemoryBarrierWithGroupSync with Radix Sort to fix issues on various GPUs

diff --git a/Engine/Shaders/Private/RadixSortShaders.usf b/Engine/Shaders/Private/RadixSortShaders.usf
--- a/Engine/Shaders/Private/RadixSortShaders.usf
+++ b/Engine/Shaders/Private/RadixSortShaders.usf
@@ -29,8 +29,8 @@
 		must respect the tile size with which these shaders have been compiled.
 ------------------------------------------------------------------------------*/
 
-/** Just use a memory barrier, not a full group sync for parallel prefix scans. */
-#define PPS_BARRIER GroupMemoryBarrier
+/** Full group sync as some GPUs will hang without it, see UE-282658. */
+#define PPS_BARRIER GroupMemoryBarrierWithGroupSync
 
 /** The number of digits per radix. */
 #define DIGIT_COUNT (1 << RADIX_BITS)
