Mobile Post Process:  Fix for bloom flickering artifacts when the viewport rect isn't at an even pixel alignment or size.  Need to clamp UVs to stay in the viewport rectangle.  If bad pixels near the edge of the sun mask are sampled, it can cause a full screen white flicker, due to the bloom logic using clamp filtering when the sun is outside the edge of the view, plus more subtle bloom issues just near the relevant edge of the screen when the sun is in view.  To make the clamp perf neutral or better, simplified the shader to only use 2 UV interpolators rather than 4.

#rnx

diff --git a/Engine/Shaders/Private/PostProcessMobile.usf b/Engine/Shaders/Private/PostProcessMobile.usf
--- a/Engine/Shaders/Private/PostProcessMobile.usf
+++ b/Engine/Shaders/Private/PostProcessMobile.usf
@@ -65,6 +65,7 @@ float4 DofBlurSizeAndInvSize;
 float4 BufferASizeAndInvSize;
 float4 BufferBSizeAndInvSize;
 float4 BloomUpSizeAndInvSize;
+float4 BufferViewRect;
 
 // Point on circle.
 float2 Circle(float Start, float Points, float Point) 
@@ -288,7 +289,7 @@ void BloomVS_Mobile(
 	in float4 InPosition : ATTRIBUTE0,
 	in float2 InTexCoord : ATTRIBUTE1,
 	in FStereoVSInput StereoInput,
-	out float2 OutTexCoords[4] : TEXCOORD0,
+	out float2 OutTexCoords[2] : TEXCOORD0,
 	out FStereoVSOutput StereoOutput,
 	out float4 OutPosition : SV_POSITION
 	)
@@ -298,9 +299,7 @@ void BloomVS_Mobile(
 	DrawRectangle(InPosition, InTexCoord, OutPosition, TransformedUV);
 
 	OutTexCoords[0] = TransformedUV + BufferSizeAndInvSize.zw * float2(-1, -1);
-	OutTexCoords[1] = TransformedUV + BufferSizeAndInvSize.zw * float2( 1, -1);
-	OutTexCoords[2] = TransformedUV + BufferSizeAndInvSize.zw * float2(-1,  1);
-	OutTexCoords[3] = TransformedUV + BufferSizeAndInvSize.zw * float2( 1,  1);
+	OutTexCoords[1] = TransformedUV + BufferSizeAndInvSize.zw * float2( 1,  1);
 }
 
 
@@ -321,7 +320,7 @@ float HdrWeightInv4(float3 Color)
 }
 
 void BloomPS_Mobile(
-	float2 InUVs[4] : TEXCOORD0,
+	float2 InUVs[2] : TEXCOORD0,
 	in FStereoPSInput StereoInput
 #if MOBILE_USEBLOOM
 
@@ -367,12 +366,20 @@ void BloomPS_Mobile(
 {
 	half3 AverageColor = 0.0f;
 
+	// Need UV clamping in case the viewport rect is at an odd size or pixel alignment.  This can be the case even if the
+	// final output rect is even pixel alignment, factoring in dynamic resolution scaling.
+	float2 ClampedUVs[4];
+	ClampedUVs[0] = max(InUVs[0], BufferViewRect.xy);		// -1, -1
+	ClampedUVs[1] = min(InUVs[1], BufferViewRect.zw); //  1   1
+	ClampedUVs[2] = float2(ClampedUVs[1].x, ClampedUVs[0].y); //  1  -1
+	ClampedUVs[3] = float2(ClampedUVs[0].x, ClampedUVs[1].y); // -1,  1
+	
 #if MOBILE_USEBLOOM || MOBILE_USESUN || MOBILE_USEEYEADAPTATION
 	
-	float4 C0 = SampleSceneColor(InUVs[0],GetEyeIndex(StereoInput));
-	float4 C1 = SampleSceneColor(InUVs[1],GetEyeIndex(StereoInput));
-	float4 C2 = SampleSceneColor(InUVs[2],GetEyeIndex(StereoInput));
-	float4 C3 = SampleSceneColor(InUVs[3],GetEyeIndex(StereoInput));
+	float4 C0 = SampleSceneColor(ClampedUVs[0],GetEyeIndex(StereoInput));
+	float4 C1 = SampleSceneColor(ClampedUVs[1],GetEyeIndex(StereoInput));
+	float4 C2 = SampleSceneColor(ClampedUVs[2],GetEyeIndex(StereoInput));
+	float4 C3 = SampleSceneColor(ClampedUVs[3],GetEyeIndex(StereoInput));
 
 	C0 *= HdrWeight4(C0);
 	C1 *= HdrWeight4(C1);
@@ -407,10 +414,10 @@ void BloomPS_Mobile(
 #endif
 
 #if MOBILE_USEDOF || MOBILE_USESUN
-	half A0 = SunShaftAndDofTexture.Sample(SunShaftAndDofSampler, InUVs[0]).r;
-	half A1 = SunShaftAndDofTexture.Sample(SunShaftAndDofSampler, InUVs[1]).r;
-	half A2 = SunShaftAndDofTexture.Sample(SunShaftAndDofSampler, InUVs[2]).r;
-	half A3 = SunShaftAndDofTexture.Sample(SunShaftAndDofSampler, InUVs[3]).r;
+	half A0 = SunShaftAndDofTexture.Sample(SunShaftAndDofSampler, ClampedUVs[0]).r;
+	half A1 = SunShaftAndDofTexture.Sample(SunShaftAndDofSampler, ClampedUVs[1]).r;
+	half A2 = SunShaftAndDofTexture.Sample(SunShaftAndDofSampler, ClampedUVs[2]).r;
+	half A3 = SunShaftAndDofTexture.Sample(SunShaftAndDofSampler, ClampedUVs[3]).r;
 #endif
 
 // In the case of both DOF and SUN,
diff --git a/Engine/Source/Runtime/Renderer/Private/PostProcess/PostProcessMobile.cpp b/Engine/Source/Runtime/Renderer/Private/PostProcess/PostProcessMobile.cpp
--- a/Engine/Source/Runtime/Renderer/Private/PostProcess/PostProcessMobile.cpp
+++ b/Engine/Source/Runtime/Renderer/Private/PostProcess/PostProcessMobile.cpp
@@ -224,6 +224,7 @@ public:
 		SHADER_PARAMETER_SAMPLER(SamplerState, SceneColorSampler)
 		SHADER_PARAMETER_RDG_TEXTURE(Texture2D, SunShaftAndDofTexture)
 		SHADER_PARAMETER_SAMPLER(SamplerState, SunShaftAndDofSampler)
+		SHADER_PARAMETER(FVector4f, BufferViewRect)
 		RENDER_TARGET_BINDING_SLOTS()
 	END_SHADER_PARAMETER_STRUCT()
 
@@ -350,6 +351,12 @@ FMobileBloomSetupOutputs AddMobileBloomSetupPass(FRDGBuilder& GraphBuilder, cons
 	PSShaderParameters->SunShaftAndDofTexture = Inputs.SunShaftAndDof.Texture;
 	PSShaderParameters->SunShaftAndDofSampler = TStaticSamplerState<SF_Bilinear, AM_Clamp, AM_Clamp, AM_Clamp>::GetRHI();
 
+	PSShaderParameters->BufferViewRect = FVector4f(
+		(Inputs.SceneColor.ViewRect.Min.X + 0.5f) / BufferSize.X,
+		(Inputs.SceneColor.ViewRect.Min.Y + 0.5f) / BufferSize.Y,
+		(Inputs.SceneColor.ViewRect.Max.X - 0.5f) / BufferSize.X,
+		(Inputs.SceneColor.ViewRect.Max.Y - 0.5f) / BufferSize.Y);
+
 	const FScreenPassTextureViewport InputViewport(Inputs.SceneColor);
 	const FScreenPassTextureViewport OutputViewport(DestRenderTargets[0]);
 
