Rough refraction fox for sm5.

#tests pc editor dx11
#rnx

diff --git a/Engine/Shaders/Private/DistortFiltering.usf b/Engine/Shaders/Private/DistortFiltering.usf
--- a/Engine/Shaders/Private/DistortFiltering.usf
+++ b/Engine/Shaders/Private/DistortFiltering.usf
@@ -60,7 +60,7 @@ void DownsampleColorCS(uint3 ThreadId : SV_DispatchThreadID)
 	}
 
 	float2 ScaledUVs = (float2(ThreadId.xy) * 2.0 + 1.0f) / float2(SrcMipResolution);
-	float4 OutColor = SourceTexture.Sample(SourceSampler, ScaledUVs, 0); // No need to specify mip index, the SRV is created for the acual mip level.
+	float4 OutColor = SourceTexture.SampleLevel(SourceSampler, ScaledUVs, 0); // No need to specify mip index, the SRV is created for the acual mip level.
 
 	OutTextureMipColor[uint2(ThreadId.xy)] = OutColor;
 }
