Force adreno shader compiler to produce correct asm, as previously branching asm was incorrect

diff --git a/Engine/Shaders/Private/ShadingModels.ush b/Engine/Shaders/Private/ShadingModels.ush
--- a/Engine/Shaders/Private/ShadingModels.ush
+++ b/Engine/Shaders/Private/ShadingModels.ush
@@ -248,6 +248,14 @@ FDirectLighting DefaultLitBxDF( FGBufferData GBuffer, half3 N, half3 V, half3 L,
 	Lighting.Specular = 0;
 	Lighting.Transmission = 0;
 
+	
+#if FEATURE_LEVEL == FEATURE_LEVEL_ES3_1
+	// adjust N for Adreno shader compilation bug. See UE-274816
+	N = N * 0.5f;
+	N = normalize(N);
+#endif
+	
+	
 	BRANCH
 	if (NoL > 0.0f)
 	{
