[PostProcessing] Fix shader compilation error in FUpscaleVS due to missing define.

diff --git a/Engine/Shaders/Private/PostProcessUpscale.usf b/Engine/Shaders/Private/PostProcessUpscale.usf
--- a/Engine/Shaders/Private/PostProcessUpscale.usf
+++ b/Engine/Shaders/Private/PostProcessUpscale.usf
@@ -19,6 +19,10 @@
 #define METHOD_SET_SIMPLE 1
 #define METHOD_SET_COMPLEX 2
 
+#ifndef DIM_METHOD
+#define DIM_METHOD 0
+#endif
+
 #if DIM_METHOD == UPSCALE_METHOD_NEAREST
 	#define INCLUDE_UPSCALE_METHOD_NEAREST 1
 #elif DIM_METHOD == UPSCALE_METHOD_BILINEAR
