Fix shader compilation warning.

diff --git a/Engine/Shaders/Private/HairStrands/HairStrandsBindingCommon.ush b/Engine/Shaders/Private/HairStrands/HairStrandsBindingCommon.ush
--- a/Engine/Shaders/Private/HairStrands/HairStrandsBindingCommon.ush
+++ b/Engine/Shaders/Private/HairStrands/HairStrandsBindingCommon.ush
@@ -180,7 +180,7 @@ float3 ApplyRBF(
 	const float4 RestSampleOffset = bEnableLocalSpace ? InRestSamplePositionsBuffer[0] : float4(0,0,0,0);
 	const float4 DeformedSampleOffset = bEnableLocalSpace ? InDeformedSamplePositionsBuffer[0] : float4(0,0,0,0);
 	
-	float3 ControlPoint = RestControlPoint-RestSampleOffset;
+	float3 ControlPoint = RestControlPoint-RestSampleOffset.xyz;
 	// Apply rbf interpolation from the samples set
 	for (int i = 0; i < InMaxSampleCount; ++i)
 	{
