Fix quantization issue on specular value.

diff --git a/Engine/Shaders/Private/PackUnpack.ush b/Engine/Shaders/Private/PackUnpack.ush
--- a/Engine/Shaders/Private/PackUnpack.ush
+++ b/Engine/Shaders/Private/PackUnpack.ush
@@ -371,3 +371,44 @@ uint  PackTileCoord14bits(uint2 In) { return PackTileCoordXXbits  (In, 0x3FFF, 1
 uint2 UnpackTileCoord14bits(uint In){ return UnpackTileCoordXXbits(In, 0x3FFF, 14u); }
 uint  PackTileCoord16bits(uint2 In) { return PackTileCoordXXbits  (In, 0xFFFF,16u); }
 uint2 UnpackTileCoord16bits(uint In){ return UnpackTileCoordXXbits(In, 0xFFFF,16u); }
+
+
+//
+// Dithering
+//
+
+// Return 0 if v is zero, 1 otherwise
+float IsNonZeroFast(float v)
+{
+	return float(min(asuint(v) & 0x7FFFFFFF, 1u));
+}
+
+float3 IsNonZeroFast(float3 v)
+{
+	return float3(min(asuint(v) & 0x7FFFFFFF, 1u.xxx));
+}
+
+float DitherIfNonBlack(float Input, float Dither)
+{
+	// IsNonZeroFast ensures dithering is not applied when Input is zero. Input us assumed to be in [0..1].
+	float NewInput = Input + Dither * IsNonZeroFast(Input);
+	return saturate(NewInput);
+}
+
+float3 DitherIfNonBlack(float3 Input, float3 Dither)
+{
+	// IsNonZeroFast ensures dithering is not applied when Input is zero. Input us assumed to be in [0..1].
+	float3 NewInput = Input + Dither * IsNonZeroFast(Input);
+	return saturate(NewInput);
+}
+
+float DitherXbits(float Value, float Dither, float QuantizationScale)
+{
+	const float CenteredDither = Dither - 0.5f;
+	const float DitheredXbits = CenteredDither * QuantizationScale;
+	return DitherIfNonBlack(saturate(Value), DitheredXbits);
+}
+
+#define Dither6bits(Value, Dither) DitherXbits(Value, Dither, 1.f/63.f)
+#define Dither7bits(Value, Dither) DitherXbits(Value, Dither, 1.f/127.f)
+#define Dither8bits(Value, Dither) DitherXbits(Value, Dither, 1.f/255.f)
\ No newline at end of file
diff --git a/Engine/Shaders/Private/ShadingModelsMaterial.ush b/Engine/Shaders/Private/ShadingModelsMaterial.ush
--- a/Engine/Shaders/Private/ShadingModelsMaterial.ush
+++ b/Engine/Shaders/Private/ShadingModelsMaterial.ush
@@ -28,7 +28,12 @@ void SetGBufferForShadingModel(
 	GBuffer.WorldTangent = MaterialParameters.WorldTangent;
 	GBuffer.BaseColor = BaseColor;
 	GBuffer.Metallic = Metallic;
+	#if SHADING_PATH_MOBILE
 	GBuffer.Specular = Specular;
+	#else
+	// Dither value for avoiding quantization artifact (for now, only Specular is dithered based on artists' feedback)
+	GBuffer.Specular = Dither8bits(Specular, Dither);
+	#endif
 	GBuffer.Roughness = Roughness;
 	GBuffer.Anisotropy = Anisotropy;
 	GBuffer.ShadingModelID = ShadingModel;
diff --git a/Engine/Shaders/Private/Substrate/Substrate.ush b/Engine/Shaders/Private/Substrate/Substrate.ush
--- a/Engine/Shaders/Private/Substrate/Substrate.ush
+++ b/Engine/Shaders/Private/Substrate/Substrate.ush
@@ -3389,31 +3389,6 @@ void FSubstrateBSDF::SubstrateBakeGlintToReflectivity()
 #endif
 }
 
-// Return 0 if v is zero, 1 otherwise
-float IsNonZeroFast(float v)
-{
-	return float(min(asuint(v) & 0x7FFFFFFF, 1u));
-}
-
-float3 IsNonZeroFast(float3 v)
-{
-	return float3(min(asuint(v) & 0x7FFFFFFF, 1u.xxx));
-}
-
-// This is used to respect artistic desire, e.g. shutdown diffuse or specular.
-float3 DitherIfNonBlack(float3 Input, float3 Dither)
-{
-	// IsNonZeroFast ensures dithering is not applied when Input is zero. Input us assumed to be in [0..1].
-	float3 NewInput = Input + Dither * IsNonZeroFast(Input);
-	return saturate(NewInput);
-}
-float DitherIfNonBlack(float Input, float Dither)
-{
-	// IsNonZeroFast ensures dithering is not applied when Input is zero. Input us assumed to be in [0..1].
-	float NewInput = Input + Dither * IsNonZeroFast(Input);
-	return saturate(NewInput);
-}
-
 #if !SUBSTRATE_COMPILER_SUPPORTS_STRUCT_FORWARD_DECLARATION
 void UpdateSingleBSDFOperatorCoverageTransmittance(
 	inout FSubstrateTree SubstrateTree,
@@ -3861,10 +3836,7 @@ float3 UnpackR7G7B6Linear(uint rgb)
 
 uint PackR7(float Value, float Dither)
 {
-	const float CenteredDither = Dither - 0.5f;
-	const float Dither7bits = CenteredDither * 1.0f / 127.0f;
-	const float DitheredValue = DitherIfNonBlack(saturate(Value), Dither7bits);
-	return uint(DitheredValue * 127.0f);
+	return uint(Dither7bits(Value, Dither) * 127.0f);
 }
 
 float UnpackR7(uint Value)
@@ -3874,10 +3846,7 @@ float UnpackR7(uint Value)
 
 uint PackR6(float Value, float Dither)
 {
-	const float CenteredDither = Dither - 0.5f;
-	const float Dither7bits = CenteredDither * 1.0f / 63.0f;
-	const float DitheredValue = DitherIfNonBlack(saturate(Value), Dither7bits);
-	return uint(DitheredValue * 63.0f);
+	return uint(Dither6bits(Value, Dither) * 63.0f);
 }
 
 // Simple Octahedral from "A Survey of Efficient Representations for Independent Unit Vectors" http://jcgt.org/published/0003/02/01/
