Fix another shader compilation issue on groom binding.

diff --git a/Engine/Shaders/Private/HairStrands/HairStrandsBindingCommon.ush b/Engine/Shaders/Private/HairStrands/HairStrandsBindingCommon.ush
--- a/Engine/Shaders/Private/HairStrands/HairStrandsBindingCommon.ush
+++ b/Engine/Shaders/Private/HairStrands/HairStrandsBindingCommon.ush
@@ -193,6 +193,6 @@ float3 ApplyRBF(
 	ControlPoint += InMeshSampleWeightsBuffer[InMaxSampleCount + 2].xyz * (RestControlPoint.y-RestSampleOffset.y);
 	ControlPoint += InMeshSampleWeightsBuffer[InMaxSampleCount + 3].xyz * (RestControlPoint.z-RestSampleOffset.z);
 
-	ControlPoint += DeformedSampleOffset;
+	ControlPoint += DeformedSampleOffset.xyz;
 	return ControlPoint;
 }
