Fix a crash caused by enabling a Nanite visualization mode when instanced stereo rendering is enabled.

#tests confirmed fix in EngineTest editor w/ Instanced Stereo enabled

diff --git a/Engine/Shaders/Private/PostProcessVisualizeBuffer.usf b/Engine/Shaders/Private/PostProcessVisualizeBuffer.usf
--- a/Engine/Shaders/Private/PostProcessVisualizeBuffer.usf
+++ b/Engine/Shaders/Private/PostProcessVisualizeBuffer.usf
@@ -11,7 +11,7 @@ SamplerState InputSampler;
 float4 SelectionColor;
 
 /** Main pixel shader function */
-void MainPS(noperspective float4 UVAndScreenPos : TEXCOORD0, float4 SvPosition : SV_POSITION, out float4 OutColor : SV_Target0)
+void MainPS(noperspective float4 UVAndScreenPos : TEXCOORD0, FStereoPSInput StereoInput, float4 SvPosition : SV_POSITION, out float4 OutColor : SV_Target0)
 {
 	float Alpha = 1;
 
