Fixed: Incorrect ParticleId for Niagara FX mesh rendering on a mobile platforms

diff --git a/Engine/Shaders/Private/SceneData.ush b/Engine/Shaders/Private/SceneData.ush
--- a/Engine/Shaders/Private/SceneData.ush
+++ b/Engine/Shaders/Private/SceneData.ush
@@ -741,23 +741,31 @@ float4 LoadInstancePayloadExtensionElement(FInstanceSceneData InstanceData, uint
 float4 LoadInstanceCustomDataElement(FInstanceSceneData SceneData, uint Float4Index)
 {
 #if ENABLE_PER_INSTANCE_CUSTOM_DATA
-	const uint NumCustomFloat4s = (SceneData.CustomDataCount + 3u) >> 2u;
-	if (SceneData.CustomDataOffset != 0xFFFFFFFFu && Float4Index < NumCustomFloat4s)
-	{
-		return LoadInstancePayloadDataElement(SceneData.CustomDataOffset + Float4Index);
-	}
+	#if SHADER_USES_PRIMITIVE_UBO
+		// TODO: atm this is not used anywhere, but will be good to implement as well
+		return (float4)0.0f;
+	#else
+		const uint NumCustomFloat4s = (SceneData.CustomDataCount + 3u) >> 2u;
+		if (SceneData.CustomDataOffset != 0xFFFFFFFFu && Float4Index < NumCustomFloat4s)
+		{
+			return LoadInstancePayloadDataElement(SceneData.CustomDataOffset + Float4Index);
+		}
+	#endif
 #endif
-
 	return (float4)0.0f;
 }
 
 float LoadInstanceCustomDataFloat(FInstanceSceneData SceneData, uint FloatIndex)
 {
 #if ENABLE_PER_INSTANCE_CUSTOM_DATA
-	const uint Float4Index = FloatIndex >> 2u;
-	const uint ComponentIndex = FloatIndex % 4u;
-	const float4 Element = LoadInstanceCustomDataElement(SceneData, Float4Index);
-	return Element[ComponentIndex];
+	#if SHADER_USES_PRIMITIVE_UBO
+		return LoadInstanceCustomDataElementUBO(SceneData.PrimitiveId, FloatIndex, 0.0f);
+	#else
+		const uint Float4Index = FloatIndex >> 2u;
+		const uint ComponentIndex = FloatIndex % 4u;
+		const float4 Element = LoadInstanceCustomDataElement(SceneData, Float4Index);
+		return Element[ComponentIndex];
+	#endif
 #else
 	return 0.0f;
 #endif
@@ -1334,6 +1342,7 @@ FSceneDataIntermediates GetSceneDataIntermediates(uint InstanceIdOffset, uint Dr
 	Intermediates.InstanceData.LocalBoundsCenter = Intermediates.Primitive.InstanceLocalBoundsCenter;
 	Intermediates.InstanceData.LocalBoundsExtent = Intermediates.Primitive.InstanceLocalBoundsExtent;
 	Intermediates.InstanceData.ValidInstance = true;
+	Intermediates.InstanceData.PrimitiveId = DrawInstanceId;
 	Intermediates.CullingFlags = INSTANCE_CULLING_FLAGS_DEFAULT;
 	Intermediates.ViewIndex = 0U;
 	Intermediates.PrimitiveId = DrawInstanceId;
