Path Tracer: Refactor primary rays visualization mode to include indirect lighting

Unify sampling controls to a single SPP cvar which controls both direct and indirect lighting quality. Add cvar for number of bounces (default to 1), emissive materials (defaults off).

Add experimental path guiding implementation (currently mostly effective on emissive surfaces)

#rnx

diff --git a/Engine/Shaders/Private/PathTracing/PathTracingDebug.usf b/Engine/Shaders/Private/PathTracing/PathTracingDebug.usf
--- a/Engine/Shaders/Private/PathTracing/PathTracingDebug.usf
+++ b/Engine/Shaders/Private/PathTracing/PathTracingDebug.usf
@@ -4,9 +4,15 @@
 
 #define USE_BSDF_SAMPLING 1 // 0: light sampling only (no MIS) 1: light sampling and BSDF sampling (with MIS)
 
+#define INCLUDE_TRANSLUCENT   0 // include objects tagged as TRANSLUCENT in the TLAS (transparency is not implemented yet, so better to leave this off for now)
+
 #define USE_PATH_TRACING_LIGHT_GRID	1
 #define USE_RAY_TRACING_DECAL_GRID	1
 
+#ifndef PATH_TRACER_PRIMARY_RAYS // Should be set by permutation
+#define PATH_TRACER_PRIMARY_RAYS 1
+#endif
+
 
 #include "../Common.ush"
 #include "../BlueNoise.ush"
@@ -30,9 +36,24 @@ RWTexture2D<float4> RWSceneColor;
 RWTexture2D<float>  RWSceneDepth;
 RaytracingAccelerationStructure TLAS;
 
+RWStructuredBuffer<FMarkovChainState> MarkovChainStateGrid;
+uint MarkovChainStateGridSize;
+uint UsePathGuiding;
+
 int DebugMode;
-int NumLightSamples;
-int NumIndirectSamples;
+int SamplesPerPixel;
+int MaxBounces;
+int EnableEmissive;
+
+uint4 GridHash(uint4 P)
+{
+	const uint4 S = uint4(14, 15, 16, 17);
+	const uint K = 1664525u, O = 1013904223u; // pcg4d constants
+	P ^= P >> S; P = P * K + O; P ^= P.yzwx ^ P.zwxy;
+	P ^= P >> S; P = P * K + O; P ^= P.yzwx ^ P.zwxy;
+	P ^= P >> S; P = P * K + O; P ^= P.yzwx ^ P.zwxy;
+	return P;
+}
 
 float3 TraceShadowRay(FRayDesc Ray, float PathRoughness, float PreviousPathRoughness, uint MissShaderIndex, bool bCastShadows)
 {
@@ -42,11 +63,6 @@ float3 TraceShadowRay(FRayDesc Ray, float PathRoughness, float PreviousPathRough
 		return 1.0;
 	}
 
-	const uint RayFlags = RAY_FLAG_ACCEPT_FIRST_HIT_AND_END_SEARCH | RAY_FLAG_SKIP_CLOSEST_HIT_SHADER;
-	const uint InstanceInclusionMask = PATHTRACER_MASK_SHADOW;
-	const uint RayContributionToHitGroupIndex = RAY_TRACING_SHADER_SLOT_SHADOW;
-	const uint MultiplierForGeometryContributionToShaderIndex = RAY_TRACING_NUM_SHADER_SLOTS;
-
 	FPackedPathTracingPayload PackedPayload = InitPathTracingVisibilityPayload(PathRoughness);
 
 	if (!bCastShadows)
@@ -57,10 +73,10 @@ float3 TraceShadowRay(FRayDesc Ray, float PathRoughness, float PreviousPathRough
 	}
 	TraceRay(
 		TLAS,
-		RayFlags,
-		InstanceInclusionMask,
-		RayContributionToHitGroupIndex,
-		MultiplierForGeometryContributionToShaderIndex,
+		RAY_FLAG_ACCEPT_FIRST_HIT_AND_END_SEARCH | RAY_FLAG_SKIP_CLOSEST_HIT_SHADER,
+		PATHTRACER_MASK_SHADOW,
+		RAY_TRACING_SHADER_SLOT_SHADOW,
+		RAY_TRACING_NUM_SHADER_SLOTS,
 		MissShaderIndex,
 		Ray.GetNativeDesc(),
 		PackedPayload);
@@ -70,20 +86,18 @@ float3 TraceShadowRay(FRayDesc Ray, float PathRoughness, float PreviousPathRough
 		// We didn't run the miss shader, therefore we must have hit something opaque (or reached full opacity)
 		return 0.0;
 	}
-
 	float3 Throughput = PackedPayload.GetRayThroughput();
-
 	return Throughput;
 }
 
-struct PathTracingResult {
+struct FPathTracingResult {
 	float4 Radiance;
 	float Z;
 };
 
-PathTracingResult MakeResult(float3 Radiance, float Alpha, float Z)
+FPathTracingResult MakeResult(float3 Radiance, float Alpha, float Z)
 {
-	PathTracingResult Result;
+	FPathTracingResult Result = (FPathTracingResult) 0;
 	Result.Radiance.rgb = Radiance;
 	Result.Radiance.a = Alpha;
 	Result.Z = Z;
@@ -106,10 +120,505 @@ float3 AnisoDebugColor(float InAniso)
 	return float3(0.0, pow(AniG, 2.2), pow(AniB, 2.2));
 }
 
-#define INCLUDE_TRANSLUCENT   0
+float3 StateGetPos(FMarkovChainState State)
+{
+	return State.WeightSum > 0.0 ? State.WeightedTarget / State.WeightSum : State.WeightedTarget;
+}
+
+float StateGetCos(FMarkovChainState State)
+{
+	return State.WeightSum > 0 ? State.WeightedCosine / State.WeightSum : 0.0;
+}
+
+float3 StateGetDir(FMarkovChainState State, float3 AbsoluteWorldPos)
+{
+	return normalize(StateGetPos(State) - AbsoluteWorldPos);
+}
+
+float4 StateGetVMFLobe(FMarkovChainState State, float3 AbsoluteWorldPos)
+{
+	float3 StateP = StateGetPos(State);
+	float  StateR = StateGetCos(State);
+	float Np = 0.2 / length2(AbsoluteWorldPos - StateP); // Prior (see Section 4, step 2)
+	float r = MISWeightBalanced(State.N * State.N, Np) * StateR;
+	r = min(r, 0.99999994); // clamp to avoid singularity at r=1
+    float Kappa = r * (3.0 - r * r) / (1.0 - r * r); // Banerjee et al. 2005 - eq 4.4
+	return float4(normalize(StateP - AbsoluteWorldPos), Kappa);
+}
+
+// Based on Yusuke Tokuyoshi, A Numerically Stable Implementation of the von Mises–Fisher Distribution on S2
+float log1p(const float x)
+{
+	return abs(x) < 0.01 ? sign(x) * mad(-0.5f * x, x, x) : log(1 + x);
+}
+float expm1(const float x)
+{
+	return abs(x) < 0.03 ? sign(x) * mad(0.5f * x, x, x) : exp(x) - 1;
+}
+
+float x_over_expm1(const float x)
+{
+	return abs(x) < 0.03 ? rcp(0.5 * x + 1) : x * rcp(exp(x) - 1);
+}
+
+float VMFPdf(const float3 W, const float4 LobeData) {
+	const float3 Mu = LobeData.xyz;
+	const float Kappa = LobeData.w;
+	const float3 D = W - Mu;
+	return exp(-0.5f * Kappa * dot(D, D)) * x_over_expm1(-2.0f * Kappa) * (1.0 / (4.0 * PI));
+}
+
+float3 VMFSample(const float4 LobeData, const float2 Rand) {
+	const float3 Mu = LobeData.xyz;
+	const float Kappa = LobeData.w;
+	const float Phi = (2 * PI) * Rand.x;
+    const float r = Kappa > ( 1.192092896e-07F / 4.0f) ? log1p(Rand.y * expm1(-2.0f * Kappa)) / Kappa : -2.0f * Rand.y;
+    const float SinTheta = sqrt(-mad(r, r, 2.0f * r));
+	const float CosTheta = 1 + r;
+    const float3 Dir = float3(cos(Phi) * SinTheta, sin(Phi) * SinTheta, CosTheta);
+	return TangentToWorld(Dir, Mu);
+}
+
+int GetDominantSide(float3 N)
+{
+	if (abs(N.x) > abs(N.y) && abs(N.x) > abs(N.z))
+		return (N.x < 0);
+	if (abs(N.y) > abs(N.z))
+		return 2 + (N.y < 0);
+	return 4 + (N.z < 0);
+}
+
+
+#if PATH_TRACER_PRIMARY_RAYS
+
+FPathTracingResult PathTracingDebugCore(FRayDesc Ray, uint2 LaunchIndex, int SampleIndex)
+{
+	FPathTracingResult Result = (FPathTracingResult) 0;
+	Result.Radiance.a = 1;
+	float3 PathThroughput = 1;
+	float PathRoughness = 0;
+	uint ScatterType = PATHTRACER_SCATTER_CAMERA;
+	RandomSequence RandSequence;
+	RandomSequence_Initialize(RandSequence, LaunchIndex, SampleIndex, View.StateFrameIndex, SamplesPerPixel);
+
+	// State about the markov chain selection that needs to be kept from one bounce to the next
+	float ScoresSum = 0;
+	float3 PreviousN = 0;
+	FMarkovChainState ChosenState = (FMarkovChainState)0;
+	// MCPG parameters
+#define N_MC  5 // From paper, seems to be the smallest value which produces good results, some subtle improvements from higher values at the cost of performance
+	const float AdaptiveGridProb = 0.7; // Mentioned in text - TODO: could rely on adaptive grid only with a tweak to level distribution?
+	const float StaticGridWidth = 25.6; // matches level 24 of adaptive grid
+	const float LevelScale = 3.0; 		// 3 levels per octave
+	const float MinWidth = 0.1; 		// 1mm - smallest possible voxel size
+	const float TanFactor = 0.0087;		// ~2*tan(0.25 degrees)
+	const int 	NMax = 1024;			// Blend factors [See Dittebrandt et al.]
+	const float AlphaMin = 1.0 / 100.0; // Allow blending over 100 frames, higher than original publication because the chains get reset during BSDF sampling
+	const float MaterialProb = 0.15f; 	// Chance of BSDF sampling (TODO: adapt this as a function of roughness)
+
+	for (int Bounce = 0; Bounce <= MaxBounces; Bounce++)
+	{
+		// Trace Ray
+		const uint RayFlags = Bounce == 0 ? RAY_FLAG_CULL_BACK_FACING_TRIANGLES : 0;
+#if INCLUDE_TRANSLUCENT
+		const uint InstanceInclusionMask = Bounce == 0 ? PATHTRACER_MASK_CAMERA | PATHTRACER_MASK_CAMERA_TRANSLUCENT : PATHTRACER_MASK_INDIRECT | PATHTRACER_MASK_INDIRECT_TRANSLUCENT;
+#else
+		const uint InstanceInclusionMask = Bounce == 0 ? PATHTRACER_MASK_CAMERA : PATHTRACER_MASK_INDIRECT;
+#endif
+		const uint MissShaderIndex = 0;
+
+		FPackedPathTracingPayload PackedPayload = InitPathTracingPayload(ScatterType, PathRoughness);
+		PackedPayload.SetStochasticSlabRand(RandomSequence_GenerateSample1D(RandSequence));
+#if !INCLUDE_TRANSLUCENT
+		PackedPayload.SetFlag(PATH_TRACING_PAYLOAD_INPUT_FLAG_IGNORE_TRANSLUCENT);
+#endif
+
+#if PATH_TRACER_USE_SER 
+		NvHitObject Hit;
+		NvTraceRayHitObject(
+			TLAS,
+			RayFlags,
+			InstanceInclusionMask,
+			RAY_TRACING_SHADER_SLOT_MATERIAL,
+			RAY_TRACING_NUM_SHADER_SLOTS,
+			MissShaderIndex,
+			Ray.GetNativeDesc(),
+			PackedPayload,
+			Hit);
+		if (Bounce > 0)
+		{
+			NvReorderThread(Hit);
+		}
+		NvInvokeHitObject(TLAS, Hit, PackedPayload);
+#else
+		TraceRay(
+			TLAS,
+			RayFlags,
+			InstanceInclusionMask,
+			RAY_TRACING_SHADER_SLOT_MATERIAL,
+			RAY_TRACING_NUM_SHADER_SLOTS,
+			MissShaderIndex,
+			Ray.GetNativeDesc(),
+			PackedPayload);
+#endif
+
+		// If miss - exit
+		if (PackedPayload.IsMiss())
+		{
+			break; // exit bounce loop
+		}
+
+		const float3 TranslatedWorldPos = Ray.Origin + PackedPayload.HitT * Ray.Direction;
+		const float3 WorldNormal = PackedPayload.UnpackWorldNormal();
+		const uint PrimitiveLightingChannelMask = PackedPayload.GetPrimitiveLightingChannelMask();
+
+		const bool bCameraRay = Bounce == 0;
+		const bool bLastBounce = Bounce == MaxBounces;
+
+		FLightLoopCount LightLoopCount = (FLightLoopCount) 0; 
+		if (!bLastBounce)
+		{
+			LightLoopCount = LightGridLookup(TranslatedWorldPos);
+		}
+
+		// Capture depth on camera ray
+		if (bCameraRay)
+		{
+			const float3 ViewZ = View.ViewToTranslatedWorld[2].xyz;
+			Result.Z = ConvertToDeviceZ(dot(TranslatedWorldPos, ViewZ));
+			Result.Radiance.a = 0; // opaque alpha (we aren't handling translucency right now)
+		}
+
+		// Add radiance (emissive surfaces)
+		float3 LocalRadiance = 0; 
+		if (bCameraRay || EnableEmissive)
+		{
+			LocalRadiance += PathThroughput * PackedPayload.UnpackRadiance();
+		}
+
+		// Compute direct lighting
+		FPathTracingPayload Payload = UnpackPathTracingPayload(PackedPayload, Ray);
+		AdjustPayloadAfterUnpack(Payload, true);
+
+		// Generate main random number for direct+indirect sampling
+		float4 RandSample = RandomSequence_GenerateSample4D(RandSequence);
+
+#if LIGHT_LOOP_METHOD == 0
+		float LightPickingCdf[RAY_TRACING_LIGHT_COUNT_MAXIMUM];
+		float LightPickingCdfSum = 0;
+
+		for (uint Index = 0, Num = LightLoopCount.NumLights; Index < Num; ++Index)
+		{
+			uint LightIndex = GetLightId(Index, LightLoopCount);
+			float LightEstimate = EstimateLight(LightIndex, TranslatedWorldPos, WorldNormal, PrimitiveLightingChannelMask, false);
+			LightPickingCdfSum += LightEstimate;
+			LightPickingCdf[Index] = LightPickingCdfSum;
+		}
+
+		if (LightPickingCdfSum > 0)
+		{
+			// init worked
+			int LightId;
+			float LightPickPdf = 0;
+			
+			FMaterialSample MaterialSample = SampleMaterial(-Ray.Direction, Payload, RandSample.xyz);
+
+			SelectLight(RandSample.x * LightPickingCdfSum, LightLoopCount.NumLights, LightPickingCdf, LightId, LightPickPdf);
+			LightPickPdf /= LightPickingCdfSum;
+			float2 LightRandSample = RandSample.yz;
+			LightId = GetLightId(LightId, LightLoopCount);
+#else // LIGHT_LOOP_METHOD != 0
+		float2 LightRandSample = RandSample.xy;
+
+#if LIGHT_LOOP_METHOD == 2
+		FRISContext TraceSample = InitRISContext(RandSample.z);
+		FRayDesc ShadowRaySample;
+		float3 ShadowRayContrib = 0;
+		uint ShadowRayLightId;
+#endif
+		FMaterialSample MaterialSample = SampleMaterial(-Ray.Direction, Payload, RandSample.xyz);
+		for (uint Index = 0, Num = LightLoopCount.NumLights; Index < Num; ++Index)
+		{
+			uint LightId = GetLightId(Index, LightLoopCount);
+			const float LightPickPdf = 1;
+#endif // LIGHT_LOOP_METHOD != 0
+
+			// Generate a sample on the light source
+			FLightSample LightSample = SampleLight(LightId, LightRandSample, TranslatedWorldPos, WorldNormal);
+			if (LightSample.Pdf > 0)
+			{
+				LightSample.RadianceOverPdf /= LightPickPdf;
+				LightSample.Pdf *= LightPickPdf;
+
+#if 1
+				// Evaluate real material
+				FMaterialEval MaterialEval = EvalMaterial(-Ray.Direction, LightSample.Direction, Payload, 1.0);
+				float3 LightContrib = PathThroughput * LightSample.RadianceOverPdf * MaterialEval.Weight * MaterialEval.Pdf;
+#if USE_BSDF_SAMPLING
+				LightContrib *= MISWeightPower(LightSample.Pdf, MaterialEval.Pdf);
+#endif
+
+#else
+				// Hard-coded diffuse lobe for benchmarking
+#if SUBSTRATE_ENABLED
+				float3 MaterialEvalWeight = PackedPayload.UnpackDiffuseColor();
+#else
+				float3 MaterialEvalWeight = PackedPayload.UnpackBaseColor();
+#endif
+				float  MaterialEvalPdf = saturate(dot(WorldNormal, LightSample.Direction)) / PI;
+				float3 LightContrib = PathThroughput * LightSample.RadianceOverPdf * MaterialEvalWeight * MaterialEvalPdf;
+#endif
+				{
+					FRayDesc ShadowRay;
+					ShadowRay.Origin = TranslatedWorldPos;
+					ShadowRay.Direction = LightSample.Direction;
+					ShadowRay.TMin = 0;
+					ShadowRay.TMax = LightSample.Distance;
+
+					ApplyRayBias(ShadowRay.Origin, PackedPayload.HitT, PackedPayload.UnpackWorldGeoNormal());
+
+#if LIGHT_LOOP_METHOD == 2
+					float SelectionWeight = LightContrib.x + LightContrib.y + LightContrib.z;
+					if (TraceSample.Accept(SelectionWeight))
+					{
+						ShadowRaySample = ShadowRay;
+						ShadowRayContrib = LightContrib / SelectionWeight;
+						ShadowRayLightId = LightId;
+					}
+#else
+					// Trace shadow ray immediately
+					const bool bCastShadows = CastsShadow(LightId);
+					const uint MissShaderIndex = GetLightMissShaderIndex(LightId);
+					LocalRadiance += LightContrib * TraceShadowRay(ShadowRay, 1.0, 0.0, MissShaderIndex, bCastShadows);
+#endif
+				}
+			}
+
+#if USE_BSDF_SAMPLING
+			// Generate a sample by intersecting a BSDF ray with the light
+			{
+				FRayDesc ShadowRay;
+				ShadowRay.Origin = TranslatedWorldPos;
+				ShadowRay.Direction = MaterialSample.Direction;
+				ShadowRay.TMin = 0;
+				ShadowRay.TMax = RAY_DEFAULT_T_MAX;
+
+				ApplyRayBias(ShadowRay.Origin, PackedPayload.HitT, PackedPayload.UnpackWorldGeoNormal());
+
+				FLightHit LightResult = TraceLight(ShadowRay, LightId);
+				if (LightResult.IsHit())
+				{
+					ShadowRay.TMax = LightResult.HitT;
+					float3 LightContrib = PathThroughput * MaterialSample.Weight * LightResult.Radiance;
+					LightContrib *= MISWeightPower(MaterialSample.Pdf, LightResult.Pdf * LightPickPdf);
+#if LIGHT_LOOP_METHOD == 2
+					// Defer tracing shadow until later
+					float SelectionWeight = LightContrib.x + LightContrib.y + LightContrib.z;
+					if (TraceSample.Accept(SelectionWeight))
+					{
+						ShadowRaySample = ShadowRay;
+						ShadowRayContrib = LightContrib / SelectionWeight;
+						ShadowRayLightId = LightId;
+					}
+#else
+					const bool bCastShadows = CastsShadow(LightId);
+					const uint MissShaderIndex = GetLightMissShaderIndex(LightId);
+					LocalRadiance += LightContrib * TraceShadowRay(ShadowRay, 1.0, 0.0, MissShaderIndex, bCastShadows);
+#endif
+				}
+			}
+#endif // USE_BSDF_SAMPLING
+		}
+#if LIGHT_LOOP_METHOD == 2
+		// Light loop gave us a valid candidate sample, trace it now
+		if (TraceSample.HasSample())
+		{
+			const bool bCastShadows = CastsShadow(ShadowRayLightId);
+			const uint MissShaderIndex = GetLightMissShaderIndex(ShadowRayLightId);
+			ShadowRayContrib *= TraceSample.GetNormalization();
+			LocalRadiance += ShadowRayContrib * TraceShadowRay(ShadowRaySample, 1.0, 0.0, MissShaderIndex, bCastShadows);
+		}
+#endif
+
+		Result.Radiance.rgb += LocalRadiance;
+
+		if (UsePathGuiding && Bounce > 0)
+		{
+			// Algorithm 3: State Transition and Maximum Likelihood Estimation
+			// Writes to grid for the parent path vertex
+			const float3 Contrib = LocalRadiance;
+			const float MCf = dot(Contrib, float3(0.25, 0.5, 0.25));
+			if (MCf > 0 && RandomSequence_GenerateSample1D(RandSequence) * ScoresSum < MCf * N_MC)
+			{
+				const float3 AbsoluteWorldPosition = DFDemote(DFFastSubtract(Ray.Origin, ResolvedView.PreViewTranslation));
+				const float Distance = length(Ray.Origin);
+
+				// hit point is the "target" from the point of view of the previous hit
+				float3 TargetTranslatedWorldPos = TranslatedWorldPos;
+				float3 TargetAbsoluteWorldPosition = DFDemote(DFFastSubtract(TargetTranslatedWorldPos, ResolvedView.PreViewTranslation));
+				float3 CurrentStateDirection = StateGetDir(ChosenState, AbsoluteWorldPosition);
+				float Cosine = max(0.0, dot(normalize(TargetAbsoluteWorldPosition - AbsoluteWorldPosition), CurrentStateDirection));
+
+				ChosenState.N = min(ChosenState.N + 1, NMax);
+				const float Alpha = max(1.0 / ChosenState.N, AlphaMin);
+
+				float Weight = MCf;
+				ChosenState.WeightSum = lerp(ChosenState.WeightSum, Weight, Alpha);
+				ChosenState.WeightedTarget = lerp(ChosenState.WeightedTarget, Weight * TargetAbsoluteWorldPosition, Alpha);
+				ChosenState.WeightedCosine = lerp(ChosenState.WeightedCosine, Weight * Cosine, Alpha);
+
+				if (AdaptiveGridProb < 1)
+				{
+					// Save to Static Grid
+					float3 RandSample = RandomSequence_GenerateSample3D(RandSequence);
+					int3 Coord = int3(floor(AbsoluteWorldPosition / StaticGridWidth + RandSample));
+					int Side = 0;
+					uint2 HashIndex = GridHash(int4(Coord, -6 + Side)).xy; // two-hashes for collision detection
+					ChosenState.Hash2 = HashIndex.y;
+					MarkovChainStateGrid[HashIndex.x & (MarkovChainStateGridSize - 1)] = ChosenState;
+				}
+				if (AdaptiveGridProb > 0)
+				{
+					// Save to adaptive grid
+					float4 RandSample = RandomSequence_GenerateSample4D(RandSequence);
+					float TargetWidth = TanFactor * Distance;
+					float Level = round(LevelScale * log2(max(TargetWidth, MinWidth) / MinWidth)) + floor(-log2(1.0 - RandSample.w));
+					float CellWidth = MinWidth * exp2(Level / LevelScale);
+					int Side = GetDominantSide(PreviousN);
+					int LevelI = int(Level);
+					int3 Coord = int3(floor(AbsoluteWorldPosition / CellWidth + RandSample.xyz));
+					uint2 HashIndex = GridHash(int4(Coord, 6 * LevelI + Side)).xy; // two-hashes for collision detection
+					ChosenState.Hash2 = HashIndex.y;
+					MarkovChainStateGrid[HashIndex.x & (MarkovChainStateGridSize - 1)] = ChosenState;
+				}
+			}
+		}
+
+		// don't generate another ray on the last bounce or second to last bounce if emissive materials are disabled
+		if (bLastBounce || (Bounce + 1 == MaxBounces && EnableEmissive == 0))
+		{
+			break;
+		}
+
+		// sample new direction (potentially guided)
+
+		ScoresSum = 0; // reset
+		ChosenState = (FMarkovChainState)0;
+		float3 SampledDirection = 0;
+		float Scores[N_MC];
+		float4 VMFLobes[N_MC];
+		if (UsePathGuiding)
+		{
+			const float3 WorldSmoothNormal = PackedPayload.UnpackWorldSmoothNormal();
+			const float3 AbsoluteWorldPosition = DFDemote(DFFastSubtract(TranslatedWorldPos, ResolvedView.PreViewTranslation));
+			const float Distance = length(TranslatedWorldPos);
+
+			// Algorithm 1: Grid Resampling 
+			UNROLL
+			for (int Index = 0; Index < N_MC; Index++)
+			{
+				// TODO: Add splitting API so we can re-use the outer RandSequence instead
+				RandomSequence SplitRandSequence;
+				RandomSequence_Initialize(SplitRandSequence, LaunchIndex, SampleIndex * N_MC + Index, View.StateFrameIndex, SamplesPerPixel * N_MC);
+				float4 RandSample1 = RandomSequence_GenerateSample4D(SplitRandSequence);
+				float4 RandSample2 = RandomSequence_GenerateSample4D(SplitRandSequence);
+
+				bool bUseAdaptiveGrid = RandSample2.z < AdaptiveGridProb;
+
+				// static grid default
+				float CellWidth = StaticGridWidth;
+				int Side = 0;
+				int LevelI = -1;
+				if (bUseAdaptiveGrid)
+				{
+					// adaptive grid
+					float TargetWidth = TanFactor * Distance;
+					float Level = round(LevelScale * log2(max(TargetWidth, MinWidth) / MinWidth)) + floor(-log2(1.0 - RandSample1.w));
+					CellWidth = MinWidth * exp2(Level / LevelScale);
+					LevelI = int(Level);
+					Side = GetDominantSide(WorldSmoothNormal);
+				}
+				// stochastic rounding (pick one of the 8 possible neighbor grid cells randomly)
+				int3 Coord = int3(floor(AbsoluteWorldPosition / CellWidth + RandSample1.xyz));
+				
+				uint2 HashIndex = GridHash(int4(Coord, 6 * LevelI + Side)).xy; // two-hashes for collision detection
+
+				FMarkovChainState State = MarkovChainStateGrid[HashIndex.x & (MarkovChainStateGridSize - 1)];
+					
+				// collision detection
+				State.WeightSum *= float(State.Hash2 == HashIndex.y);
+				if (bUseAdaptiveGrid)
+				{
+					// adaptive grid
+				}
+				else
+				{
+					// static grid
+					float3 StateP = State.WeightSum > 0.0 ? State.WeightedTarget / State.WeightSum : State.WeightedTarget;
+					State.WeightSum *= float(dot(WorldSmoothNormal, normalize(StateP - AbsoluteWorldPosition)) > 0.0);
+				}
+
+				Scores[Index] = State.WeightSum;
+				ScoresSum += State.WeightSum;
+				VMFLobes[Index] = StateGetVMFLobe(State, AbsoluteWorldPosition);
+				if (ScoresSum * RandSample2.w < Scores[Index])
+				{
+					ChosenState = State;
+					SampledDirection = VMFSample(VMFLobes[Index], RandSample2.xy);
+				}
+			}
+		}
+
+		// Algorithm 2: Sampling and Ray Casting
+		// TODO: Adjust SurfaceProb based on roughness
+		ScatterType = PATHTRACER_SCATTER_DIFFUSE; // TODO: is there a good way to classify the scattering type for guided rays?
+		PathRoughness = 1.0;                      // TODO: return average roughness in EvalMaterial?
+		float3 Throughput = 0;
+		if (ScoresSum == 0 || RandSample.w < MaterialProb)
+		{
+			// BSDF Sampling
+			FMaterialSample MaterialSample = SampleMaterial(-Ray.Direction, Payload, RandSample.xyz);
+			Throughput = MaterialSample.Weight;
+			SampledDirection = MaterialSample.Direction;
+			ScatterType = MaterialSample.ScatterType;
+			PathRoughness = MaterialSample.Roughness;
+			ChosenState = (FMarkovChainState)0;
+		}
+		if (ScoresSum > 0)
+		{
+			// MIS between guided lobes and BSDF sampling
+			float GuidedPdf = 0;
+			UNROLL
+			for (int Index = 0; Index < N_MC; Index++)
+			{
+				if (Scores[Index] > 0)
+				{
+					GuidedPdf += Scores[Index] * VMFPdf(SampledDirection, VMFLobes[Index]) / ScoresSum;
+				}
+			}
+			FMaterialEval MaterialEval = EvalMaterial(-Ray.Direction, SampledDirection, Payload, 1.0);
+			Throughput = MaterialEval.Weight * MISWeightBalanced(MaterialProb * MaterialEval.Pdf, GuidedPdf - GuidedPdf * MaterialProb) / MaterialProb;
+		}
+
+		// Setup next ray
+		PathThroughput *= Throughput;
+		PreviousN = PackedPayload.UnpackWorldSmoothNormal();
+		Ray.Origin = TranslatedWorldPos;
+		Ray.Direction = SampledDirection;
+		Ray.TMin = 0;
+		Ray.TMax = RAY_DEFAULT_T_MAX;
+
+		ApplyRayBias(Ray.Origin, PackedPayload.HitT, PackedPayload.UnpackWorldGeoNormal());
+	}
+	Result.Radiance.rgb *= View.PreExposure;
+	return Result;
+}
+
+#else // PATH_TRACER_PRIMARY_RAYS
 
-PathTracingResult PathTracingDebugCore(FRayDesc Ray, uint2 LaunchIndex)
+FPathTracingResult PathTracingDebugCore(FRayDesc Ray, uint2 LaunchIndex)
 {
+
 	const uint RayFlags = RAY_FLAG_CULL_BACK_FACING_TRIANGLES;
 #if INCLUDE_TRANSLUCENT
 	const uint InstanceInclusionMask = PATHTRACER_MASK_CAMERA | PATHTRACER_MASK_CAMERA_TRANSLUCENT;
@@ -136,6 +645,7 @@ PathTracingResult PathTracingDebugCore(FRayDesc Ray, uint2 LaunchIndex)
 		MissShaderIndex,
 		Ray.GetNativeDesc(),
 		PackedPayload);
+
 	float4 Result = float4(0, 0, 0, 1);
 	float Z = 0;
 	if (PackedPayload.IsHit())
@@ -145,7 +655,6 @@ PathTracingResult PathTracingDebugCore(FRayDesc Ray, uint2 LaunchIndex)
 	}
 
 
-
 	if (DebugMode == PATH_TRACER_DEBUG_VIZ_VOLUME_LIGHT_COUNT)
 	{
 		float TMax = PackedPayload.IsHit() ? PackedPayload.HitT : Ray.TMax;
@@ -177,11 +686,6 @@ PathTracingResult PathTracingDebugCore(FRayDesc Ray, uint2 LaunchIndex)
 		return MakeResult(Result.rgb, Result.a, Z);
 	}
 
-
-
-	float LightPickingCdf[RAY_TRACING_LIGHT_COUNT_MAXIMUM];
-	float LightPickingCdfSum = 0;
-	float3 PathThroughput = 1;
 	if (PackedPayload.IsHit())
 	{
 		const float3 TranslatedWorldPos = Ray.Origin + PackedPayload.HitT * Ray.Direction;
@@ -221,14 +725,15 @@ PathTracingResult PathTracingDebugCore(FRayDesc Ray, uint2 LaunchIndex)
 			case PATH_TRACER_DEBUG_VIZ_IOR:							return MakeResult(PackedPayload.UnpackIor(), 0.0, Z);
 			case PATH_TRACER_DEBUG_VIZ_SHADING_MODEL:				return MakeResult(GetShadingModelColor(PackedPayload.GetShadingModelID()), 0.0, Z);
 			case PATH_TRACER_DEBUG_VIZ_LIGHTING_CHANNEL_MASK:		return MakeResult(GetLightingChannelMaskDebugColor(PackedPayload.GetPrimitiveLightingChannelMask()), 0.0, Z);
-			case PATH_TRACER_DEBUG_VIZ_WORLD_POSITION: {
+			case PATH_TRACER_DEBUG_VIZ_WORLD_POSITION:
+			{
 				float3 AbsoluteWorldPosition = DFDemote(DFFastSubtract(TranslatedWorldPos, ResolvedView.PreViewTranslation));
 				const float CellWidth = 10.0f; // 10cm
 				int3 Coord = int3(round(AbsoluteWorldPosition / CellWidth));
-				uint HashIndex = Rand3DPCG32(Coord).x;
+				uint HashIndex = GridHash(int4(Coord, 0)).x;
 				float3 Color = HUEtoRGB((HashIndex & 0xFFFFFF) * 5.96046447754e-08);
 				Color *= abs(dot(WorldNormal, Ray.Direction));
-				return MakeResult(Color * View.PreExposure, 0.0, Z);
+				return MakeResult(Color, 0.0, Z);
 			}
 			case PATH_TRACER_DEBUG_VIZ_PRIMARY_RAYS: break; // fallthrough to implementation below
 			case PATH_TRACER_DEBUG_VIZ_WORLD_TANGENT: 				return MakeResult(PackedPayload.UnpackWorldTangent() * 0.5 + 0.5, 0.0, Z);
@@ -253,204 +758,11 @@ PathTracingResult PathTracingDebugCore(FRayDesc Ray, uint2 LaunchIndex)
 				return MakeResult((PackedPayload.IsFrontFace() ? float3(0, 1, 0) : float3(1, 0, 0)) * abs(dot(WorldNormal, Ray.Direction)), 0.0, Z);
 			}
 		}
-
-		FPathTracingPayload Payload = UnpackPathTracingPayload(PackedPayload, Ray);
-		AdjustPayloadAfterUnpack(Payload, true);
-#if LIGHT_LOOP_METHOD == 0
-		for (uint Index = 0, Num = LightLoopCount.NumLights; Index < Num; ++Index)
-		{
-			uint LightIndex = GetLightId(Index, LightLoopCount);
-			float LightEstimate = EstimateLight(LightIndex, TranslatedWorldPos, WorldNormal, PrimitiveLightingChannelMask, false);
-			LightPickingCdfSum += LightEstimate;
-			LightPickingCdf[Index] = LightPickingCdfSum;
-		}
-
-		if (LightPickingCdfSum > 0)
-		{
-			// init worked
-			int LightId;
-			float LightPickPdf = 0;
-			
-			for (int Sample = 0; Sample < NumLightSamples; Sample++)
-			{
-				RandomSequence RandSequence;
-				RandomSequence_Initialize(RandSequence, LaunchIndex, Sample, View.StateFrameIndex, NumLightSamples);
-				float4 RandSample = RandomSequence_GenerateSample4D(RandSequence);
-
-				FMaterialSample MaterialSample = SampleMaterial(-Ray.Direction, Payload, RandSample.xyz);
-
-				SelectLight(RandSample.x * LightPickingCdfSum, LightLoopCount.NumLights, LightPickingCdf, LightId, LightPickPdf);
-				LightPickPdf /= LightPickingCdfSum;
-				float2 LightRandSample = RandSample.yz;
-				LightId = GetLightId(LightId, LightLoopCount);
-#else
-		for (int Sample = 0; Sample < NumLightSamples; Sample++)
-		{
-			RandomSequence RandSequence;
-			RandomSequence_Initialize(RandSequence, LaunchIndex, Sample, View.StateFrameIndex, NumLightSamples);
-			float3 RandSample = RandomSequence_GenerateSample3D(RandSequence);			
-			float2 LightRandSample = RandSample.xy;
-
-#if LIGHT_LOOP_METHOD == 2
-			FRISContext TraceSample = InitRISContext(RandSample.z);
-			FRayDesc ShadowRaySample;
-			float3 ShadowRayContrib = 0;
-			uint ShadowRayLightId;
-#endif
-
-			FMaterialSample MaterialSample = SampleMaterial(-Ray.Direction, Payload, RandSample.xyz);
-			for (uint Index = 0, Num = LightLoopCount.NumLights; Index < Num; ++Index)
-			{
-				uint LightId = GetLightId(Index, LightLoopCount);
-				const float LightPickPdf = 1;
-#endif
-
-				FLightSample LightSample = SampleLight(LightId, LightRandSample, TranslatedWorldPos, WorldNormal);
-
-				if (LightSample.Pdf > 0)
-				{
-					LightSample.RadianceOverPdf /= LightPickPdf;
-					LightSample.Pdf *= LightPickPdf;
-
-#if 1
-					FMaterialEval MaterialEval = EvalMaterial(-Ray.Direction, LightSample.Direction, Payload, 1.0);
-					float3 LightContrib = PathThroughput * LightSample.RadianceOverPdf * MaterialEval.Weight * MaterialEval.Pdf;
-#if USE_BSDF_SAMPLING
-					LightContrib *= MISWeightPower(LightSample.Pdf, MaterialEval.Pdf);
-#endif
-
-#else
-#if SUBSTRATE_ENABLED
-					float3 MaterialEvalWeight = PackedPayload.UnpackDiffuseColor();
-#else
-					float3 MaterialEvalWeight = PackedPayload.UnpackBaseColor();
-#endif
-					float  MaterialEvalPdf = saturate(dot(WorldNormal, LightSample.Direction)) / PI;
-					float3 LightContrib = PathThroughput * LightSample.RadianceOverPdf * MaterialEvalWeight * MaterialEvalPdf;
-#endif
-					{
-						FRayDesc ShadowRay;
-						ShadowRay.Origin = TranslatedWorldPos;
-						ShadowRay.Direction = LightSample.Direction;
-						ShadowRay.TMin = 0;
-						ShadowRay.TMax = LightSample.Distance;
-
-						ApplyRayBias(ShadowRay.Origin, PackedPayload.HitT, PackedPayload.UnpackWorldGeoNormal());
-
-#if LIGHT_LOOP_METHOD == 2
-						float SelectionWeight = LightContrib.x + LightContrib.y + LightContrib.z;
-						if (TraceSample.Accept(SelectionWeight))
-						{
-							ShadowRaySample = ShadowRay;
-							ShadowRayContrib = LightContrib / SelectionWeight;
-							ShadowRayLightId = LightId;
-						}
-#else
-						const bool bCastShadows = CastsShadow(LightId);
-						const uint MissShaderIndex = GetLightMissShaderIndex(LightId);
-						Result.xyz += LightContrib * TraceShadowRay(ShadowRay, 1.0, 0.0, MissShaderIndex, bCastShadows);
-#endif
-					}
-				}
-#if USE_BSDF_SAMPLING
-				{
-					FRayDesc ShadowRay;
-					ShadowRay.Origin = TranslatedWorldPos;
-					ShadowRay.Direction = MaterialSample.Direction;
-					ShadowRay.TMin = 0;
-					ShadowRay.TMax = RAY_DEFAULT_T_MAX;
-
-					ApplyRayBias(ShadowRay.Origin, PackedPayload.HitT, PackedPayload.UnpackWorldGeoNormal());
-
-					FLightHit LightResult = TraceLight(ShadowRay, LightId);
-					if (LightResult.IsHit())
-					{
-						ShadowRay.TMax = LightResult.HitT;
-						float3 LightContrib = PathThroughput * MaterialSample.Weight * LightResult.Radiance;
-						LightContrib *= MISWeightPower(MaterialSample.Pdf, LightResult.Pdf * LightPickPdf);
-#if LIGHT_LOOP_METHOD == 2
-						float SelectionWeight = LightContrib.x + LightContrib.y + LightContrib.z;
-						if (TraceSample.Accept(SelectionWeight))
-						{
-							ShadowRaySample = ShadowRay;
-							ShadowRayContrib = LightContrib / SelectionWeight;
-							ShadowRayLightId = LightId;
-						}
-#else
-						const bool bCastShadows = CastsShadow(LightId);
-						const uint MissShaderIndex = GetLightMissShaderIndex(LightId);
-						Result.xyz += LightContrib * TraceShadowRay(ShadowRay, 1.0, 0.0, MissShaderIndex, bCastShadows);
-#endif
-					}
-				}
-#endif // USE_BSDF_SAMPLING
-			}
-#if LIGHT_LOOP_METHOD == 2
-			if (TraceSample.HasSample())
-			{
-				const bool bCastShadows = CastsShadow(ShadowRayLightId);
-				const uint MissShaderIndex = GetLightMissShaderIndex(ShadowRayLightId);
-				ShadowRayContrib *= TraceSample.GetNormalization();
-				Result.xyz += ShadowRayContrib * TraceShadowRay(ShadowRaySample, 1.0, 0.0, MissShaderIndex, bCastShadows);
-			}
-#endif
-		}
-		Result.rgb = Result.rgb * (NumLightSamples > 1 ? 1.0 / NumLightSamples : 1.0) + PackedPayload.UnpackRadiance();
-		float3 IndirectRadiance = 0;
-		for (int Sample = 0; Sample < NumIndirectSamples; Sample++)
-		{
-			RandomSequence RandSequence;
-			RandomSequence_Initialize(RandSequence, LaunchIndex, Sample, View.StateFrameIndex, NumIndirectSamples);
-			float3 RandSample = RandomSequence_GenerateSample3D(RandSequence);
-			FMaterialSample MaterialSample = SampleMaterial(-Ray.Direction, Payload, RandSample.xyz);
-
-			const uint RayFlags = 0;
-			const uint InstanceInclusionMask = PATHTRACER_MASK_INDIRECT;
-			const uint RayContributionToHitGroupIndex = RAY_TRACING_SHADER_SLOT_MATERIAL;
-			const uint MultiplierForGeometryContributionToShaderIndex = RAY_TRACING_NUM_SHADER_SLOTS;
-
-			FRayDesc Ray;
-			Ray.Origin = TranslatedWorldPos;
-			Ray.Direction = MaterialSample.Direction;
-			Ray.TMin = 0;
-			Ray.TMax = RAY_DEFAULT_T_MAX;
-
-			ApplyRayBias(Ray.Origin, PackedPayload.HitT, PackedPayload.UnpackWorldGeoNormal());
-
-			FPackedPathTracingPayload PackedPayloadIndirect = InitPathTracingPayload(MaterialSample.ScatterType, MaterialSample.Roughness);
-#if PATH_TRACER_USE_SER 
-			NvHitObject Hit;
-			NvTraceRayHitObject(
-				TLAS,
-				RayFlags,
-				InstanceInclusionMask,
-				RayContributionToHitGroupIndex,
-				MultiplierForGeometryContributionToShaderIndex,
-				MissShaderIndex,
-				Ray.GetNativeDesc(),
-				PackedPayloadIndirect,
-				Hit);
-			NvReorderThread(Hit);
-			NvInvokeHitObject(TLAS, Hit, PackedPayloadIndirect);
-#else
-			TraceRay(
-				TLAS,
-				RayFlags,
-				InstanceInclusionMask,
-				RayContributionToHitGroupIndex,
-				MultiplierForGeometryContributionToShaderIndex,
-				MissShaderIndex,
-				Ray.GetNativeDesc(),
-				PackedPayloadIndirect);
-#endif
-			IndirectRadiance += MaterialSample.Weight * PackedPayloadIndirect.UnpackRadiance();
-		}
-		Result.rgb += IndirectRadiance * (NumIndirectSamples > 1 ? 1.0 / NumIndirectSamples : 1.0);
-		Result.a = 0;
 	}
-	return MakeResult(Result.rgb * View.PreExposure, Result.a, Z);
+	return MakeResult(0, 1, Z);
 }
 
+#endif // !PATH_TRACER_PRIMARY_RAYS
 
 RAY_TRACING_ENTRY_RAYGEN(PathTracingDebugRG)
 {
@@ -464,9 +776,23 @@ RAY_TRACING_ENTRY_RAYGEN(PathTracingDebugRG)
 
 	const float2 AAJitter = 0.5; // trace through pixel center
 	const float2 ViewportUV = (LaunchIndex + AAJitter) * View.BufferSizeAndInvSize.zw;
+	const FRayDesc CameraRay = CreatePrimaryRay(ViewportUV);
 
+#if PATH_TRACER_PRIMARY_RAYS
+	FPathTracingResult Result = (FPathTracingResult) 0;
 	
-	PathTracingResult Result = PathTracingDebugCore(CreatePrimaryRay(ViewportUV), LaunchIndex);
+	// TODO: launch separate kernels and accumulate instead
+	for (int SampleIndex = 0; SampleIndex < SamplesPerPixel; SampleIndex++)
+	{
+		FPathTracingResult SampleResult = PathTracingDebugCore(CameraRay, LaunchIndex, SampleIndex);
+
+		Result.Radiance += SampleResult.Radiance;
+		Result.Z = max(SampleResult.Z, Result.Z); // inverted Z so larger numbers are closer
+	}
+	Result.Radiance *= 1.0 / SamplesPerPixel;
+#else
+	FPathTracingResult Result = PathTracingDebugCore(CameraRay, LaunchIndex);
+#endif
 
 	RWSceneColor[DispatchIdx] = Result.Radiance;
 	RWSceneDepth[DispatchIdx] = Result.Z;
diff --git a/Engine/Shaders/Shared/RayTracingTypes.h b/Engine/Shaders/Shared/RayTracingTypes.h
--- a/Engine/Shaders/Shared/RayTracingTypes.h
+++ b/Engine/Shaders/Shared/RayTracingTypes.h
@@ -88,6 +88,17 @@ struct FRayTracingDecal
 };
 HLSL_STATIC_ASSERT(sizeof(FRayTracingDecal) == 96, "Ray tracing decal structure should be aligned to 32 bytes for optimal access on the GPU");
 
+struct FMarkovChainState
+{
+    float3 		WeightedTarget;
+    float 		WeightSum;
+	float 		WeightedCosine;
+    uint 		N;
+    uint        Hash2; // second hash for resolving collisions
+	uint        Padding0; // for alignment
+};
+HLSL_STATIC_ASSERT(sizeof(FMarkovChainState) == 32, "");
+
 #ifdef __cplusplus
 } // namespace UE::HLSL
 
@@ -95,5 +106,6 @@ using FRTLightingData = UE::HLSL::FRTLightingData;
 using FPathTracingLight = UE::HLSL::FPathTracingLight;
 using FPathTracingPackedPathState = UE::HLSL::FPathTracingPackedPathState;
 using FRayTracingDecal = UE::HLSL::FRayTracingDecal;
+using FMarkovChainState = UE::HLSL::FMarkovChainState;
 
 #endif
diff --git a/Engine/Source/Runtime/Renderer/Private/PathTracing.cpp b/Engine/Source/Runtime/Renderer/Private/PathTracing.cpp
--- a/Engine/Source/Runtime/Renderer/Private/PathTracing.cpp
+++ b/Engine/Source/Runtime/Renderer/Private/PathTracing.cpp
@@ -109,20 +109,43 @@ FAutoConsoleVariableRef CVarPathTracingVisualize(
 	ECVF_RenderThreadSafe
 );
 
-TAutoConsoleVariable<int32> CVarPathTracingVisualizeLightSamples(
-	TEXT("r.PathTracing.Visualize.LightSamples"),
+TAutoConsoleVariable<int32> CVarPathTracingVisualizeMaxBounces(
+	TEXT("r.PathTracing.Visualize.MaxBounces"),
 	1,
-	TEXT("Number of light samples for the path tracing light sampling visualize mode"),
+	TEXT("Number of light bounces for primary rays visualization mode"),
 	ECVF_RenderThreadSafe
 );
 
-TAutoConsoleVariable<int32> CVarPathTracingVisualizeIndirectSamples(
-	TEXT("r.PathTracing.Visualize.IndirectSamples"),
-	0,
-	TEXT("Number of indirect samples for the path tracing visualize mode"),
+TAutoConsoleVariable<bool> CVarPathTracingVisualizeEnableEmissive(
+	TEXT("r.PathTracing.Visualize.EnableEmissive"),
+	false,
+	TEXT("Indicates if emissive materials should contribute to scene lighting for primary rays visualization mode"),
+	ECVF_RenderThreadSafe
+);
+
+
+TAutoConsoleVariable<int32> CVarPathTracingVisualizeSamplesPerPixel(
+	TEXT("r.PathTracing.Visualize.SamplesPerPixel"),
+	1,
+	TEXT("Number of samples for primary rays visualization mode"),
+	ECVF_RenderThreadSafe
+);
+
+TAutoConsoleVariable<bool> CVarPathTracingVisualizeIndirectGuiding(
+	TEXT("r.PathTracing.Visualize.IndirectGuiding"),
+	false,
+	TEXT("Should indirect rays use path guiding to improve quality?"),
+	ECVF_RenderThreadSafe
+);
+
+TAutoConsoleVariable<int32> CVarPathTracingVisualizeIndirectGuidingHashSize(
+	TEXT("r.PathTracing.Visualize.IndirectGuiding.HashSize"),
+	1024 * 1024,
+	TEXT("Size of the hash grid to use for storing path guiding information (automatically rounded up to nearest power of 2)"),
 	ECVF_RenderThreadSafe
 );
 
+
 TAutoConsoleVariable<int32> CVarPathTracingCompactionDepth(
 	TEXT("r.PathTracing.CompactionDepth"),
 	-1,
@@ -773,6 +796,7 @@ struct FPathTracingState
 	TRefCountPtr<IPooledRenderTarget> DepthRT;
 	TRefCountPtr<FRDGPooledBuffer> VarianceBuffer;
 	TRefCountPtr<IPooledRenderTarget> CloudAccelerationMap;
+	TRefCountPtr<FRDGPooledBuffer> MarkovChainStateGrid;
 
 	// Cache to improve the stability when frame denoising (SPP=r.pathtracing.SamplesPerPixel) is used in animation rendering
 	TRefCountPtr<IPooledRenderTarget> LastDenoisedRadianceRT;
@@ -1390,15 +1414,24 @@ class FPathTracingDebugRG : public FGlobalShader
 	SHADER_USE_ROOT_PARAMETER_STRUCT(FPathTracingDebugRG, FGlobalShader)
 
 	class FUseSER : SHADER_PERMUTATION_BOOL("PATH_TRACER_USE_SER");
-	using FPermutationDomain = TShaderPermutationDomain<FUseSER>;
+	class FUsePrimaryRays : SHADER_PERMUTATION_BOOL("PATH_TRACER_PRIMARY_RAYS");
+	using FPermutationDomain = TShaderPermutationDomain<FUseSER, FUsePrimaryRays>;
 
 	static bool ShouldCompilePermutation(const FGlobalShaderPermutationParameters& Parameters)
 	{
-		// Only compile SER extensions on platforms that could support it
 		FPermutationDomain PermutationVector(Parameters.PermutationId);
-		if (PermutationVector.Get<FUseSER>() && !FDataDrivenShaderPlatformInfo::GetSupportsShaderExecutionReordering(Parameters.Platform))
+		if (PermutationVector.Get<FUseSER>())
 		{
-			return false;
+			// SER only makes sense in the PrimaryRays mode
+			if (!PermutationVector.Get<FUsePrimaryRays>())
+			{
+				return false;
+			}
+			// Only compile SER extensions on platforms that could support it
+			if (!FDataDrivenShaderPlatformInfo::GetSupportsShaderExecutionReordering(Parameters.Platform))
+			{
+				return false;
+			}
 		}
 
 		return PathTracing::ShouldCompilePathTracingShadersForProject(Parameters.Platform);
@@ -1425,10 +1458,14 @@ class FPathTracingDebugRG : public FGlobalShader
 		SHADER_PARAMETER_RDG_TEXTURE_UAV(RWTexture2D<float4>, RWSceneColor)
 		SHADER_PARAMETER_RDG_TEXTURE_UAV(RWTexture2D<float>,  RWSceneDepth)
 		SHADER_PARAMETER_RDG_BUFFER_SRV(RaytracingAccelerationStructure, TLAS)
+		SHADER_PARAMETER_RDG_BUFFER_UAV(RWStructuredBuffer<FMarkovChainState>, MarkovChainStateGrid)
 		SHADER_PARAMETER_STRUCT_REF(FViewUniformShaderParameters, ViewUniformBuffer)
 		SHADER_PARAMETER(int32, DebugMode)
-		SHADER_PARAMETER(int32, NumLightSamples)
-		SHADER_PARAMETER(int32, NumIndirectSamples)
+		SHADER_PARAMETER(int32, SamplesPerPixel)
+		SHADER_PARAMETER(int32, MaxBounces)
+		SHADER_PARAMETER(int32, EnableEmissive)
+		SHADER_PARAMETER(uint32, MarkovChainStateGridSize)
+		SHADER_PARAMETER(uint32, UsePathGuiding)
 
 
 		// scene lights
@@ -3239,12 +3276,12 @@ void FDeferredShadingSceneRenderer::PreparePathTracing(const FViewInfo& View, co
 	if (ViewFamily.EngineShowFlags.PathTracing
 		&& PathTracing::ShouldCompilePathTracingShadersForProject(ViewFamily.GetShaderPlatform()))
 	{
-		if (GetPathTracingVisualizationMode() >= 0)
+		if (int DebugMode = GetPathTracingVisualizationMode(); DebugMode >= 0)
 		{
 			const bool bUseSER = GRHIGlobals.SupportsShaderExecutionReordering && CVarPathTracingShaderExecutionReordering.GetValueOnRenderThread();
-
 			FPathTracingDebugRG::FPermutationDomain PermutationVector;
 			PermutationVector.Set<FPathTracingDebugRG::FUseSER>(bUseSER);
+			PermutationVector.Set<FPathTracingDebugRG::FUsePrimaryRays>(DebugMode == PATH_TRACER_DEBUG_VIZ_PRIMARY_RAYS);
 			auto RayGenShader = GetGlobalShaderMap(ViewFamily.GetShaderPlatform())->GetShader<FPathTracingDebugRG>(PermutationVector);
 			OutRayGenShaders.Add(RayGenShader.GetRayTracingShader());
 		}
@@ -3443,16 +3480,34 @@ void FDeferredShadingSceneRenderer::RenderPathTracing(
 			TexCreate_ShaderResource | TexCreate_UAV);
 		FRDGTextureRef DepthTexture    = GraphBuilder.CreateTexture(DepthDesc, TEXT("PathTracer.Depth"));
 
+		FPathTracingState* PathTracingState = GetPathTracingStateFromView(View);
+		FRDGBuffer* MarkovChainStateGrid = nullptr;
+		uint32 MarkovChainStateGridHashSize = FMath::RoundUpToPowerOfTwo(CVarPathTracingVisualizeIndirectGuidingHashSize.GetValueOnRenderThread());
+		if (PathTracingState->MarkovChainStateGrid && PathTracingState->MarkovChainStateGrid->Desc.NumElements == MarkovChainStateGridHashSize)
+		{
+			// ok to re-use previous grid
+			MarkovChainStateGrid = GraphBuilder.RegisterExternalBuffer(PathTracingState->MarkovChainStateGrid, TEXT("PathTracer.MarkovChainStateGrid"));
+		}
+		else
+		{
+			// need to make a new grid and initialize it to all zeros
+			MarkovChainStateGrid = GraphBuilder.CreateBuffer(FRDGBufferDesc::CreateStructuredDesc(sizeof(FMarkovChainState), MarkovChainStateGridHashSize), TEXT("PathTracer.MarkovChainStateGrid"));
+			AddClearUAVPass(GraphBuilder, GraphBuilder.CreateUAV(MarkovChainStateGrid), 0);
+		}
 
 		FPathTracingDebugRG::FParameters* PassParameters = GraphBuilder.AllocParameters<FPathTracingDebugRG::FParameters>();
 		PassParameters->RWSceneColor = GraphBuilder.CreateUAV(SceneColorOutputTexture);
 		PassParameters->RWSceneDepth = GraphBuilder.CreateUAV(DepthTexture);
 		PassParameters->TLAS = Scene->RayTracingScene.GetLayerView(ERayTracingSceneLayer::Base);
+		PassParameters->UsePathGuiding = CVarPathTracingVisualizeIndirectGuiding.GetValueOnRenderThread();
+		PassParameters->MarkovChainStateGrid = GraphBuilder.CreateUAV(MarkovChainStateGrid);
+		PassParameters->MarkovChainStateGridSize = MarkovChainStateGridHashSize;
 		PassParameters->ViewUniformBuffer = View.ViewUniformBuffer;
 		PassParameters->DebugMode = DebugMode;
 
-		PassParameters->NumLightSamples = CVarPathTracingVisualizeLightSamples.GetValueOnRenderThread();
-		PassParameters->NumIndirectSamples = CVarPathTracingVisualizeIndirectSamples.GetValueOnRenderThread();
+		PassParameters->SamplesPerPixel = CVarPathTracingVisualizeSamplesPerPixel.GetValueOnRenderThread();
+		PassParameters->MaxBounces = CVarPathTracingVisualizeMaxBounces.GetValueOnRenderThread();
+		PassParameters->EnableEmissive = CVarPathTracingVisualizeEnableEmissive.GetValueOnRenderThread();
 		SetLightParameters(
 			GraphBuilder,
 			Scene,
@@ -3471,9 +3526,10 @@ void FDeferredShadingSceneRenderer::RenderPathTracing(
 		PassParameters->NaniteRayTracing = Nanite::GRayTracingManager.GetUniformBuffer();
 
 		const bool bUseSER = GRHIGlobals.SupportsShaderExecutionReordering && CVarPathTracingShaderExecutionReordering.GetValueOnRenderThread();
-		FPathTracingDebugRG::FPermutationDomain Permutation;
-		Permutation.Set<FPathTracingDebugRG::FUseSER>(bUseSER);
-		TShaderMapRef<FPathTracingDebugRG> RayGenShader(View.ShaderMap, Permutation);
+		FPathTracingDebugRG::FPermutationDomain PermutationVector;
+		PermutationVector.Set<FPathTracingDebugRG::FUseSER>(bUseSER);
+		PermutationVector.Set<FPathTracingDebugRG::FUsePrimaryRays>(DebugMode == PATH_TRACER_DEBUG_VIZ_PRIMARY_RAYS);
+		TShaderMapRef<FPathTracingDebugRG> RayGenShader(View.ShaderMap, PermutationVector);
 		ClearUnusedGraphResources(RayGenShader, PassParameters);
 		GraphBuilder.AddPass(
 			RDG_EVENT_NAME("Path Tracing Visualize"),
@@ -3518,6 +3574,7 @@ void FDeferredShadingSceneRenderer::RenderPathTracing(
 				DisplayParameters);
 		}
 
+		GraphBuilder.QueueBufferExtraction(MarkovChainStateGrid, &PathTracingState->MarkovChainStateGrid);
 		return;
 	}
 
