Lumen - fix leaking of foreground GI into baground translucent meshes that are out of the translucenct lighting volume.

diff --git a/Engine/Shaders/Private/BasePassPixelShader.usf b/Engine/Shaders/Private/BasePassPixelShader.usf
--- a/Engine/Shaders/Private/BasePassPixelShader.usf
+++ b/Engine/Shaders/Private/BasePassPixelShader.usf
@@ -362,8 +362,11 @@ void GetSkyLighting(FMaterialPixelParameters MaterialParameters, VTPageTableResu
 #if (MATERIALBLENDING_TRANSLUCENT || MATERIALBLENDING_ADDITIVE || MATERIALBLENDING_MODULATE) && PROJECT_SUPPORTS_LUMEN && !FORWARD_SHADING
 	if (IsLumenTranslucencyGIEnabled())
 	{
-		// Lumen Dynamic GI + shadowed Skylight
-		FTwoBandSHVectorRGB TranslucencyGISH = GetTranslucencyGIVolumeLighting(MaterialParameters.AbsoluteWorldPosition, ResolvedView.WorldToClip, true);
+		const float3 VolumeUV = ComputeTranslucencyGIVolumeUV(MaterialParameters.AbsoluteWorldPosition, ResolvedView.WorldToClip);
+		if(VolumeUV.z < 1.0)
+		{
+			// Lumen Dynamic GI + shadowed Skylight
+			FTwoBandSHVectorRGB TranslucencyGISH = GetTranslucencyGIVolumeLighting(MaterialParameters.AbsoluteWorldPosition, ResolvedView.WorldToClip, true);
 
 		#if TRANSLUCENCY_LIGHTING_VOLUMETRIC_NONDIRECTIONAL || TRANSLUCENCY_LIGHTING_VOLUMETRIC_PERVERTEX_NONDIRECTIONAL
 
@@ -388,6 +391,17 @@ void GetSkyLighting(FMaterialPixelParameters MaterialParameters, VTPageTableResu
 			}
 			#endif
 		#endif
+		}
+		else
+		{
+			OutDiffuseLighting += GetEffectiveSkySHDiffuse(WorldNormal) * ResolvedView.SkyLightColor.rgb;
+			#if SHADINGMODEL_REQUIRES_BACKFACE_LIGHTING
+			if (bEvaluateBackface)
+			{
+				OutDiffuseLighting += GetEffectiveSkySHDiffuse(-WorldNormal) * ResolvedView.SkyLightColor.rgb;
+			}
+			#endif
+		}
 	}
 	else
 #endif
