Fixed mismatching vertex in/out when rendering with CompositeUIPixelShader

diff --git a/Engine/Shaders/Private/CompositeUIPixelShader.usf b/Engine/Shaders/Private/CompositeUIPixelShader.usf
--- a/Engine/Shaders/Private/CompositeUIPixelShader.usf
+++ b/Engine/Shaders/Private/CompositeUIPixelShader.usf
@@ -121,8 +121,14 @@ float3 ComputeHDRSceneColor(float3 SceneColor)
 	return SceneColor;
 }
 
+struct FCompositeVertexInput
+{
+	noperspective float4 UV : TEXCOORD0;
+	float4 Position : SV_POSITION;
+};
+
 void Main(
-	FScreenVertexOutput Input,
+	FCompositeVertexInput Input,
 	out float4 OutColor : SV_Target0
 	)
 {
@@ -138,7 +144,7 @@ void Main(
 #endif
 #endif
 
-	float4 UIColor = ComputeHDRUIColor(Input.UV);
+	float4 UIColor = ComputeHDRUIColor(Input.UV.xy);
 #if PLATFORM_SUPPORTS_RENDERTARGET_WRITE_MASK
 	// If the UI texture has a mask of 0, this means the UI color needs to be forced to 0 as it may contain uninitialized data
 	if (CompositeUIMask == 0)
@@ -147,7 +153,7 @@ void Main(
 	}
 #endif
 
-	float3 SceneColor = Texture2DSample(SceneTexture, SceneSampler, Input.UV).xyz;
+	float3 SceneColor = Texture2DSample(SceneTexture, SceneSampler, Input.UV.xy).xyz;
 	SceneColor = ComputeHDRSceneColor(SceneColor);
 	OutColor.xyz = ComposeUIAndScene(SceneColor.xyz, UIColor, UILevel);
 
@@ -248,4 +254,4 @@ void CompositeUIForBlur(FScreenVertexOutput Input, out float4 OutColor : SV_Targ
 	OutColor.a = 1;
 }
 
-#endif
\ No newline at end of file
+#endif
