Add option to disable instance detail in r.hairstrands.viewmode 1.

diff --git a/Engine/Shaders/Private/HairStrands/HairStrandsDebug.usf b/Engine/Shaders/Private/HairStrands/HairStrandsDebug.usf
--- a/Engine/Shaders/Private/HairStrands/HairStrandsDebug.usf
+++ b/Engine/Shaders/Private/HairStrands/HairStrandsDebug.usf
@@ -1618,6 +1618,7 @@ void MainCS(uint GroupIndex : SV_GroupIndex, uint3 DispatchThreadId : SV_Dispatc
 	}
 
 	FShaderPrintContext Context = InitShaderPrintContext(true, uint2(50, 100));
+	bool bShowInstanceDetails = false;
 	bool bShowAdvancedStrandsData = false;
 
 	// Draw totals
@@ -1663,123 +1664,127 @@ void MainCS(uint GroupIndex : SV_GroupIndex, uint3 DispatchThreadId : SV_Dispatc
 		Print(Context, InstanceCount_CardsOrMeshesShadowView, FontYellow, 4);
 		Newline(Context);
 
+		bShowInstanceDetails     = AddCheckbox(Context, TEXT("Details"), true, FontSilver);
 		bShowAdvancedStrandsData = AddCheckbox(Context, TEXT("Advanced"), false, FontSilver);
 		Newline(Context);
 	}
 	Newline(Context);
 
-	// Draw instance details
-	Print(Context, TEXT("   Group  LOD     Geometry Binding Cache Sim RBF Clip Screen CovSc. Strands            Points             Bind.Tri/Sect Vis. Shad. SwLOD RT  Name"));
-	Newline(Context);
-
-	LOOP
-	for (uint InstanceIndex = 0; InstanceIndex < InstanceCount; ++InstanceIndex)
+	if (bShowInstanceDetails)
 	{
-		const FInstanceInfo Info = ReadInstanceInfo(InstanceIndex);
+		// Draw instance details
+		Print(Context, TEXT("   Group  LOD     Geometry Binding Cache Sim RBF Clip Screen CovSc. Strands            Points             Bind.Tri/Sect Vis. Shad. SwLOD RT  Name"));
+		Newline(Context);
 
-		// Index
-		Print(Context,InstanceIndex, FontSilver, 2);
-		Print(Context, TEXT(" "));
+		LOOP
+		for (uint InstanceIndex = 0; InstanceIndex < InstanceCount; ++InstanceIndex)
+		{
+			const FInstanceInfo Info = ReadInstanceInfo(InstanceIndex);
 
-		// Group
-		Print(Context,Info.GroupIndex, FontYellow, 1);
-		Print(Context, TEXT("/"));
-		Print(Context,Info.GroupCount, FontYellow, 1);
-		Print(Context, TEXT("    "));
+			// Index
+			Print(Context,InstanceIndex, FontSilver, 2);
+			Print(Context, TEXT(" "));
 
-		// LOD
-		if (Info.bAutoLOD)
-		{
-			Print(Context, TEXT("Auto    "), FontYellow);
-		}
-		else
-		{
-			Print(Context,Info.LOD, FontYellow, 4, 2);
+			// Group
+			Print(Context,Info.GroupIndex, FontYellow, 1);
+			Print(Context, TEXT("/"));
+			Print(Context,Info.GroupCount, FontYellow, 1);
+			Print(Context, TEXT("    "));
+
+			// LOD
+			if (Info.bAutoLOD)
+			{
+				Print(Context, TEXT("Auto    "), FontYellow);
+			}
+			else
+			{
+				Print(Context,Info.LOD, FontYellow, 4, 2);
+				Print(Context, TEXT("/"));
+				Print(Context,Info.LODCount-1, FontYellow, 1);
+				Print(Context, TEXT("  "));
+			}
+
+			// Geometry
+			if (Info.GeometryType == 0)		{ Print(Context, TEXT("Strands "), FontEmerald); }
+			else if (Info.GeometryType == 1){ Print(Context, TEXT("Cards   "), FontEmerald); }
+			else if (Info.GeometryType == 2){ Print(Context, TEXT("Meshes  "), FontEmerald); }
+			Print(Context, TEXT(" "));
+
+			// Binding
+			if (Info.BindingType == 0)		{ Print(Context, TEXT("None    "), FontTurquoise); }
+			else if (Info.BindingType == 1)	{ Print(Context, TEXT("Rigid   "), FontTurquoise); }
+			else if (Info.BindingType == 2)	{ Print(Context, TEXT("Skinned "), FontTurquoise); }
+
+			// Cache
+			if (Info.CacheRen)      { Print(Context, TEXT("Ren  "), FontOrange); }
+			else if (Info.CacheSim) { Print(Context, TEXT("Sim  "), FontOrange); }
+			else                    { Print(Context, TEXT("None "), FontSilver); }
+			Print(Context, TEXT(" "));
+			// Sim
+			PrintBool(Context, Info.bHasSim);
+
+			// RBF
+			PrintBool(Context, Info.bHasRBF);
+
+			// Clip
+			if (Info.ClipValue >= 0)
+				Print(Context,Info.ClipValue, FontGreen, 5,3);
+			else 
+				Print(Context,0u, FontSilver, 5, 3);
+
+			// Strands effective
+			Print(Context, Info.ScreenSize, FontOrange, 4, 2);
+			Print(Context, TEXT("   "));
+
+			Print(Context, Info.CoverageScale, FontOrange, 4, 2);
+			Print(Context, TEXT("   "));
+
+			// Curve sCount
+			const float CurveRatio = Info.ActiveCurveCount / float(Info.StrandsCount);
+			PrintPercentage(Context, CurveRatio, InitFontColor(ColorMapTurbo(CurveRatio)));
+			PrintLargeNumber(Context, Info.ActiveCurveCount, FontOrange);
 			Print(Context, TEXT("/"));
-			Print(Context,Info.LODCount-1, FontYellow, 1);
+			PrintLargeNumber(Context, Info.StrandsCount, FontSilver);
+			Print(Context, TEXT(" "), FontWhite);
+
+			// Point sCount
+			const float VertexRatio = Info.ActivePointCount / float(Info.VertexCount);
+			PrintPercentage(Context, VertexRatio, InitFontColor(ColorMapTurbo(VertexRatio)));
+			PrintLargeNumber(Context, Info.ActivePointCount, FontOrange);
+			Print(Context, TEXT("/"));
+			PrintLargeNumber(Context, Info.VertexCount, FontSilver);
+			Print(Context, TEXT(" "), FontWhite);
+
+			// Binding data
+			Print(Context,Info.UniqueTriangle, FontOrange, 8);
+			Print(Context, TEXT(" "));
+			Print(Context,Info.UniqueSection, FontSilver, 3);
 			Print(Context, TEXT("  "));
-		}
 
-		// Geometry
-		if (Info.GeometryType == 0)		{ Print(Context, TEXT("Strands "), FontEmerald); }
-		else if (Info.GeometryType == 1){ Print(Context, TEXT("Cards   "), FontEmerald); }
-		else if (Info.GeometryType == 2){ Print(Context, TEXT("Meshes  "), FontEmerald); }
-		Print(Context, TEXT(" "));
-
-		// Binding
-		if (Info.BindingType == 0)		{ Print(Context, TEXT("None    "), FontTurquoise); }
-		else if (Info.BindingType == 1)	{ Print(Context, TEXT("Rigid   "), FontTurquoise); }
-		else if (Info.BindingType == 2)	{ Print(Context, TEXT("Skinned "), FontTurquoise); }
-
-		// Cache
-		if (Info.CacheRen)      { Print(Context, TEXT("Ren  "), FontOrange); }
-		else if (Info.CacheSim) { Print(Context, TEXT("Sim  "), FontOrange); }
-		else                    { Print(Context, TEXT("None "), FontSilver); }
-		Print(Context, TEXT(" "));
-		// Sim
-		PrintBool(Context, Info.bHasSim);
+			// Primary view
+			PrintBool(Context, Info.bIsVisibleInPrimaryView);
+			Print(Context, TEXT(" "));
 
-		// RBF
-		PrintBool(Context, Info.bHasRBF);
-
-		// Clip
-		if (Info.ClipValue >= 0)
-			Print(Context,Info.ClipValue, FontGreen, 5,3);
-		else 
-			Print(Context,0u, FontSilver, 5, 3);
-
-		// Strands effective
-		Print(Context, Info.ScreenSize, FontOrange, 4, 2);
-		Print(Context, TEXT("   "));
-
-		Print(Context, Info.CoverageScale, FontOrange, 4, 2);
-		Print(Context, TEXT("   "));
-
-		// Curve sCount
-		const float CurveRatio = Info.ActiveCurveCount / float(Info.StrandsCount);
-		PrintPercentage(Context, CurveRatio, InitFontColor(ColorMapTurbo(CurveRatio)));
-		PrintLargeNumber(Context, Info.ActiveCurveCount, FontOrange);
-		Print(Context, TEXT("/"));
-		PrintLargeNumber(Context, Info.StrandsCount, FontSilver);
-		Print(Context, TEXT(" "), FontWhite);
-
-		// Point sCount
-		const float VertexRatio = Info.ActivePointCount / float(Info.VertexCount);
-		PrintPercentage(Context, VertexRatio, InitFontColor(ColorMapTurbo(VertexRatio)));
-		PrintLargeNumber(Context, Info.ActivePointCount, FontOrange);
-		Print(Context, TEXT("/"));
-		PrintLargeNumber(Context, Info.VertexCount, FontSilver);
-		Print(Context, TEXT(" "), FontWhite);
-
-		// Binding data
-		Print(Context,Info.UniqueTriangle, FontOrange, 8);
-		Print(Context, TEXT(" "));
-		Print(Context,Info.UniqueSection, FontSilver, 3);
-		Print(Context, TEXT("  "));
-
-		// Primary view
-		PrintBool(Context, Info.bIsVisibleInPrimaryView);
-		Print(Context, TEXT(" "));
-
-		// Shadow view
-		PrintBool(Context, Info.bIsVisibleInShadowView);
-		Print(Context, TEXT("  "));
-
-		// Has LOD switched
-		PrintBool(Context, Info.bHasLODSwitched);
-		Print(Context, TEXT("  "));
-
-		// Has RT geom
-		PrintBool(Context, Info.bHasRTGeom);
-		Print(Context, TEXT("  "));
+			// Shadow view
+			PrintBool(Context, Info.bIsVisibleInShadowView);
+			Print(Context, TEXT("  "));
+
+			// Has LOD switched
+			PrintBool(Context, Info.bHasLODSwitched);
+			Print(Context, TEXT("  "));
+
+			// Has RT geom
+			PrintBool(Context, Info.bHasRTGeom);
+			Print(Context, TEXT("  "));
 		
 		
-		// Name
-		PrintInstanceNames(Context, InstanceIndex, FontWhite);
+			// Name
+			PrintInstanceNames(Context, InstanceIndex, FontWhite);
 
+			Newline(Context);
+		}
 		Newline(Context);
 	}
-	Newline(Context);
 
 	// Optional advanced debugging
 	if (bShowAdvancedStrandsData)
