Substrate - classification shade should not reference the SubstrateGBuffer UAV when blendable Gbuffer is used.

#tests pc editor, console
#rnx

diff --git a/Engine/Shaders/Private/Substrate/SubstrateMaterialClassification.usf b/Engine/Shaders/Private/Substrate/SubstrateMaterialClassification.usf
--- a/Engine/Shaders/Private/Substrate/SubstrateMaterialClassification.usf
+++ b/Engine/Shaders/Private/Substrate/SubstrateMaterialClassification.usf
@@ -25,7 +25,9 @@ Texture2D<SUBSTRATE_TOP_LAYER_TYPE> TopLayerTexture;
 #if PERMUTATION_CMASK
 Texture2D<uint> TopLayerCmaskTexture;
 #endif
+#if SUBSTRATE_GBUFFER_FORMAT==1
 RWTexture2DArray<uint> MaterialTextureArrayUAV;
+#endif
 
 uint TileEncoding;
 uint4 TileListBufferOffsets[SUBSTRATE_TILE_TYPE_COUNT];
@@ -228,13 +230,15 @@ void TileMainCS(uint2 DispatchThreadId : SV_DispatchThreadID, uint LinearIndex :
 
 #endif // SUBSTRATE_GBUFFER_FORMAT==0
 	}
-	
+
+#if SUBSTRATE_GBUFFER_FORMAT==1
 	BRANCH
 	if (!bContainsScreenSpaceSubsurfaceScattering)
 	{
 		// We must fill all the pixel which does not have subsurface scattering by default so that the SSS code is not executed where it should not.
 		SubstrateStoreSubsurfaceHeader(MaterialTextureArrayUAV, FirstSliceStoringSubstrateSSSData, PixelCoord, 0u); // This is a good clear for FSubstrateSubsurfaceHeader, and we only need to clear the header.
 	}
+#endif
 
 #if PERMUTATION_WAVE_OPS
 	const bool bTileContainsSubstrate						= WaveActiveAnyTrue(bContainsSubstrateMaterial);
