Always perform history weight calculation at full precision. Fixes visual stability issues in cases of extreme brightness differences, such as physically-plausible sun against the sky.

diff --git a/Engine/Shaders/Private/TemporalSuperResolution/TSRUpdateHistory.usf b/Engine/Shaders/Private/TemporalSuperResolution/TSRUpdateHistory.usf
--- a/Engine/Shaders/Private/TemporalSuperResolution/TSRUpdateHistory.usf
+++ b/Engine/Shaders/Private/TemporalSuperResolution/TSRUpdateHistory.usf
@@ -1283,15 +1283,18 @@ void MainCS(
 
 		tsr_half2 PrevHistoryToneWeight = HdrWeightY(Luma4(BlendedPrevHighFrequencyColor));
 		tsr_half2 FilteredInputToneWeight = HdrWeight4(FilteredInputColor);
+
+		// NOTE: weight calculation is intentionally done at full float precision to 
+		// avoid quantization errors in cases of extreme difference.
 			
-		tsr_half2 BlendPrevHistory = PrevWeight * PrevHistoryToneWeight;
-		tsr_half2 BlendFilteredInput = CurrentWeight * FilteredInputToneWeight;
+		float2 BlendPrevHistory = PrevWeight * PrevHistoryToneWeight;
+		float2 BlendFilteredInput = CurrentWeight * FilteredInputToneWeight;
 
-		tsr_half2 CommonWeight = SafeRcp(BlendPrevHistory + BlendFilteredInput);
+		float2 CommonWeight = SafeRcp(BlendPrevHistory + BlendFilteredInput);
 
 		FinalHighFrequencyColor = (
-			dpv_scale(BlendedPrevHighFrequencyColor, CommonWeight * BlendPrevHistory) +
-			dpv_scale(FilteredInputColor, CommonWeight * BlendFilteredInput));
+			dpv_scale(BlendedPrevHighFrequencyColor, tsr_half2(CommonWeight * BlendPrevHistory)) +
+			dpv_scale(FilteredInputColor, tsr_half2(CommonWeight * BlendFilteredInput)));
 				
 		// Quantize validity for the 8bit encoding to avoid numerical shift between color and validity.
 		FinalHistoryValidity = ceil(tsr_half(255.0) * OutputValidity) * rcp(tsr_half(255.0));
