Fix missing vertex offset causing incorrect groom card UV to be fetched.

diff --git a/Engine/Shaders/Private/HairStrands/HairCardsVertexFactory.ush b/Engine/Shaders/Private/HairStrands/HairCardsVertexFactory.ush
--- a/Engine/Shaders/Private/HairStrands/HairCardsVertexFactory.ush
+++ b/Engine/Shaders/Private/HairStrands/HairCardsVertexFactory.ush
@@ -252,6 +252,7 @@ void GetTexCoords(FVertexFactoryInput Input, inout float2 OutTexCoords[NUM_MATER
 {
 #if NUM_MATERIAL_TEXCOORDS_VERTEX
 #if MANUAL_VERTEX_FETCH
+	const uint NumFetchUVs = DivideAndRoundUp2(HairCardsVF.NumUVs);
 	UNROLL
 	for (uint CoordIndex = 0; CoordIndex < NUM_MATERIAL_TEXCOORDS_VERTEX; ++CoordIndex)
 	{
@@ -259,7 +260,7 @@ void GetTexCoords(FVertexFactoryInput Input, inout float2 OutTexCoords[NUM_MATER
 		const uint ReadCoordIndex = CoordIndex + 1; // UVs[0] == RootUV, which is skipped 
 		if (ReadCoordIndex < HairCardsVF.NumUVs)
 		{
-			const uint FetchIndex = ReadCoordIndex >> 1u;
+			const uint FetchIndex = Input.VertexId * NumFetchUVs + (ReadCoordIndex >> 1u);
 			const float4 PackedUVs = HairCardsVF.UVsBuffer[FetchIndex];
 			OutTexCoords[CoordIndex] = (CoordIndex & 1) == 0 ? PackedUVs.zw : PackedUVs.xy; //.xy/.zw are inversed, as TexCoord are shifted by 1 due to UVs[0] == RootUV, which is skipped 
 		}
