Reduce VSM shadow projection cost by 0.04ms on PS5

Avoid loading SSS data during initial part of the shader, and avoid reading SSS data when not necessary.

diff --git a/Engine/Shaders/Private/VirtualShadowMaps/VirtualShadowMapProjection.usf b/Engine/Shaders/Private/VirtualShadowMaps/VirtualShadowMapProjection.usf
--- a/Engine/Shaders/Private/VirtualShadowMaps/VirtualShadowMapProjection.usf
+++ b/Engine/Shaders/Private/VirtualShadowMaps/VirtualShadowMapProjection.usf
@@ -63,18 +63,13 @@ FProjectionShadingInfo GetProjectionShadingInfo(uint2 PixelPos)
 	FSubstratePixelHeader SubstratePixelHeader = UnpackSubstrateHeaderIn(Substrate.MaterialTextureArray, SubstrateAddressing, Substrate.TopLayerTexture);
 	const FSubstrateTopLayerData TopLayerData = SubstrateUnpackTopLayerData(Substrate.TopLayerTexture.Load(uint3(PixelPos, 0)));
 
-	const FSubstrateSubsurfaceHeader SSSHeader = SubstrateLoadSubsurfaceHeader(Substrate.MaterialTextureArray, Substrate.FirstSliceStoringSubstrateSSSData, PixelPos);
-	const bool bIsValid = SubstrateSubSurfaceHeaderGetIsValid(SSSHeader);
-	const bool bIsProfile = SubstrateSubSurfaceHeaderGetIsProfile(SSSHeader);
-
-	const uint SSSType = SSSHEADER_TYPE(SSSHeader);
 	const bool bHasSSS = SubstratePixelHeader.HasSubsurface();
 	const uint BSDFType = SubstratePixelHeader.SubstrateGetBSDFType();
 
 	Out.bIsValid = SubstratePixelHeader.IsSubstrateMaterial();
 	Out.bIsHair = BSDFType == SUBSTRATE_BSDF_TYPE_HAIR;
 	Out.WorldNormal = TopLayerData.WorldNormal;
-	Out.bIsSubsurface = SSSType != SSS_TYPE_NONE; // legacy IsSubsurfaceModel is true for subsurface profile, eyes, etc.
+	Out.bIsSubsurface = bHasSSS; // legacy IsSubsurfaceModel is true for subsurface profile, eyes, etc.
 
 	// VSM traces using the top layer normal when NoL>0 only. So if normals for slabs below do not match, light leaking will occur.
 	// For instance: a flat translucent top layer above a normal mapped bottom layer.
@@ -811,7 +806,7 @@ void VirtualShadowMapProjection(
 	const FProjectionShadingInfo ShadingInfo = GetProjectionShadingInfo(PixelPos);
 
 	FSubsurfaceOpacityMFP SubsurfaceOpacityMFP = GetInitialisedSubsurfaceOpacityMFP();
-	if (!bIsHairInput)
+	if (!bIsHairInput && ShadingInfo.bIsSubsurface)
 	{
 		SubsurfaceOpacityMFP = GetSubsurfaceOpacityFromGbuffer(PixelPos);
 	}
