Substrate - do not activate energy conservation with blendable gbuffer by default when Substrate is enabled. For performance parity.

#tests console
#rnx

diff --git a/Engine/Shaders/Private/ShadingEnergyConservation.ush b/Engine/Shaders/Private/ShadingEnergyConservation.ush
--- a/Engine/Shaders/Private/ShadingEnergyConservation.ush
+++ b/Engine/Shaders/Private/ShadingEnergyConservation.ush
@@ -4,8 +4,9 @@
 
 // If EnergyConservation flag is not enabled explicitly then:
 // * If Substrate is enabled: ON by default
+// * If Substrate is enabled with legacy blendable GBuffer: OFF by default for performance parity with legacy on low end platforms.
 // * If Substrate is disabled: OFF by default
-#if SUBSTRATE_ENABLED || LEGACY_MATERIAL_ENERGYCONSERVATION
+#if (SUBSTRATE_ENABLED && SUBSTRATE_GBUFFER_FORMAT>0) || LEGACY_MATERIAL_ENERGYCONSERVATION
 	#undef USE_ENERGY_CONSERVATION
 	#if SUPPORTS_INDEPENDENT_SAMPLERS 
 		#define USE_ENERGY_CONSERVATION 1
