RandomSequence: Cleaner API for random sequences

Remove unused halton implementation, change struct to have methods to draw both uniform floats and uniform uints (which gives extra precision for contexts where it might be useful to work with 32 bits of precision).

Added a splitting method for contexts where a sampling tree needs to be wider within an existing loop over samples.

The old API is left intact so code can be transitioned incrementally. The behavior of the sequences is not changed (except random, which is just used as a reference as is not performance critical).

Removed commented out experiments.

#rnx

diff --git a/Engine/Shaders/Private/PathTracing/Utilities/PathTracingRandomSequence.ush b/Engine/Shaders/Private/PathTracing/Utilities/PathTracingRandomSequence.ush
--- a/Engine/Shaders/Private/PathTracing/Utilities/PathTracingRandomSequence.ush
+++ b/Engine/Shaders/Private/PathTracing/Utilities/PathTracingRandomSequence.ush
@@ -9,9 +9,8 @@
 #include "../../MortonCode.ush"
 
 #define RANDSEQ_PURERANDOM	0
-#define RANDSEQ_HALTON		1
-#define RANDSEQ_OWENSOBOL	2
-#define RANDSEQ_LATTICE     3
+#define RANDSEQ_OWENSOBOL	1
+#define RANDSEQ_LATTICE     2
 
 // Select a default if none was specified
 #ifndef RANDSEQ
@@ -68,87 +67,10 @@ uint WeakIntegerHash(uint x)
 	return x;
 }
 
-struct RandomSequence
-{
-	uint SampleIndex;		// index into the random sequence
-	uint SampleSeed;		// changes as we draw samples to reflect the change in dimension
-};
-
-float Rand(inout uint Seed)
-{
-	// Counter based PRNG -- safer than most small-state PRNGs since we use the random values directly here.
-	Seed += 1;
-	uint Output = StrongIntegerHash(Seed);
-	// take low 24 bits
-	return (Output & 0xFFFFFF) * 5.96046447754e-08; // * 2^-24
-}
-
-// #dxr_todo: convert prime factors to constant buffer
-static const uint Primes512LUT[] = {
-	2,3,5,7,11,13,17,19,23,29,
-	31,37,41,43,47,53,59,61,67,71,
-	73,79,83,89,97,101,103,107,109,113,
-	127,131,137,139,149,151,157,163,167,173,
-	179,181,191,193,197,199,211,223,227,229,
-	233,239,241,251,257,263,269,271,277,281,
-	283,293,307,311,313,317,331,337,347,349,
-	353,359,367,373,379,383,389,397,401,409,
-	419,421,431,433,439,443,449,457,461,463,
-	467,479,487,491,499,503,509,521,523,541,
-	547,557,563,569,571,577,587,593,599,601,
-	607,613,617,619,631,641,643,647,653,659,
-	661,673,677,683,691,701,709,719,727,733,
-	739,743,751,757,761,769,773,787,797,809,
-	811,821,823,827,829,839,853,857,859,863,
-	877,881,883,887,907,911,919,929,937,941,
-	947,953,967,971,977,983,991,997,1009,1013,
-	1019,1021,1031,1033,1039,1049,1051,1061,1063,1069,
-	1087,1091,1093,1097,1103,1109,1117,1123,1129,1151,
-	1153,1163,1171,1181,1187,1193,1201,1213,1217,1223,
-	1229,1231,1237,1249,1259,1277,1279,1283,1289,1291,
-	1297,1301,1303,1307,1319,1321,1327,1361,1367,1373,
-	1381,1399,1409,1423,1427,1429,1433,1439,1447,1451,
-	1453,1459,1471,1481,1483,1487,1489,1493,1499,1511,
-	1523,1531,1543,1549,1553,1559,1567,1571,1579,1583,
-	1597,1601,1607,1609,1613,1619,1621,1627,1637,1657,
-	1663,1667,1669,1693,1697,1699,1709,1721,1723,1733,
-	1741,1747,1753,1759,1777,1783,1787,1789,1801,1811,
-	1823,1831,1847,1861,1867,1871,1873,1877,1879,1889,
-	1901,1907,1913,1931,1933,1949,1951,1973,1979,1987,
-	1993,1997,1999,2003,2011,2017,2027,2029,2039,2053,
-	2063,2069,2081,2083,2087,2089,2099,2111,2113,2129,
-	2131,2137,2141,2143,2153,2161,2179,2203,2207,2213,
-	2221,2237,2239,2243,2251,2267,2269,2273,2281,2287,
-	2293,2297,2309,2311,2333,2339,2341,2347,2351,2357,
-	2371,2377,2381,2383,2389,2393,2399,2411,2417,2423,
-	2437,2441,2447,2459,2467,2473,2477,2503,2521,2531,
-	2539,2543,2549,2551,2557,2579,2591,2593,2609,2617,
-	2621,2633,2647,2657,2659,2663,2671,2677,2683,2687,
-	2689,2693,2699,2707,2711,2713,2719,2729,2731,2741,
-	2749,2753,2767,2777,2789,2791,2797,2801,2803,2819,
-	2833,2837,2843,2851,2857,2861,2879,2887,2897,2903,
-	2909,2917,2927,2939,2953,2957,2963,2969,2971,2999,
-	3001,3011,3019,3023,3037,3041,3049,3061,3067,3079,
-	3083,3089,3109,3119,3121,3137,3163,3167,3169,3181,
-	3187,3191,3203,3209,3217,3221,3229,3251,3253,3257,
-	3259,3271,3299,3301,3307,3313,3319,3323,3329,3331,
-	3343,3347,3359,3361,3371,3373,3389,3391,3407,3413,
-	3433,3449,3457,3461,3463,3467,3469,3491,3499,3511,
-	3517,3527,3529,3533,3539,3541,3547,3557,3559,3571,
-	3581,3583,3593,3607,3613,3617,3623,3631,3637,3643,
-	3659,3671
-};
-uint Prime512(uint Dimension)
-{
-	return Primes512LUT[Dimension % 512];
-}
-
-
+// Base must be a prime number
 float Halton(uint Index, uint Base)
 {
-	float r = 0.0;
-	float f = 1.0;
-
+	float r = 0.0, f = 1.0;
 	float BaseInv = 1.0 / Base;
 	while (Index > 0)
 	{
@@ -156,12 +78,11 @@ float Halton(uint Index, uint Base)
 		r += f * (Index % Base);
 		Index /= Base;
 	}
-
 	return r;
 }
 
-
-uint FastOwenScrambling(uint Index, uint Seed) {
+uint FastOwenScramblingCore(uint Index, uint Seed)
+{
 #if 0
 	// reference implementation - scramble bits one at a time
 
@@ -173,7 +94,7 @@ uint FastOwenScrambling(uint Index, uint Seed) {
 		Seed ^= Input & Mask;  // flip state  (depending on input)
 	}
 	return Index;
-#else
+#elif 1
 	// Laine and Karras / Stratified Sampling for Stochastic Transparency / EGSR 2011
 
 	// The operations below will mix bits toward the left, so temporarily reverse the order
@@ -200,11 +121,42 @@ uint FastOwenScrambling(uint Index, uint Seed) {
 	Index += Seed; // randomize the index by our seed (pushes bits toward the left)
 	Index ^= Index * 0x9c117646u;
 	Index ^= Index * 0xe0705d72u;
+#else
+	// This is a higher quality owen-scrambling hash but doesn't appear to make much difference
+	// https://psychopath.io/post/2021_01_30_building_a_better_lk_hash
+	Index ^= Index * 0x3d20adeau;
+	Index += Seed;
+	Index *= (Seed >> 16) | 1;
+	Index ^= Index * 0x05526c56u;
+	Index ^= Index * 0x53a22864u;
+#endif
+	return Index;
+}
 
+uint FastOwenScrambling(uint Index, uint Seed)
+{
+	Index = FastOwenScramblingCore(Index, Seed);
 	// Undo the reverse so that we get left-to-right scrambling
 	// thereby emulating owen-scrambling
 	return reversebits(Index);
-#endif
+}
+
+// Given a point (X,Y) on the unit square [0,2^B)^2, return the index along the hilbert curve in [0,2^2B)
+// Because this function only loops over the bits it uses, there is no need to mask out bits beyond the
+// bit count, the pattern will simply tile.
+// TODO: This procedure should probably be turned into a small lookup texture
+uint HilbertCurveInverse(uint X, uint Y, int NumBits)
+{
+	// [REDACTED_PROFILE_URL]/Hackers-Delight/blob/master/hilbert/hil_s_from_xy.c.txt
+	uint TileID = 0, HilbertState = 0;
+	for (int i = NumBits; i--;) {
+		uint xi = (X >> i) & 1;
+		uint yi = (Y >> i) & 1;
+		uint Row = 8 * HilbertState + 4 * xi + 2 * yi;
+		TileID = TileID * 4 + ((0x361E9CB4u >> Row) & 3);
+		HilbertState = (0x8FE65831u >> Row) & 3;
+	}
+	return TileID;
 }
 
 // 32-bit Sobol matrices for dimension 1,2,3 from:
@@ -213,7 +165,7 @@ uint FastOwenScrambling(uint Index, uint Seed) {
 // NOTE: we don't bother storing dimension 0 since it is just a bit reversal
 // NOTE2: we don't bother storing dimension 1 either since it has a very simple pattern
 // NOTE3: the matrix elements are reversed to save one reverse in the owen scrambling
-static const uint2 SobolMatrices[] = {
+static const uint2 SobolMatrices[32] = {
 	uint2(0x00000001, 0x00000001),
 	uint2(0x00000003, 0x00000003),
 	uint2(0x00000006, 0x00000004),
@@ -252,52 +204,22 @@ uint EvolveSobolSeed(inout uint Seed)
 {
 	// constant from: https://www.pcg-random.org/posts/does-it-beat-the-minimal-standard.html
 	const uint MCG_C = 2739110765;
-#if 0
-	return Seed += MCG_C;
-#elif 0
-	Seed += MCG_C;
-	return Seed ^ (Seed >> 16); // scramble lower bits
-#elif 0
-	return PerfectIntegerHash(Seed += 1);
-#elif 0
-	return StrongIntegerHash(Seed += 1);
-#else
 	// a slightly weaker hash is ok since this drives FastOwenScrambling which is itself a hash
 	// Note that the Seed evolution is just an integer addition and the hash should optimize away
 	// when a particular dimension is not used
 	return WeakIntegerHash(Seed += MCG_C);
-#endif
-}
-
-float4 LatticeSampler(uint SampleIndex, inout uint Seed)
-{
-#if 1
-	// Same as FastOwenScrambling, but without the final reversebits
-	uint LatticeIndex = SampleIndex + EvolveSobolSeed(Seed);
-	LatticeIndex ^= LatticeIndex * 0x9c117646u;
-	LatticeIndex ^= LatticeIndex * 0xe0705d72u;
-#else
-	// This is a higher quality owen-scrambling hash but doesn't appear to make much difference
-	// https://psychopath.io/post/2021_01_30_building_a_better_lk_hash
-	uint LatticeIndex = SampleIndex, S = EvolveSobolSeed(Seed);
-	LatticeIndex ^= LatticeIndex * 0x3d20adeau;
-	LatticeIndex += S;
-	LatticeIndex *= (S >> 16) | 1;
-	LatticeIndex ^= LatticeIndex * 0x05526c56u;
-	LatticeIndex ^= LatticeIndex * 0x53a22864u;
-#endif
-
-	// Lattice parameters taken from:
-	// Weighted compound integration rules with higher order convergence for all N
-	// Fred J. Hickernell, Peter Kritzer, Frances Y. Kuo, Dirk Nuyens
-	// Numerical Algorithms - February 2012
-	uint4 Result = LatticeIndex * uint4(1, 364981, 245389, 97823);
-
-	return (Result >> 8) * 5.96046447754e-08; // * 2^-24
 }
 
-float4 SobolSampler(uint SampleIndex, inout uint Seed)
+uint4 SamplerCore(uint SampleIndex, inout uint Seed)
 {
+#if RANDSEQ == RANDSEQ_PURERANDOM
+	uint4 Result = uint4(StrongIntegerHash(Seed + 0),
+						 StrongIntegerHash(Seed + 1),
+						 StrongIntegerHash(Seed + 2),
+						 StrongIntegerHash(Seed + 3));
+	Seed += 4;
+	return Result;
+#elif RANDSEQ == RANDSEQ_OWENSOBOL
 	// first scramble the index to decorelate from other 4-tuples
 	uint SobolIndex = FastOwenScrambling(SampleIndex, EvolveSobolSeed(Seed));
 	// now get Sobol' point from this index
@@ -313,6 +235,7 @@ float4 SobolSampler(uint SampleIndex, inout uint Seed)
 	Result.y ^= (Result.y & 0xF0F0F0F0) >>  4;
 	Result.y ^= (Result.y & 0xCCCCCCCC) >>  2;
 	Result.y ^= (Result.y & 0xAAAAAAAA) >>  1;
+	// Remaining bits require a bit-vector by matrix multiplication (in F2, so multiplication is 'and' and addition is 'xor')
 #if RANDSEQ_UNROLL_SOBOL // unrolled version seems faster in most cases
 	UNROLL for (uint b = 0; b < 32; b++)
 	{
@@ -329,43 +252,87 @@ float4 SobolSampler(uint SampleIndex, inout uint Seed)
 	Result.y = FastOwenScrambling(Result.y, EvolveSobolSeed(Seed));
 	Result.z = FastOwenScrambling(Result.z, EvolveSobolSeed(Seed));
 	Result.w = FastOwenScrambling(Result.w, EvolveSobolSeed(Seed));
-
-	// output as float in [0,1) taking care not to skew the distribution
-	// due to the non-uniform spacing of floats in this range
-	return (Result >> 8) * 5.96046447754e-08; // * 2^-24
+	return Result;
+#elif RANDSEQ == RANDSEQ_LATTICE
+	// Simplified algorithm using rank-1 lattice sequences. It avoids the extra call to reverse bits
+	// Scrambling the _output_ does not appear to help, so we only need to scramble the index, leading to a very fast algorithm
+	// Unfortunately, the resulting sequence does not have quite as good quality as Sobol in practice
+	uint LatticeIndex = FastOwenScramblingCore(SampleIndex, EvolveSobolSeed(Seed));
+	// Lattice parameters taken from:
+	// Weighted compound integration rules with higher order convergence for all N
+	// Fred J. Hickernell, Peter Kritzer, Frances Y. Kuo, Dirk Nuyens
+	// Numerical Algorithms - February 2012
+	return LatticeIndex * uint4(1, 364981, 245389, 97823);
+#endif
 }
 
-// This initializes a multi-dimensional sampler for a particular pixel. The TimeSeed parameter should be set
+/// A simple struct that encapsulates a random number sequence (extensible in dimension)
+/// Once you create the struct, you can draw samples in the unit square (as either floats or bits) in any dimension from 1 to 4.
+/// If you need more dimensions, you can repeat the calls, which will be guaranteed to be statistically independent from the previous
+/// However the values from a single call (Get1D(), Get2D(), etc ...) should have nice stratification properties
+/// If you need to draw N samples for a pixel, you need to create the struct N times.
+struct FRandomSequence
+{
+	uint SampleIndex;		// index into the random sequence
+	uint SampleSeed;		// changes as we draw samples to reflect the change in dimension
+
+	// Return a random value in [0,2^32)^d
+	// All functions rely on the 4D variant and we rely on compiler to optimize out the dead code
+	// This also guarantees that the sequence will not "shift" if you change a call from one dimension to another because the seed will be changed in a consistent way in all cases
+	uint  Get1DBits() { return SamplerCore(SampleIndex, SampleSeed).x; }
+	uint2 Get2DBits() { return SamplerCore(SampleIndex, SampleSeed).xy; }
+	uint3 Get3DBits() { return SamplerCore(SampleIndex, SampleSeed).xyz; }
+	uint4 Get4DBits() { return SamplerCore(SampleIndex, SampleSeed); }
+
+	// Convert the random value in [0,2^32)^d to a uniformly distributed float in [0,1)^d,
+	// taking care not to skew the distribution due to the non-uniform spacing of floats.
+	// Use the upper 24 bits to preserve the good properties of low-discrenpancy samplers.
+	// The floating point constant is 2^-24 (would be more nicely expressed with hex-floats which are not yet supported in HLSL)
+	// The implicit int to float cast here is guaranteed to be lossless because float can represent integers up to 2^24 exactly
+	float  Get1D() { return (Get1DBits() >> 8) * 5.96046447754e-08; }
+	float2 Get2D() { return (Get2DBits() >> 8) * 5.96046447754e-08; }
+	float3 Get3D() { return (Get3DBits() >> 8) * 5.96046447754e-08; }
+	float4 Get4D() { return (Get4DBits() >> 8) * 5.96046447754e-08; }
+
+	// Split the path into Num entries. The RandomSequence object returned can be used to keep sampling along the "Index" branch of the split
+	// Index must be a value in [0,Num), Num should be >0
+	FRandomSequence Split(uint Index, uint Num);
+};
+
+// This initializes a multi-dimensional sampler for a particular pixel. The FrameIndex parameter should be set
 // such that it changes per sample. See the function below when you know how many samples you plan on taking
 // for a particular integral.
-void RandomSequence_Initialize(inout RandomSequence RandSequence, uint PositionSeed, uint TimeSeed)
+FRandomSequence RandomSequenceCreate(uint PositionSeed, uint FrameIndex)
 {
 	// optionally disable randomization per pixel so that all pixels follow the same path in primary sample space
 #if RANDSEQ_RANDOMIZED == 0
 	PositionSeed = 0;
 #endif
 
+	FRandomSequence RandSequence = (FRandomSequence) 0;
+
 #if RANDSEQ == RANDSEQ_PURERANDOM
 	RandSequence.SampleIndex = 0; // not used
-	RandSequence.SampleSeed = StrongIntegerHash(PositionSeed + StrongIntegerHash(TimeSeed));
-#elif RANDSEQ == RANDSEQ_HALTON
-	RandSequence.SampleIndex = TimeSeed + StrongIntegerHash(PositionSeed);
-	RandSequence.SampleSeed = 0; // this sampler just counts dimensions
+	RandSequence.SampleSeed = StrongIntegerHash(PositionSeed + StrongIntegerHash(FrameIndex));
 #elif RANDSEQ == RANDSEQ_OWENSOBOL || RANDSEQ == RANDSEQ_LATTICE
 	// pre-compute bit reversal needed for FastOwenScrambling since this index doesn't change
-	RandSequence.SampleIndex = reversebits(TimeSeed);
+	RandSequence.SampleIndex = reversebits(FrameIndex);
 	// change seed to get a unique sequence per pixel
 	RandSequence.SampleSeed  = StrongIntegerHash(PositionSeed);
 #else
 #error "Unknown random sequence chosen in path tracer"
 #endif
+	return RandSequence;
 }
 
+
 // This is an alternative initialization to the method above when you know how many samples you are planning on taking. This incorporates the 
 // TimeSeed0 (this should usually be the frame number when using all samples in one frame, or should be the frame on which the first sample was
 // taken if taking one sample per frame like the path tracer).
-void RandomSequence_Initialize(inout RandomSequence RandSequence, uint2 PixelCoord, uint SampleIndex, uint TimeSeed0, uint MaxSamples)
+FRandomSequence RandomSequenceCreate(uint3 PixelCoordAndFrame, uint SampleIndex, uint MaxSamples)
 {
+	FRandomSequence RandSequence = (FRandomSequence) 0;
+
 #if RANDSEQ_ERROR_DIFFUSION == 1 && (RANDSEQ == RANDSEQ_OWENSOBOL || RANDSEQ == RANDSEQ_LATTICE)
 	// The core idea is to use a single sobol pattern across the screen using a space filling curve to correlate nearby pixels
 	// This achieves a much better error diffusion than the random sequence assignment of the method above
@@ -381,35 +348,10 @@ void RandomSequence_Initialize(inout RandomSequence RandSequence, uint2 PixelCoo
 	//   - quality does not improve beyond the target sample count, making usage with adaptive sampling difficult
 	//   - error-diffusion property is only achieved at the target sample count, earlier samples look random (sometimes slightly worse than random)
 
-	uint X = PixelCoord.x;
-	uint Y = PixelCoord.y;
-
-#if 0
-	uint TileID = ((X & 255) + (Y & 255) * 256); // Scanline order
-#elif 0
-	uint TileID = ((Y & 1) ? 255 - (X & 255) : (X & 255)) + (Y & 255) * 256; // zig-zag curve
-#elif 0
-	uint TileID = MortonEncode(uint2(X, Y)); // Z-order curve
-#elif 0
-	// ordered dither curve
-	uint TileID = MortonEncode(uint2(X ^ Y, X));
-#elif 0
-    // pure random (breaks nice properties)
-	uint TileID = StrongIntegerHash(X + StrongIntegerHash(Y));
-#else
-	// Hilbert curve: [REDACTED_PROFILE_URL]/Hackers-Delight/blob/master/hilbert/hil_s_from_xy.c.txt
-	uint TileID = 0, HilbertState = 0;
-	for (int i = 7; i >= 0; i--) {
-		uint xi = (X >> i) & 1;
-		uint yi = (Y >> i) & 1;
-		uint Row = 8 * HilbertState + 4 * xi + 2 * yi;
-		TileID = TileID * 4 + ((0x361E9CB4u >> Row) & 3);
-		HilbertState = (0x8FE65831u >> Row) & 3;
-	}
-#endif
+	uint TileID = HilbertCurveInverse(PixelCoordAndFrame.x, PixelCoordAndFrame.y, 8);
 
 	// Combine pixel-level and sample-level bits into the sample index (visible structure will be hidden by owen scrambling of the index)
-	RandSequence.SampleIndex = reversebits((65536 * TimeSeed0 + TileID) * MaxSamples + SampleIndex);
+	RandSequence.SampleIndex = reversebits((65536 * PixelCoordAndFrame.z + TileID) * MaxSamples + SampleIndex);
 
 	// progressive encodings (these kind of work, but quality is much worse)
 	//RandSequence.SampleIndex = reversebits(MortonEncode(uint2(TimeSeed0, TileID)) * MaxSamples + SampleIndex);
@@ -417,99 +359,57 @@ void RandomSequence_Initialize(inout RandomSequence RandSequence, uint2 PixelCoo
 	//RandSequence.SampleIndex = reversebits(MortonEncode(uint2(MortonEncode(uint2(X, Y)), SampleIndex)));
 
 	RandSequence.SampleSeed = 0; // always use the same sequence
+	return RandSequence;
 #else
 	// Error diffusion sampler is disabled, just use the simple init function instead
-	RandomSequence_Initialize(RandSequence, PixelCoord.x + PixelCoord.y * 65536, TimeSeed0 * MaxSamples + SampleIndex);
+	return RandomSequenceCreate(PixelCoordAndFrame.x + PixelCoordAndFrame.y * 65536, PixelCoordAndFrame.z * MaxSamples + SampleIndex);
 #endif
 }
 
-float RandomSequence_GenerateSample1D(inout RandomSequence RandSequence)
+FRandomSequence FRandomSequence::Split(uint Index, uint Num)
 {
-	float Result;
+	FRandomSequence SplitSequence;
 #if RANDSEQ == RANDSEQ_PURERANDOM
-	Result = Rand(RandSequence.SampleSeed);
-#elif RANDSEQ == RANDSEQ_HALTON
-	Result = Halton(RandSequence.SampleIndex, Prime512(RandSequence.SampleSeed + 0));
-	RandSequence.SampleSeed += 1;
-#elif RANDSEQ == RANDSEQ_OWENSOBOL
-	// rely on compiler to optimize out dead code for now
-	// #dxr_todo: benchmark and specialize if needed
-	Result = SobolSampler(RandSequence.SampleIndex, RandSequence.SampleSeed).x;
-#elif RANDSEQ == RANDSEQ_LATTICE
-	Result = LatticeSampler(RandSequence.SampleIndex, RandSequence.SampleSeed).x;
+	SplitSequence.SampleIndex = 0;
+	SplitSequence.SampleSeed = StrongIntegerHash(SampleSeed * Num + Index);
 #else
-#error "Unknown random sequence chosen in path tracer"
+	// Undo the built-in reversebits, stretch the index and re-apply it
+	SplitSequence.SampleSeed = SampleSeed;
+	SplitSequence.SampleIndex = reversebits(reversebits(SampleIndex) * Num + Index);
 #endif
-	return Result;
+	return SplitSequence;
+}
+
+
+// Legacy API - wrappers around the more convenient calls above, but kept for back-compatiblity
+// Deprecated (5.7)
+typedef FRandomSequence RandomSequence;
+void RandomSequence_Initialize(inout RandomSequence RandSequence, uint PositionSeed, uint FrameIndex)
+{
+	RandSequence = RandomSequenceCreate(PositionSeed, FrameIndex);
+}
+
+void RandomSequence_Initialize(inout RandomSequence RandSequence, uint2 PixelCoord, uint SampleIndex, uint FrameIndex, uint MaxSamples)
+{
+	RandSequence = RandomSequenceCreate(uint3(PixelCoord, FrameIndex), SampleIndex, MaxSamples);
+}
+
+float RandomSequence_GenerateSample1D(inout RandomSequence RandSequence)
+{
+	return RandSequence.Get1D();
 }
 
 float2 RandomSequence_GenerateSample2D(inout RandomSequence RandSequence)
 {
-	float2 Result;
-#if RANDSEQ == RANDSEQ_PURERANDOM
-	Result.x = Rand(RandSequence.SampleSeed);
-	Result.y = Rand(RandSequence.SampleSeed);
-#elif RANDSEQ == RANDSEQ_HALTON
-	Result.x = Halton(RandSequence.SampleIndex, Prime512(RandSequence.SampleSeed + 0));
-	Result.y = Halton(RandSequence.SampleIndex, Prime512(RandSequence.SampleSeed + 1));
-	RandSequence.SampleSeed += 2;
-#elif RANDSEQ == RANDSEQ_OWENSOBOL
-	// rely on compiler to optimize out dead code for now
-	// #dxr_todo: benchmark and specialize if needed
-	Result = SobolSampler(RandSequence.SampleIndex, RandSequence.SampleSeed).xy;
-#elif RANDSEQ == RANDSEQ_LATTICE
-	Result = LatticeSampler(RandSequence.SampleIndex, RandSequence.SampleSeed).xy;
-#else
-#error "Unknown random sequence chosen in path tracer"
-#endif
-	return Result;
+	return RandSequence.Get2D();
 }
 
 float3 RandomSequence_GenerateSample3D(inout RandomSequence RandSequence)
 {
-	float3 Result;
-#if RANDSEQ == RANDSEQ_PURERANDOM
-	Result.x = Rand(RandSequence.SampleSeed);
-	Result.y = Rand(RandSequence.SampleSeed);
-	Result.z = Rand(RandSequence.SampleSeed);
-#elif RANDSEQ == RANDSEQ_HALTON
-	Result.x = Halton(RandSequence.SampleIndex, Prime512(RandSequence.SampleSeed + 0));
-	Result.y = Halton(RandSequence.SampleIndex, Prime512(RandSequence.SampleSeed + 1));
-	Result.z = Halton(RandSequence.SampleIndex, Prime512(RandSequence.SampleSeed + 2));
-	RandSequence.SampleSeed += 3;
-#elif RANDSEQ == RANDSEQ_OWENSOBOL
-	// rely on compiler to optimize out dead code for now
-	// #dxr_todo: benchmark and specialize if needed
-	Result = SobolSampler(RandSequence.SampleIndex, RandSequence.SampleSeed).xyz;
-#elif RANDSEQ == RANDSEQ_LATTICE
-	Result = LatticeSampler(RandSequence.SampleIndex, RandSequence.SampleSeed).xyz;
-#else
-#error "Unknown random sequence chosen in path tracer"
-#endif
-	return Result;
+	return RandSequence.Get3D();
 }
 
-
 float4 RandomSequence_GenerateSample4D(inout RandomSequence RandSequence)
 {
-	float4 Result;
-#if RANDSEQ == RANDSEQ_PURERANDOM
-	Result.x = Rand(RandSequence.SampleSeed);
-	Result.y = Rand(RandSequence.SampleSeed);
-	Result.z = Rand(RandSequence.SampleSeed);
-	Result.w = Rand(RandSequence.SampleSeed);
-#elif RANDSEQ == RANDSEQ_HALTON
-	Result.x = Halton(RandSequence.SampleIndex, Prime512(RandSequence.SampleSeed + 0));
-	Result.y = Halton(RandSequence.SampleIndex, Prime512(RandSequence.SampleSeed + 1));
-	Result.z = Halton(RandSequence.SampleIndex, Prime512(RandSequence.SampleSeed + 2));
-	Result.w = Halton(RandSequence.SampleIndex, Prime512(RandSequence.SampleSeed + 3));
-	RandSequence.SampleSeed += 4;
-#elif RANDSEQ == RANDSEQ_OWENSOBOL
-	Result = SobolSampler(RandSequence.SampleIndex, RandSequence.SampleSeed);
-#elif RANDSEQ == RANDSEQ_LATTICE
-	Result = LatticeSampler(RandSequence.SampleIndex, RandSequence.SampleSeed);
-#else
-#error "Unknown random sequence chosen in path tracer"
-#endif
-	return Result;
+	return RandSequence.Get4D();
 }
