Path Tracer: Adopt new RandomSequence API

Update the new API functions instead of the older, more verbose ones

#rnx

diff --git a/Engine/Plugins/Experimental/GPULightmass/Shaders/Private/LightmapPathTracing.usf b/Engine/Plugins/Experimental/GPULightmass/Shaders/Private/LightmapPathTracing.usf
--- a/Engine/Plugins/Experimental/GPULightmass/Shaders/Private/LightmapPathTracing.usf
+++ b/Engine/Plugins/Experimental/GPULightmass/Shaders/Private/LightmapPathTracing.usf
@@ -57,7 +57,7 @@ RWTexture2D<float4> OutputTileAtlas;
 void GenerateCosineNormalRay(
 	float3 TranslatedWorldPosition,
 	float3 WorldNormal,
-	inout RandomSequence RandSequence,
+	inout FRandomSequence RandSequence,
 	out float3 RayOrigin,
 	out float3 RayDirection,
 	out float3 TangentDirection,
@@ -68,7 +68,7 @@ void GenerateCosineNormalRay(
 {
 	// Draw random variable
 	float2 BufferSize = View.BufferSizeAndInvSize.xy;
-	float2 RandSample = RandomSequence_GenerateSample2D(RandSequence);
+	float2 RandSample = RandSequence.Get2D();
 
 	// Perform cosine-hemispherical sampling and convert to world-space
 	float4 Direction_Tangent = CosineSampleHemisphere(RandSample);
@@ -224,7 +224,7 @@ FPathTracingPayload TraceTransparentRay(FRayDesc Ray, bool IsCameraRay, bool Las
 				{
 					// the acceptance probability is not 1.0
 					// generate a random number to see if we should accept this hit
-					float RandValue = RandomSequence_GenerateSample1D(RandSequence);
+					float RandValue = RandSequence.Get1D();
 					AcceptHit = RandValue * SelectionWeightSum < SelectionWeight;
 				}
 				else
@@ -515,7 +515,7 @@ void PathTracingKernel(
 		FLightLoopCount LightLoopCount = LightGridLookup(Payload.TranslatedWorldPos);
 
 		// Choose a random number for both Light sampling and BxDF sampling
-		float4 RandSample = RandomSequence_GenerateSample4D(RandSequence);
+		float4 RandSample = RandSequence.Get4D();
 
 		float LightPickingCdfSum = 0;
 
@@ -680,7 +680,6 @@ void PathTracingKernel(
 		{
 			// If there is some chance we should terminate the ray, draw an extra random value
 			float RussianRouletteRand = RandSample.w;
-			//RussianRouletteRand = RandomSequence_GenerateSample1D(RandSequence);
 			if (RussianRouletteRand >= ContinuationProb)
 			{
 				// stochastically terminate the path
@@ -905,8 +904,7 @@ void LightmapPathTracingMainRG()
 			UNROLL_N(2)
 			for (int HemisphereIndex = 0; HemisphereIndex < (bMaterialTwoSided ? 2 : 1); HemisphereIndex++)
 			{
-				RandomSequence RandSequence;
-				RandomSequence_Initialize(RandSequence, Seed, EffectiveRenderPassIndex);
+				FRandomSequence RandSequence = RandomSequenceCreate(Seed, EffectiveRenderPassIndex);
 				
 				uint SampleIndex = 2;
 
@@ -1269,8 +1267,7 @@ void VolumetricLightmapPathTracingMainRG()
 
 		if (bIsValidSample)
 		{
-			RandomSequence RandSequence;
-			RandomSequence_Initialize(RandSequence, Seed, EffectiveRenderPassIndex);	
+			RandomSequence RandSequence = RandomSequenceCreate(Seed, EffectiveRenderPassIndex);
 			uint SampleIndex = 4;
 		
 			PathTracingKernel(
diff --git a/Engine/Shaders/Private/PathTracing/PathTracingBuildCloudAccelerationMap.usf b/Engine/Shaders/Private/PathTracing/PathTracingBuildCloudAccelerationMap.usf
--- a/Engine/Shaders/Private/PathTracing/PathTracingBuildCloudAccelerationMap.usf
+++ b/Engine/Shaders/Private/PathTracing/PathTracingBuildCloudAccelerationMap.usf
@@ -45,9 +45,7 @@ void PathTracingBuildCloudAccelerationMapCS(uint2 DispatchThreadId : SV_Dispatch
 	}
 	for (uint Sample = 0; Sample < NumSamples; Sample++)
 	{
-		RandomSequence RandSequence;
-		RandomSequence_Initialize(RandSequence, DispatchThreadId.x + DispatchThreadId.y * 65536, Iteration * NumSamples + Sample);
-		float3 Jitter = RandomSequence_GenerateSample3D(RandSequence);
+		float3 Jitter = RandomSequenceCreate(DispatchThreadId.x + DispatchThreadId.y * 65536, Iteration * NumSamples + Sample).Get3D();
 		const float2 UV = 2.0 * ((DispatchThreadId + Jitter.xy) / float(CloudAccelMapResolution)) - 1.0;
 		// Project a 2d map onto the surface of a sphere
 		float3 ClipPos;
diff --git a/Engine/Shaders/Private/PathTracing/PathTracingCommon.ush b/Engine/Shaders/Private/PathTracing/PathTracingCommon.ush
--- a/Engine/Shaders/Private/PathTracing/PathTracingCommon.ush
+++ b/Engine/Shaders/Private/PathTracing/PathTracingCommon.ush
@@ -455,7 +455,7 @@ struct FPackedPathTracingPayload : FMinimalPayload
 
 #if PATH_TRACER_USE_CLOUD_SHADER
 	// Volumetric callable shaders
-	void SetVolumetricCallableShaderInput(float3 Origin, float3 Direction, float TMin, float TMax, RandomSequence RandSeq, uint Flags, float3 Throughput, float CloudDensity)
+	void SetVolumetricCallableShaderInput(float3 Origin, float3 Direction, float TMin, float TMax, FRandomSequence RandSeq, uint Flags, float3 Throughput, float CloudDensity)
 	{
 		HitT = TMin;
 		PackedData[0]  = asuint(TMax);
@@ -494,9 +494,9 @@ struct FPackedPathTracingPayload : FMinimalPayload
 		return Ray;
 	}
 
-	RandomSequence GetVolumetricCallableShaderInputRandSequence()
+	FRandomSequence GetVolumetricCallableShaderInputRandSequence()
 	{
-		RandomSequence RandSeq;
+		FRandomSequence RandSeq;
 		RandSeq.SampleIndex = PackedData[7];
 		RandSeq.SampleSeed  = PackedData[8];
 		return RandSeq;
@@ -530,7 +530,7 @@ struct FPackedPathTracingPayload : FMinimalPayload
 		return Context;
 	}
 
-	void SetVolumetricCallableShaderOutput(RandomSequence RandSeq, float3 Throughput, float3 Albedo)
+	void SetVolumetricCallableShaderOutput(FRandomSequence RandSeq, float3 Throughput, float3 Albedo)
 	{
 		PackedData[0] = RandSeq.SampleIndex;
 		PackedData[1] = RandSeq.SampleSeed;
@@ -540,7 +540,7 @@ struct FPackedPathTracingPayload : FMinimalPayload
 		PackedData[4] = PackFloat2ToUInt(Albedo.yz);
 	}
 
-	void SetVolumetricCallableShaderOutput(RandomSequence RandSeq, float3 Throughput)
+	void SetVolumetricCallableShaderOutput(FRandomSequence RandSeq, float3 Throughput)
 	{
 		// Increased precision for Transmittance mode
 		PackedData[0] = RandSeq.SampleIndex;
@@ -574,9 +574,9 @@ struct FPackedPathTracingPayload : FMinimalPayload
 		PackedData[12] |= f32tof16(Alpha);
 	}
 
-	RandomSequence GetVolumetricCallableShaderOutputRandSeq()
+	FRandomSequence GetVolumetricCallableShaderOutputRandSeq()
 	{
-		RandomSequence RandSeq;
+		FRandomSequence RandSeq;
 		RandSeq.SampleIndex = PackedData[0];
 		RandSeq.SampleSeed  = PackedData[1];
 		return RandSeq;
diff --git a/Engine/Shaders/Private/PathTracing/PathTracingCore.ush b/Engine/Shaders/Private/PathTracing/PathTracingCore.ush
--- a/Engine/Shaders/Private/PathTracing/PathTracingCore.ush
+++ b/Engine/Shaders/Private/PathTracing/PathTracingCore.ush
@@ -3,6 +3,7 @@
 #pragma once
 
 #define RANDSEQ_UNROLL_SOBOL PATH_TRACER_USE_COMPACTION // unrolling sobol sampler is best when we can guarantee high occupancy
+#define RANDSEQ RANDSEQ_LATTICE
 
 #define USE_PATH_TRACING_LIGHT_GRID	1
 #define USE_RAY_TRACING_DECAL_GRID	1
@@ -102,7 +103,7 @@ Buffer<float> StartingExtinctionCoefficient;
 
 struct FPathState
 {                                   // packed size
-	RandomSequence RandSequence;    //  8 bytes
+	FRandomSequence RandSequence;   //  8 bytes
 	float3 Radiance;                // 12 bytes
 	float  Alpha;                   //  4 bytes
 	uint3 PackedAlbedoNormal;       //  6 bytes // half precision (Albedo)
@@ -492,7 +493,7 @@ FPackedPathTracingPayload TraceTransparentRay(inout FPathState PathState, int Bo
 
 
 	const uint MissShaderIndex = 0;
-	float3 RISRandSample = RandomSequence_GenerateSample3D(PathState.RandSequence);
+	float3 RISRandSample = PathState.RandSequence.Get3D();
 	FRISContext HitSample = InitRISContext(RISRandSample.x);
 	FRISContext VolSample = InitRISContext(RISRandSample.y);
 
@@ -1031,7 +1032,7 @@ FPackedPathTracingPayload TraceTransparentRay(inout FPathState PathState, int Bo
 	return Payload;
 }
 
-float3 TraceTransparentVisibilityRay(FRayDesc Ray, float PathRoughness, float PreviousPathRoughness, uint MissShaderIndex, bool bCastShadows, bool bCloudCastShadows, float CloudFactor, inout RandomSequence RandSequence)
+float3 TraceTransparentVisibilityRay(FRayDesc Ray, float PathRoughness, float PreviousPathRoughness, uint MissShaderIndex, bool bCastShadows, bool bCloudCastShadows, float CloudFactor, inout FRandomSequence RandSequence)
 {
 	if (!bCastShadows && MissShaderIndex == 0)
 	{
@@ -1352,7 +1353,7 @@ float4 SampleDwivediPhaseFunction(float3 ColorChannelPdf, float3 DwivediScale, f
 	return Result;
 }
 
-bool ProcessSubsurfaceRandomWalk(inout FPathTracingPayload Payload, inout float3 PathThroughput, inout RandomSequence RandSequence, float3 RayDirection, bool SimplifySSS)
+bool ProcessSubsurfaceRandomWalk(inout FPathTracingPayload Payload, inout float3 PathThroughput, inout FRandomSequence RandSequence, float3 RayDirection, bool SimplifySSS)
 {
 	if (!Payload.IsSubsurfaceMaterial())
 	{
@@ -1374,7 +1375,7 @@ bool ProcessSubsurfaceRandomWalk(inout FPathTracingPayload Payload, inout float3
 	}
 	
 	// decide randomly to evaluate the SSS part of the material, or keep shading the current part
-	float3 RandSample = RandomSequence_GenerateSample3D(RandSequence);
+	float3 RandSample = RandSequence.Get3D();
 	if (RandSample.x < SSS.Prob)
 	{
 		// do random walk (and adjust throughput to remove energy reflected away by the specular lobe)
@@ -1499,7 +1500,7 @@ bool ProcessSubsurfaceRandomWalk(inout FPathTracingPayload Payload, inout float3
 
 		FProbeResult ProbeResult = TraceSSSProbeRay(Ray, InterfaceCounter);
 
-		RandSample = RandomSequence_GenerateSample3D(RandSequence);
+		RandSample = RandSequence.Get3D();
 		if (ProbeResult.IsMiss())
 		{
 			// we didn't hit anything, so scatter according to the scattering distribution in the volume and keep tracing
@@ -1633,25 +1634,25 @@ FPathState CreatePathState(int2 PixelIndex, int2 TextureIndex)
 	// initialize progressive sobol sequence based on current index for this pixel
 	// NOTE: Error diffusion sampler cannot be used in this mode
 	int SampleIndex = int(VarianceTexture[TextureIndex].z);
-	RandomSequence_Initialize(Output.RandSequence, LaunchIndex.x + LaunchIndex.y * 65536, TemporalSeed - Iteration + SampleIndex);
+	Output.RandSequence = RandomSequenceCreate(LaunchIndex.x + LaunchIndex.y * 65536, TemporalSeed - Iteration + SampleIndex);
 #else
 	// Initialize random sequence 
 	if (SamplerType == PATHTRACER_SAMPLER_ERROR_DIFFUSION)
 	{
 		// z-sampler init
-		RandomSequence_Initialize(Output.RandSequence, LaunchIndex, Iteration, TemporalSeed - Iteration, MaxSamples);
+		Output.RandSequence = RandomSequenceCreate(uint3(LaunchIndex, TemporalSeed - Iteration), Iteration, MaxSamples);
 	}
 	else
 	{
 		// random sobol init
-		RandomSequence_Initialize(Output.RandSequence, LaunchIndex.x + LaunchIndex.y * 65536, TemporalSeed);
+		Output.RandSequence = RandomSequenceCreate(LaunchIndex.x + LaunchIndex.y * 65536, TemporalSeed);
 	}
 #endif
 
 	// Initialize ray and payload
 	float2 ViewportUV;
 	{
-		float2 AAJitter = RandomSequence_GenerateSample2D(Output.RandSequence);
+		float2 AAJitter = Output.RandSequence.Get2D();
 #if 0
 		// Tent filter
 		float2 P = 2 * AAJitter - 0.99999994f;
@@ -1670,7 +1671,7 @@ FPathState CreatePathState(int2 PixelIndex, int2 TextureIndex)
 	if (CameraLensRadius.y > 0)
 	{
 	    // DOF enabled - apply simple thin-lens model
-		float2 LensSample = RandomSequence_GenerateSample2D(Output.RandSequence);
+		float2 LensSample = Output.RandSequence.Get2D();
 		float3 ViewX = View.ViewToTranslatedWorld[0].xyz;
 		float3 ViewY = View.ViewToTranslatedWorld[1].xyz;
 		float3 ViewZ = View.ViewToTranslatedWorld[2].xyz;
@@ -1758,7 +1759,7 @@ bool PathTracingKernel(inout FPathState PathState, int Bounce)
 		if (LightPickingCdfSum > 0.0)
 		{
 			// at least one light is overlapping with our ray, so we have a chance to sample it
-			float3 RandSample = RandomSequence_GenerateSample3D(PathState.RandSequence);
+			float3 RandSample = PathState.RandSequence.Get3D();
 			int LightSampleLightId = 0;
 			float LightPickPdf = 0.0;
 			SelectLight(RandSample.x * LightPickingCdfSum, NumLights, LightPickingCdf, LightSampleLightId, LightPickPdf);
@@ -1789,7 +1790,7 @@ bool PathTracingKernel(inout FPathState PathState, int Bounce)
 
 				// prepare a minimal payload that describes the hit we need to shade
 				FPathTracingPayload VolPayload = CreateMediumHitPayload(EquiAngularT, TranslatedWorldPos, Result);
-				float3 LightRandValue = RandomSequence_GenerateSample3D(PathState.RandSequence);
+				float3 LightRandValue = PathState.RandSequence.Get3D();
 
 				const bool bCastShadows = CastsVolumeShadow(LightSampleLightId);
 				const bool bCloudCastShadows = true; // volumes should get shadows from clouds
@@ -1887,7 +1888,7 @@ bool PathTracingKernel(inout FPathState PathState, int Bounce)
 	}
 
 	// Choose a random number for both Light sampling and BxDF sampling
-	float4 RandSample = RandomSequence_GenerateSample4D(PathState.RandSequence);
+	float4 RandSample = PathState.RandSequence.Get4D();
 
 	const float2 DiffuseSpecularScale = PathState.GetDiffuseSpecularScale(bIsVolumeSample);
 
diff --git a/Engine/Shaders/Private/PathTracing/PathTracingVolumetricCloudMaterialShader.usf b/Engine/Shaders/Private/PathTracing/PathTracingVolumetricCloudMaterialShader.usf
--- a/Engine/Shaders/Private/PathTracing/PathTracingVolumetricCloudMaterialShader.usf
+++ b/Engine/Shaders/Private/PathTracing/PathTracingVolumetricCloudMaterialShader.usf
@@ -97,7 +97,7 @@ RAY_TRACING_ENTRY_CALLABLE(PathTracingVolumetricCloudMaterialShader, FPackedPath
 	// This comes at the cost of some duplication of the logic
 	// TODO: Can we use HLSL templates to merge some of the duplicated logic?
 
-	RandomSequence RandSeq = PackedPayload.GetVolumetricCallableShaderInputRandSequence();
+	FRandomSequence RandSeq = PackedPayload.GetVolumetricCallableShaderInputRandSequence();
 	FRayDesc Ray = PackedPayload.GetVolumetricCallableShaderInputRay();
 	float3 PathThroughput = PackedPayload.GetVolumetricCallableShaderInputThroughput();
 	float CloudDensity = PackedPayload.GetVolumetricCallableShaderInputCloudDensity();
@@ -124,7 +124,7 @@ RAY_TRACING_ENTRY_CALLABLE(PathTracingVolumetricCloudMaterialShader, FPackedPath
 		{
 			/* take stochastic steps along the ray to estimate transmittance (null scattering) */
 			/* Sample the distance to the next interaction */										
-			float RandValue = RandomSequence_GenerateSample1D(RandSeq);						
+			float RandValue = RandSeq.Get1D();
 			TMin -= log(1 - RandValue) * InvSigmaBar;
 			if (TMin >= TMax)
 			{
@@ -154,7 +154,7 @@ RAY_TRACING_ENTRY_CALLABLE(PathTracingVolumetricCloudMaterialShader, FPackedPath
 			{
 				MS = CloudMaterialGetMSParams();
 			}
-			FRISContext HitSample = InitRISContext(RandomSequence_GenerateSample1D(RandSeq));
+			FRISContext HitSample = InitRISContext(RandSeq.Get1D());
 			float3 FinalThroughput = 0;
 #endif
 
@@ -314,7 +314,7 @@ RAY_TRACING_ENTRY_CALLABLE(PathTracingVolumetricCloudMaterialShader, FPackedPath
 			}
 
 			// Interval is only crossing atmo/clouds/fog
-			float RandValue = RandomSequence_GenerateSample1D(RandSeq);
+			float RandValue = RandSeq.Get1D();
 			float Distance = SampleSpectralTransmittance(RandValue, SigmaBar, PathThroughput);
 			if (Distance < 0.0)
 			{
diff --git a/Engine/Shaders/Private/PathTracing/Volume/PathTracingHeterogeneousVolumes.ush b/Engine/Shaders/Private/PathTracing/Volume/PathTracingHeterogeneousVolumes.ush
--- a/Engine/Shaders/Private/PathTracing/Volume/PathTracingHeterogeneousVolumes.ush
+++ b/Engine/Shaders/Private/PathTracing/Volume/PathTracingHeterogeneousVolumes.ush
@@ -216,12 +216,12 @@ FVolumeTrackingResult ConcatenateVolumeTrackingResult(FVolumeTrackingResult lhs,
 	return Result;
 }
 
-FVolumeTrackingResult HeterogeneousVolumesDeltaTracking(float3 PathThroughput, float3 Origin, float3 Direction, float TMax, FMajorantData MajorantData, inout RandomSequence RandSequence)
+FVolumeTrackingResult HeterogeneousVolumesDeltaTracking(float3 PathThroughput, float3 Origin, float3 Direction, float TMax, FMajorantData MajorantData, inout FRandomSequence RandSequence)
 {
 	FVolumeTrackingResult Sample = CreateVolumeTrackingResult(0);
 	Sample.SigmaBar = max(MajorantData.Majorant, 0);
 
-	float RandValue = RandomSequence_GenerateSample1D(RandSequence);
+	float RandValue = RandSequence.Get1D();
 #define SAMPLE_ACHROMATIC 0
 #if SAMPLE_ACHROMATIC
 	float DeltaT = -log(1.0 - RandValue) / Sample.SigmaBar.x;
@@ -260,14 +260,14 @@ FVolumeTrackingResult HeterogeneousVolumesDeltaTracking(float3 PathThroughput, f
 	return Sample;
 }
 
-FVolumeTrackingResult HeterogeneousVolumesRatioTracking(float3 PathThroughput, float3 Origin, float3 Direction, float TMax, FMajorantData MajorantData, inout RandomSequence RandSequence)
+FVolumeTrackingResult HeterogeneousVolumesRatioTracking(float3 PathThroughput, float3 Origin, float3 Direction, float TMax, FMajorantData MajorantData, inout FRandomSequence RandSequence)
 {
 	FVolumeTrackingResult Sample = CreateVolumeTrackingResult(0);
 	Sample.SigmaBar = max(MajorantData.Majorant, 0);
 
 	while ((Sample.Distance < TMax))
 	{
-		float RandValue = RandomSequence_GenerateSample1D(RandSequence);
+		float RandValue = RandSequence.Get1D();
 #if SAMPLE_ACHROMATIC
 		float DeltaT = -log(1.0 - RandValue) / Sample.SigmaBar.x;
 #else
@@ -310,7 +310,7 @@ FVolumeTrackingResult HeterogeneousVolumesRatioTracking(float3 PathThroughput, f
 	return Sample;
 }
 
-FVolumeTrackingResult HeterogeneousVolumesMajorantBasedTracking(float3 PathThroughput, float3 WorldRayOrigin, float3 Direction, float TMin, float TMax, int bTransmittanceOnly, FMajorantData OverlappingMajorant, inout RandomSequence RandSequence)
+FVolumeTrackingResult HeterogeneousVolumesMajorantBasedTracking(float3 PathThroughput, float3 WorldRayOrigin, float3 Direction, float TMin, float TMax, int bTransmittanceOnly, FMajorantData OverlappingMajorant, inout FRandomSequence RandSequence)
 {
 	FVolumeTrackingResult Result = CreateVolumeTrackingResult(TMin);
 	Result.Throughput = PathThroughput;
@@ -431,14 +431,14 @@ FVolumeTrackingResult HeterogeneousVolumesMajorantBasedTracking(float3 PathThrou
 	return Result;
 }
 
-float3 HeterogeneousVolumesGetTransmittance(float3 PathThroughput, float3 Origin, float3 Direction, float TMin, float TMax, inout RandomSequence RandSequence)
+float3 HeterogeneousVolumesGetTransmittance(float3 PathThroughput, float3 Origin, float3 Direction, float TMin, float TMax, inout FRandomSequence RandSequence)
 {
 	bool bTransmittanceOnly = true;
 	FMajorantData EmptyOverlappingMajorant = CreateMajorantData(0.0f, 0.0f);
 	return HeterogeneousVolumesMajorantBasedTracking(PathThroughput, Origin, Direction, TMin, TMax, bTransmittanceOnly, EmptyOverlappingMajorant, RandSequence).Throughput;
 }
 
-FVolumeTrackingResult HeterogeneousVolumesSampleDistance(float3 PathThroughput, float3 Origin, float3 Direction, float TMin, float TMax, FMajorantData OverlappingMajorant, inout RandomSequence RandSequence)
+FVolumeTrackingResult HeterogeneousVolumesSampleDistance(float3 PathThroughput, float3 Origin, float3 Direction, float TMin, float TMax, FMajorantData OverlappingMajorant, inout FRandomSequence RandSequence)
 {
 	bool bTransmittanceOnly = false;
 	return HeterogeneousVolumesMajorantBasedTracking(PathThroughput, Origin, Direction, TMin, TMax, bTransmittanceOnly, OverlappingMajorant, RandSequence);
diff --git a/Engine/Shaders/Private/PathTracing/Volume/PathTracingVolume.ush b/Engine/Shaders/Private/PathTracing/Volume/PathTracingVolume.ush
--- a/Engine/Shaders/Private/PathTracing/Volume/PathTracingVolume.ush
+++ b/Engine/Shaders/Private/PathTracing/Volume/PathTracingVolume.ush
@@ -91,7 +91,7 @@ FVolumeShadedResult VolumeGetDensity(float3 TranslatedWorldPos, uint VolumeMask,
 		/* take stochastic steps along the ray to estimate transmittance (null scattering) */		\
 		float3 ColorChannelPdf = Throughput;														\
 		/* Sample the distance to the next interaction */											\
-		float RandValue = RandomSequence_GenerateSample1D(RandSequence);							\
+		float RandValue = RandSequence.Get1D();														\
 		float DeltaT = SampleSpectralTransmittance(RandValue.x, SigmaBar, ColorChannelPdf);			\
 		if (DeltaT < 0.0)																			\
 		{																							\
@@ -146,7 +146,7 @@ FVolumeShadedResult VolumeGetDensity(float3 TranslatedWorldPos, uint VolumeMask,
 
 // The returned value factors in the initial throughput. This is used to avoid sampling too much if throughput reaches 0.
 // Return the transmittance along a ray segment
-float3 VolumeGetTransmittance(float3 StartThroughput, float3 Origin, float3 Direction, FVolumeIntersectionInterval Interval, inout RandomSequence RandSequence)
+float3 VolumeGetTransmittance(float3 StartThroughput, float3 Origin, float3 Direction, FVolumeIntersectionInterval Interval, inout FRandomSequence RandSequence)
 {
 	float3 Throughput = StartThroughput;
 	if (Interval.VolumeMask & PATH_TRACER_VOLUME_ENABLE_ATMOSPHERE)
@@ -188,7 +188,7 @@ float3 VolumeGetTransmittance(float3 StartThroughput, float3 Origin, float3 Dire
 
 // Return the transmittance along a whole intersection list
 // Each volume type has its own range unlike the function above. This is used for shadow rays, so this version adds the CloudFactor for cloud shadows
-float3 VolumeGetTransmittance(float3 StartThroughput, float3 Origin, float3 Direction, FVolumeIntersectionList IntersectionList, inout RandomSequence RandSequence, float CloudFactor)
+float3 VolumeGetTransmittance(float3 StartThroughput, float3 Origin, float3 Direction, FVolumeIntersectionList IntersectionList, inout FRandomSequence RandSequence, float CloudFactor)
 {
 	float3 Throughput = StartThroughput;
 	if (IntersectionList.VolumeMask & PATH_TRACER_VOLUME_ENABLE_ATMOSPHERE)
@@ -240,7 +240,7 @@ float3 VolumeGetTransmittance(float3 StartThroughput, float3 Origin, float3 Dire
 	return Throughput;
 }
 
-FVolumeShadedResult VolumeGetDensityAndTransmittance(float3 Origin, float3 Direction, FVolumeIntersectionInterval Interval, uint Bounce, inout RandomSequence RandSequence, inout float3 Throughput, inout float CloudFactor)
+FVolumeShadedResult VolumeGetDensityAndTransmittance(float3 Origin, float3 Direction, FVolumeIntersectionInterval Interval, uint Bounce, inout FRandomSequence RandSequence, inout float3 Throughput, inout float CloudFactor)
 {
 	FVolumeShadedResult Result = (FVolumeShadedResult) 0;
 #if PATH_TRACER_USE_CLOUD_SHADER
@@ -288,7 +288,7 @@ FVolumeShadedResult VolumeGetDensityAndTransmittance(float3 Origin, float3 Direc
 	return Result;
 }
 
-FVolumeTrackingResult VolumeSampleDistance(float3 PathThroughput, float3 Origin, float3 Direction, FVolumeIntersectionInterval Interval, bool bIsCameraRay, inout RandomSequence RandSequence)
+FVolumeTrackingResult VolumeSampleDistance(float3 PathThroughput, float3 Origin, float3 Direction, FVolumeIntersectionInterval Interval, bool bIsCameraRay, inout FRandomSequence RandSequence)
 {
 	FVolumeDensityBounds VolumeDensityBounds = VolumeGetDensityBounds(Origin, Direction, Interval);
 	float3 SigmaBar = VolumeDensityBounds.SigmaMax;
@@ -310,7 +310,7 @@ FVolumeTrackingResult VolumeSampleDistance(float3 PathThroughput, float3 Origin,
 	else
 	{
 		// Interval is only crossing atmo/clouds/fog
-		float RandValue = RandomSequence_GenerateSample1D(RandSequence);
+		float RandValue = RandSequence.Get1D();
 		TrackingResult.Distance = SampleSpectralTransmittance(RandValue, SigmaBar, PathThroughput);
 		if (TrackingResult.Distance < 0.0)
 		{
