Ray Tracing: Update to new RandomSequence API

Fix an incorrect use of the API in the PrimaryRays pass which used the bounce index as a sample count

#rnx

diff --git a/Engine/Shaders/Private/RayTracing/GenerateSkyLightVisibilityRaysCS.usf b/Engine/Shaders/Private/RayTracing/GenerateSkyLightVisibilityRaysCS.usf
--- a/Engine/Shaders/Private/RayTracing/GenerateSkyLightVisibilityRaysCS.usf
+++ b/Engine/Shaders/Private/RayTracing/GenerateSkyLightVisibilityRaysCS.usf
@@ -26,11 +26,10 @@ void MainCS(
 
 	for (uint SampleIndex = 0; SampleIndex < SkyLight.SamplesPerPixel; ++SampleIndex)
 	{
-		RandomSequence RandSequence;
-		RandomSequence_Initialize(RandSequence, PixelCoord, SampleIndex, View.StateFrameIndex, SkyLight.SamplesPerPixel);
+		FRandomSequence RandSequence = RandomSequenceCreate(uint3(PixelCoord, View.StateFrameIndex), SampleIndex, SkyLight.SamplesPerPixel);
 
 		// Determine sky light or lambert ray
-		float2 RandSample = RandomSequence_GenerateSample2D(RandSequence);
+		float2 RandSample = RandSequence.Get2D();
 
 		float3 RayDirection;
 		float UniformPdf = 1.0 / (4 * PI);
diff --git a/Engine/Shaders/Private/RayTracing/RayTracingAmbientOcclusionRGS.usf b/Engine/Shaders/Private/RayTracing/RayTracingAmbientOcclusionRGS.usf
--- a/Engine/Shaders/Private/RayTracing/RayTracingAmbientOcclusionRGS.usf
+++ b/Engine/Shaders/Private/RayTracing/RayTracingAmbientOcclusionRGS.usf
@@ -35,7 +35,7 @@ void GenerateCosineNormalRay(
 	float3 TranslatedWorldPosition,
 	float3 WorldNormal,
 	float3 GeometricNormal,
-	inout RandomSequence RandSequence,
+	inout FRandomSequence RandSequence,
 	out float3 RayOrigin,
 	out float3 RayDirection,
 	out float RayTMin,
@@ -45,7 +45,7 @@ void GenerateCosineNormalRay(
 {
 	// Draw random variable
 	float2 BufferSize = View.BufferSizeAndInvSize.xy;
-	float2 RandSample = RandomSequence_GenerateSample2D(RandSequence);
+	float2 RandSample = RandSequence.Get2D();
 
 	// Perform cosine-hemispherical sampling and convert to world-space
 	float4 Direction_Tangent;
@@ -175,8 +175,7 @@ RAY_TRACING_ENTRY_RAYGEN(AmbientOcclusionRGS)
 		FRayDesc Ray;
 		float RayPDF;
 
-		RandomSequence RandSequence;
-		RandomSequence_Initialize(RandSequence, PixelCoord, SampleIndex, TimeSeed, SamplesPerPixelLocal);
+		FRandomSequence RandSequence = RandomSequenceCreate(uint3(PixelCoord, TimeSeed), SampleIndex, SamplesPerPixelLocal);
 
 		GenerateCosineNormalRay(PixelCoord, TranslatedWorldPosition, WorldNormal, GeometricNormal, RandSequence, Ray.Origin, Ray.Direction, Ray.TMin, Ray.TMax, RayPDF);
 		ApplyCameraRelativeDepthBias(Ray, PixelCoord, DeviceZ, WorldNormal, MaxNormalBias);
diff --git a/Engine/Shaders/Private/RayTracing/RayTracingLightingCommon.ush b/Engine/Shaders/Private/RayTracing/RayTracingLightingCommon.ush
--- a/Engine/Shaders/Private/RayTracing/RayTracingLightingCommon.ush
+++ b/Engine/Shaders/Private/RayTracing/RayTracingLightingCommon.ush
@@ -149,7 +149,7 @@ float4 SampleAreaLightDirection(
 	in float3 TranslatedWorldPosition,
 	in float3 WorldNormal,
 	in uint LightType, 
-	inout RandomSequence RandSequence )
+	inout FRandomSequence RandSequence )
 {
 	float3 ShadowRayDirection = 0.0;
 
@@ -158,7 +158,7 @@ float4 SampleAreaLightDirection(
 	float RayTMax = 0;
 	float RayPdf = 0;
 
-	float2 RandSample = RandomSequence_GenerateSample2D(RandSequence);
+	float2 RandSample = RandSequence.Get2D();
 
 	FLightShaderParameters LightParameters;
 	LightParameters.TranslatedWorldPosition = LightData.TranslatedWorldPosition;
@@ -320,7 +320,7 @@ float3 ComputeDirectLighting(
 	in RaytracingAccelerationStructure TLAS,
 	in RaytracingAccelerationStructure FarFieldTLAS,
 	inout FPackedMaterialClosestHitPayload Payload,
-	inout RandomSequence RandSequence,
+	inout FRandomSequence RandSequence,
 	in uint ShadowsType,
 	in uint ShadowsTranslucencyType,
 	float ShadowMaxNormalBias,
@@ -487,7 +487,7 @@ void AccumulateResults(
 	in float3 ViewDirection,
 	in RaytracingAccelerationStructure TLAS,
 	in RaytracingAccelerationStructure FarFieldTLAS,
-	inout RandomSequence RandSequence,
+	inout FRandomSequence RandSequence,
 	in uint2 PixelCoord,
 	in float ShadowMaxNormalBias,
 	float MaxTraceDistance,
@@ -546,7 +546,7 @@ FMaterialClosestHitPayload TraceRayAndAccumulateResults(
 	in RaytracingAccelerationStructure FarFieldTLAS,
 	in uint RayFlags,
 	in uint InstanceInclusionMask,
-	inout RandomSequence RandSequence,
+	inout FRandomSequence RandSequence,
 	in uint2 PixelCoord,
 	in float ShadowMaxNormalBias,
 	float MaxTraceDistance,
diff --git a/Engine/Shaders/Private/RayTracing/RayTracingOcclusionRGS.usf b/Engine/Shaders/Private/RayTracing/RayTracingOcclusionRGS.usf
--- a/Engine/Shaders/Private/RayTracing/RayTracingOcclusionRGS.usf
+++ b/Engine/Shaders/Private/RayTracing/RayTracingOcclusionRGS.usf
@@ -166,7 +166,7 @@ float ComputeDiffuseTransmission(
 	float3 V, float3 P, float3 N,
 	float MeanFreePath, float MaxTransmissionDistance, float Eta,
 	FLightShaderParameters LightParameters, float2 LightSample,
-	inout RandomSequence RandSequence, 
+	inout FRandomSequence RandSequence, 
 	uint RejectionRetryMax, uint2 PixelCoord)
 {
 	float TransmissionDistance = MaxTransmissionDistance;
@@ -209,7 +209,7 @@ float ComputeDiffuseTransmission(
 		}*/
 	}
 
-	float RandSample = RandomSequence_GenerateSample1D(RandSequence);
+	float RandSample = RandSequence.Get1D();
 
 	float ScatterDistance = TransmissionDistance;
 	float ScatterDistancePdf = 0.0;
@@ -223,15 +223,15 @@ float ComputeDiffuseTransmission(
 	else // TransmissionSamplingTechnique = TransmissionSampling_ConstantTrackingInfiniteDomain
 	{
 		// Use constant tracking through homogeneous media, with rejection sampling to force a scatter event
-		ScatterDistance = -log(RandSample) * MeanFreePath;
+		ScatterDistance = -log(1 - RandSample) * MeanFreePath;
 		ScatterDistancePdf = RcpMeanFreePath * exp(-ScatterDistance * RcpMeanFreePath);
 
 		// Reject scatter distances which penetrate through the medium
 		uint RetryCount = 0;
 		while (ScatterDistance > TMax && RetryCount < RejectionSamplingTrials)
 		{
-			RandSample = RandomSequence_GenerateSample1D(RandSequence);
-			ScatterDistance = -log(RandSample) * MeanFreePath;
+			RandSample = RandSequence.Get1D();
+			ScatterDistance = -log(1 - RandSample) * MeanFreePath;
 			ScatterDistancePdf = RcpMeanFreePath * exp(-ScatterDistance * RcpMeanFreePath);
 			RetryCount++;
 		}
@@ -382,23 +382,21 @@ FOcclusionResult ComputeOcclusion(
 		bApplyNormalCulling = false;
 	}
 
-	uint TimeSeed = View.StateFrameIndex;
-
 #if ENABLE_MULTIPLE_SAMPLES_PER_PIXEL
 	LOOP for (uint SampleIndex = 0; SampleIndex < LocalSamplesPerPixel; ++SampleIndex)
 #else // ENABLE_MULTIPLE_SAMPLES_PER_PIXEL
 	do if (LocalSamplesPerPixel > 0)
 #endif // ENABLE_MULTIPLE_SAMPLES_PER_PIXEL
 	{
-		RandomSequence RandSequence;
+		
 
 #if ENABLE_MULTIPLE_SAMPLES_PER_PIXEL
-		RandomSequence_Initialize(RandSequence, PixelCoord, SampleIndex, TimeSeed, LocalSamplesPerPixel);
+		FRandomSequence RandSequence  = RandomSequenceCreate(uint3(PixelCoord, View.StateFrameIndex), SampleIndex, LocalSamplesPerPixel);
 #else
-		RandomSequence_Initialize(RandSequence, PixelCoord, 0, TimeSeed, 1);
+		FRandomSequence RandSequence  = RandomSequenceCreate(uint3(PixelCoord, View.StateFrameIndex), 0, 1);
 #endif
 
-		float2 RandSample = RandomSequence_GenerateSample2D(RandSequence);
+		float2 RandSample = RandSequence.Get2D();
 
 		FRayDesc Ray = (FRayDesc)0;
 		bool bIsValidRay = GenerateOcclusionRay(
@@ -587,8 +585,7 @@ FOcclusionResult ComputeOcclusion(
 
 	if (ENABLE_TRANSMISSION && LocalSamplesPerPixel > 0 && ShadingDesc.bHasSubsurfaceProfile)
 	{
-		RandomSequence RandSequence;
-		RandomSequence_Initialize(RandSequence, PixelCoord, 0, TimeSeed, 1);
+		FRandomSequence RandSequence = RandomSequenceCreate(uint3(PixelCoord, View.StateFrameIndex), 0, 1);
 
 		// Fix light sampling to occur at the light center for all transmission computations..
 		float2 LightSample = 0.5;
diff --git a/Engine/Shaders/Private/RayTracing/RayTracingPrimaryRays.usf b/Engine/Shaders/Private/RayTracing/RayTracingPrimaryRays.usf
--- a/Engine/Shaders/Private/RayTracing/RayTracingPrimaryRays.usf
+++ b/Engine/Shaders/Private/RayTracing/RayTracingPrimaryRays.usf
@@ -134,6 +134,7 @@ RAY_TRACING_ENTRY_RAYGEN(RayTracingPrimaryRaysRGS)
 	float LastRoughness = 0.0;
 	float HitDistance = 0.0f;
 
+	FRandomSequence RandSequence = RandomSequenceCreate(uint3(PixelCoord, View.StateFrameIndex), 0, 1);
 	for (uint RefractionRayIndex = 0; RefractionRayIndex < MaxRefractionRays; ++RefractionRayIndex)
 	{
 		const bool bIsCameraRay = !bHasScattered;
@@ -149,9 +150,6 @@ RAY_TRACING_ENTRY_RAYGEN(RayTracingPrimaryRaysRGS)
 		float3 PathVertexRadiance = float3(0, 0, 0);
 		float MaxTraceDistance = 100000.0f;
 
-		RandomSequence RandSequence;
-
-		RandomSequence_Initialize(RandSequence, PixelCoord, RefractionRayIndex, View.StateFrameIndex, MaxRefractionRays);
 
 		FMaterialClosestHitPayload Payload = TraceRayAndAccumulateResults(
 			Ray,
@@ -290,7 +288,7 @@ RAY_TRACING_ENTRY_RAYGEN(RayTracingPrimaryRaysRGS)
 		if (PayloadRoughness < TranslucencyMaxRoughness)
 		{
 			// Trace reflection ray 
-			float2 RandSample = RandomSequence_GenerateSample2D(RandSequence);
+			float2 RandSample = RandSequence.Get2D();
 
 			FRayDesc ReflectionRay;
 			ReflectionRay.TMin = 0.01;
diff --git a/Engine/Shaders/Private/RayTracing/RayTracingSkyLightEvaluation.ush b/Engine/Shaders/Private/RayTracing/RayTracingSkyLightEvaluation.ush
--- a/Engine/Shaders/Private/RayTracing/RayTracingSkyLightEvaluation.ush
+++ b/Engine/Shaders/Private/RayTracing/RayTracingSkyLightEvaluation.ush
@@ -87,11 +87,10 @@ void SkyLightEvaluate(
 		}
 		else
 		{
-			RandomSequence RandSequence;
-			RandomSequence_Initialize(RandSequence, PixelCoord, SampleIndex, View.StateFrameIndex, SamplesPerPixel);
+			FRandomSequence RandSequence = RandomSequenceCreate(uint3(PixelCoord, View.StateFrameIndex), SampleIndex, SamplesPerPixel);
 
 			// Determine sky light or lambert ray
-			float2 RandSample = RandomSequence_GenerateSample2D(RandSequence);
+			float2 RandSample = RandSequence.Get2D();
 
 			// Generate a visibility ray for the current sample
 			float SkyLightPdf = 0;
@@ -174,9 +173,8 @@ void SkyLightEvaluate(
 		if (GBufferData.ShadingModelID != SHADINGMODELID_HAIR)
 		{
 			const float HairOcclusionThreshold = 1;
-			RandomSequence RandSequence;
 			uint LinearIndex = CalcLinearIndex(PixelCoord);
-			RandomSequence_Initialize(RandSequence, LinearIndex, View.StateFrameIndex);
+			FRandomSequence RandSequence = RandomSequenceCreate(LinearIndex, View.StateFrameIndex); // TODO: should include the SampleIndex here
 			MinimalPayload.HitT = TraverseHair(PixelCoord, RandSequence, Ray.Origin, Ray.Direction, MinimalPayload.HitT, VirtualVoxel.Raytracing_SkyOcclusionThreshold);
 		}
 		#endif
