Update various call sites to newer random sequence API

#rnx

diff --git a/Engine/Shaders/Private/HairStrands/HairStrandsRaytracing.ush b/Engine/Shaders/Private/HairStrands/HairStrandsRaytracing.ush
--- a/Engine/Shaders/Private/HairStrands/HairStrandsRaytracing.ush
+++ b/Engine/Shaders/Private/HairStrands/HairStrandsRaytracing.ush
@@ -85,7 +85,7 @@ float TraverseHair(uint2 PixelCoord, float3 RayOrigin, float3 RayDirection, floa
 	return InternalTraverseHair(PixelCoord, RayOrigin, RayDirection, InMinT, InOcclusionThreshold, VoxelRandom, RayRandom);
 }
 
-float TraverseHair(uint2 PixelCoord, inout RandomSequence RandSequence, float3 RayOrigin, float3 RayDirection, float InMinT, float InOcclusionThreshold)
+float TraverseHair(uint2 PixelCoord, inout FRandomSequence RandSequence, float3 RayOrigin, float3 RayDirection, float InMinT, float InOcclusionThreshold)
 {
 	float HitT = InMinT;
 
@@ -94,7 +94,7 @@ float TraverseHair(uint2 PixelCoord, inout RandomSequence RandSequence, float3 R
 	float  RayRandom   = 1.0f;
 	if (VirtualVoxel.JitterMode == 1)
 	{
-		const float4 Rnd = RandomSequence_GenerateSample4D(RandSequence);
+		const float4 Rnd = RandSequence.Get4D();
 		VoxelRandom = Rnd.xyz;
 		RayRandom = Rnd.w;
 	}
diff --git a/Engine/Shaders/Private/ScreenSpaceDenoise/SSDSignalFramework.ush b/Engine/Shaders/Private/ScreenSpaceDenoise/SSDSignalFramework.ush
--- a/Engine/Shaders/Private/ScreenSpaceDenoise/SSDSignalFramework.ush
+++ b/Engine/Shaders/Private/ScreenSpaceDenoise/SSDSignalFramework.ush
@@ -685,11 +685,7 @@ float GetSamplePenumbra(FSSDSignalSample Sample)
 #if COMPILE_SIGNAL_COLOR_SH && COMPILE_SIGNAL_COLOR
 FSSDSphericalHarmonicRGB ComputeSampleColorSH(FSSDSampleSceneInfos SampleSceneMetadata, FSSDSignalSample Sample, uint2 NeighborPixelCoord)
 {
-	const uint SamplesPerPixel = 1;
-	RandomSequence RandSequence;
-	RandomSequence_Initialize(RandSequence, NeighborPixelCoord, /* SampleIndex */ 0,  FrameIndex, SamplesPerPixel);
-
-	float2 RandSample = RandomSequence_GenerateSample2D(RandSequence);
+	float2 RandSample = RandomSequenceCreate(uint3(NeighborPixelCoord, FrameIndex), 0, 1).Get2D();
 
 	float4 TangentDirection = CosineSampleHemisphere(RandSample);
 	float3 WorldDirection = TangentToWorld(TangentDirection.xyz, GetWorldNormal(SampleSceneMetadata));
diff --git a/Engine/Shaders/Private/ShadingEnergyConservationTable.usf b/Engine/Shaders/Private/ShadingEnergyConservationTable.usf
--- a/Engine/Shaders/Private/ShadingEnergyConservationTable.usf
+++ b/Engine/Shaders/Private/ShadingEnergyConservationTable.usf
@@ -47,15 +47,13 @@ void BuildEnergyTableCS(uint3 DispatchThreadId : SV_DispatchThreadID)
 	float2 Alpha; GetGGXBasis(Roughness, float3(0, 0, 1), Alpha);
 	float3 V = float3(sqrt(1 - NoV * NoV), 0, NoV);
 
-	RandomSequence RandSequence;
 	float2 Result = 0;
 	// use plenty of samples -- this is computed on startup once and at fairly low resolution
 	LOOP
 	for (int Sample = 0; Sample < NumSamples; Sample++)
 	{
 		// draw a random sample
-		RandomSequence_Initialize(RandSequence, DispatchThreadId.xy, Sample, DispatchThreadId.z, NumSamples);
-		float3 RandSample = RandomSequence_GenerateSample3D(RandSequence);
+		float3 RandSample = RandomSequenceCreate(DispatchThreadId, Sample, NumSamples).Get3D();
 
 		float2 Weight = 0;
 
diff --git a/Engine/Shaders/Private/VirtualShadowMaps/VirtualShadowMapSMRTCommon.ush b/Engine/Shaders/Private/VirtualShadowMaps/VirtualShadowMapSMRTCommon.ush
--- a/Engine/Shaders/Private/VirtualShadowMaps/VirtualShadowMapSMRTCommon.ush
+++ b/Engine/Shaders/Private/VirtualShadowMaps/VirtualShadowMapSMRTCommon.ush
@@ -59,9 +59,7 @@ float4 VirtualShadowMapGetRandomSample(
 	}
 #else
 	{
-		RandomSequence RandSequence;
-		RandomSequence_Initialize(RandSequence, PixelPos, SampleIndex, TimeIndex, MaxSampleCount);
-		return RandomSequence_GenerateSample4D(RandSequence);
+		return RandomSequenceCreate(uint3(PixelPos, TimeIndex), SampleIndex, MaxSampleCount).Get4D();
 	}
 #endif
 }
