Heterogeneous Volumes: Update to newer RandomSequence API

Also fix a few minor cases where log(0) can be taken during unbiased ray marching.

#rnx

diff --git a/Engine/Shaders/Private/HeterogeneousVolumes/HeterogeneousVolumesAmbientOcclusionPipeline.usf b/Engine/Shaders/Private/HeterogeneousVolumes/HeterogeneousVolumesAmbientOcclusionPipeline.usf
--- a/Engine/Shaders/Private/HeterogeneousVolumes/HeterogeneousVolumesAmbientOcclusionPipeline.usf
+++ b/Engine/Shaders/Private/HeterogeneousVolumes/HeterogeneousVolumesAmbientOcclusionPipeline.usf
@@ -398,10 +398,8 @@ void RenderExistenceMaskWithLiveShadingCS(
 	if (IsOfflineRender())
 	{
 		float3 FilterWidth = 1.0;
-		RandomSequence RandSequence;
 		int LinearIndex = VoxelIndex.x * (VoxelResolution.y * VoxelResolution.z) + VoxelIndex.y * VoxelResolution.z + VoxelIndex.z;
-		RandomSequence_Initialize(RandSequence, LinearIndex, View.StateFrameIndex);
-		float Rand = RandomSequence_GenerateSample1D(RandSequence);
+		float Rand = RandomSequenceCreate(LinearIndex, View.StateFrameIndex).Get1D();
 		int StochasticFilteringMode = GetStochasticFilteringMode();
 		SampleContext = CreateVolumeSampleContext(LocalPosition, WorldPosition, FilterWidth, MipLevel, Rand, StochasticFilteringMode);
 	}
diff --git a/Engine/Shaders/Private/HeterogeneousVolumes/HeterogeneousVolumesLiveShadingPipeline.usf b/Engine/Shaders/Private/HeterogeneousVolumes/HeterogeneousVolumesLiveShadingPipeline.usf
--- a/Engine/Shaders/Private/HeterogeneousVolumes/HeterogeneousVolumesLiveShadingPipeline.usf
+++ b/Engine/Shaders/Private/HeterogeneousVolumes/HeterogeneousVolumesLiveShadingPipeline.usf
@@ -469,10 +469,8 @@ void RenderLightingCacheWithLiveShadingCS(
 	if (IsOfflineRender())
 	{
 		float3 FilterWidth = 1.0;
-		RandomSequence RandSequence;
 		int LinearIndex = VoxelIndex.x * (GetLightingCacheResolution().y * GetLightingCacheResolution().z) + VoxelIndex.y * GetLightingCacheResolution().z + VoxelIndex.z;
-		RandomSequence_Initialize(RandSequence, LinearIndex, View.StateFrameIndex);
-		float Rand = RandomSequence_GenerateSample1D(RandSequence);
+		float Rand = RandomSequenceCreate(LinearIndex, View.StateFrameIndex).Get1D();
 		int StochasticFilteringMode = GetStochasticFilteringMode();
 		SampleContext = CreateVolumeSampleContext(LocalRayOrigin, WorldRayOrigin, FilterWidth, MipLevel, Rand, StochasticFilteringMode);
 	}
@@ -816,9 +814,7 @@ void RenderSingleScatteringWithLiveShading(
 			bApplyShadowTransmittance
 		);
 
-		RandomSequence RandSequence;
-		RandomSequence_Initialize(RandSequence, LinearIndex, View.StateFrameIndex);
-		RayMarchingContext.RandSequence = RandSequence;
+		RayMarchingContext.RandSequence = RandomSequenceCreate(LinearIndex, View.StateFrameIndex);
 
 		const FDeferredLightData LightData = InitDeferredLightFromUniforms(LightType, VolumetricScatteringIntensity);
 		uint StepCount = CalcStepCount(RayMarchingContext);
diff --git a/Engine/Shaders/Private/HeterogeneousVolumes/HeterogeneousVolumesLiveShadingShadows.usf b/Engine/Shaders/Private/HeterogeneousVolumes/HeterogeneousVolumesLiveShadingShadows.usf
--- a/Engine/Shaders/Private/HeterogeneousVolumes/HeterogeneousVolumesLiveShadingShadows.usf
+++ b/Engine/Shaders/Private/HeterogeneousVolumes/HeterogeneousVolumesLiveShadingShadows.usf
@@ -493,9 +493,7 @@ void RenderVolumetricShadowMapForLightWithLiveShadingCS(
 			bApplyShadowTransmittance
 		);
 
-		RandomSequence RandSequence;
-		RandomSequence_Initialize(RandSequence, LinearIndex, View.StateFrameIndex);
-		RayMarchingContext.RandSequence = RandSequence;
+		RayMarchingContext.RandSequence = RandomSequenceCreate(LinearIndex, View.StateFrameIndex);
 		uint StepCount = CalcStepCount(RayMarchingContext);
 
 		float3 ExpectedExtinction = 0.0;
diff --git a/Engine/Shaders/Private/HeterogeneousVolumes/HeterogeneousVolumesRayMarchingTypes.ush b/Engine/Shaders/Private/HeterogeneousVolumes/HeterogeneousVolumesRayMarchingTypes.ush
--- a/Engine/Shaders/Private/HeterogeneousVolumes/HeterogeneousVolumesRayMarchingTypes.ush
+++ b/Engine/Shaders/Private/HeterogeneousVolumes/HeterogeneousVolumesRayMarchingTypes.ush
@@ -32,7 +32,7 @@ struct FRayMarchingContext
 	int bApplyShadowTransmittance;
 
 	// TEMPORARY STORAGE!!
-	RandomSequence RandSequence;
+	FRandomSequence RandSequence;
 	float LocalScatterT0;
 	float LocalScatterT1;
 	float3 Transmittance0;
diff --git a/Engine/Shaders/Private/HeterogeneousVolumes/HeterogeneousVolumesRayMarchingUtils.ush b/Engine/Shaders/Private/HeterogeneousVolumes/HeterogeneousVolumesRayMarchingUtils.ush
--- a/Engine/Shaders/Private/HeterogeneousVolumes/HeterogeneousVolumesRayMarchingUtils.ush
+++ b/Engine/Shaders/Private/HeterogeneousVolumes/HeterogeneousVolumesRayMarchingUtils.ush
@@ -429,7 +429,7 @@ void RayMarchSingleScattering(
 		if (IsOfflineRender())
 		{
 			float3 FilterWidth = 1.0;
-			float Rand = RandomSequence_GenerateSample1D(RayMarchingContext.RandSequence);
+			float Rand = RayMarchingContext.RandSequence.Get1D();
 			int StochasticFilteringMode = GetStochasticFilteringMode();
 			SampleContext = CreateVolumeSampleContext(LocalPosition, WorldPosition, FilterWidth, RayMarchingContext.MipLevel, Rand, StochasticFilteringMode);
 		}
diff --git a/Engine/Shaders/Private/HeterogeneousVolumes/HeterogeneousVolumesVoxelGridRendering.usf b/Engine/Shaders/Private/HeterogeneousVolumes/HeterogeneousVolumesVoxelGridRendering.usf
--- a/Engine/Shaders/Private/HeterogeneousVolumes/HeterogeneousVolumesVoxelGridRendering.usf
+++ b/Engine/Shaders/Private/HeterogeneousVolumes/HeterogeneousVolumesVoxelGridRendering.usf
@@ -238,7 +238,7 @@ float MaxComponent(float3 Value)
 	return max(Value.x, max(Value.y, Value.z));
 }
 
-FTrackingResult DeltaTracking(float3 WorldRayOrigin, float3 WorldRayDirection, float TMin, float TMax, float SigmaBar, inout RandomSequence RandSequence)
+FTrackingResult DeltaTracking(float3 WorldRayOrigin, float3 WorldRayDirection, float TMin, float TMax, float SigmaBar, inout FRandomSequence RandSequence)
 {
 	FTrackingResult Result = CreateTrackingResult(TMin);
 
@@ -246,8 +246,8 @@ FTrackingResult DeltaTracking(float3 WorldRayOrigin, float3 WorldRayDirection, f
 	while (Result.Distance < TMax)
 	{
 		// Pick a distance, based on SigmaBar
-		float RandValue = RandomSequence_GenerateSample1D(RandSequence);
-		HitT = -log(RandValue) / SigmaBar;
+		float RandValue = RandSequence.Get1D();
+		HitT = -log(1 - RandValue) / SigmaBar;
 		Result.Distance += HitT;
 
 		// Evaluate SigmaT at that location
@@ -260,7 +260,7 @@ FTrackingResult DeltaTracking(float3 WorldRayOrigin, float3 WorldRayDirection, f
 
 			//float AcceptanceProbability = MaxComponent(SigmaT) / SigmaBar;
 			float AcceptanceProbability = min(SigmaT / SigmaBar, 1);
-			if (RandomSequence_GenerateSample1D(RandSequence) < AcceptanceProbability)
+			if (RandSequence.Get1D() < AcceptanceProbability)
 			{
 				Result.Throughput = 0;
 				Result.Radiance = GetEmission(OrthoGridUniformBuffer.EmissionGridBuffer[LinearVoxelPos]) / Luminance(SigmaT);
@@ -274,7 +274,7 @@ FTrackingResult DeltaTracking(float3 WorldRayOrigin, float3 WorldRayDirection, f
 	return Result;
 }
 
-FTrackingResult SpectralWeightedDeltaTracking(float3 WorldRayOrigin, float3 WorldRayDirection, float TMin, float TMax, float SigmaBar, int CollisionMode, inout RandomSequence RandSequence)
+FTrackingResult SpectralWeightedDeltaTracking(float3 WorldRayOrigin, float3 WorldRayDirection, float TMin, float TMax, float SigmaBar, int CollisionMode, inout FRandomSequence RandSequence)
 {
 	FTrackingResult Result = CreateTrackingResult(TMin);
 
@@ -283,8 +283,8 @@ FTrackingResult SpectralWeightedDeltaTracking(float3 WorldRayOrigin, float3 Worl
 	while (Result.Distance < TMax && all(Result.Throughput > Epsilon))
 	{
 		// Distance-sampling, based on SigmaBar
-		float RandValue = RandomSequence_GenerateSample1D(RandSequence);
-		HitT = -log(RandValue) / SigmaBar;
+		float RandValue = RandSequence.Get1D();
+		HitT = -log(1 - RandValue) / SigmaBar;
 		Result.Distance += HitT;
 
 		if (Result.Distance < TMax)
@@ -327,7 +327,7 @@ FTrackingResult SpectralWeightedDeltaTracking(float3 WorldRayOrigin, float3 Worl
 				AbsorptionProbability + ScatteringProbability,
 				AbsorptionProbability + ScatteringProbability + NullCollisionProbability
 			);
-			float CollisionSample = RandomSequence_GenerateSample1D(RandSequence) * CollisionCdf.z;
+			float CollisionSample = RandSequence.Get1D() * CollisionCdf.z;
 
 			// Normalize probabilities
 			AbsorptionProbability *= rcp(CollisionCdf.z);
@@ -363,7 +363,7 @@ FTrackingResult SpectralWeightedDeltaTracking(float3 WorldRayOrigin, float3 Worl
 	return Result;
 }
 
-FTrackingResult RatioTracking(float3 WorldRayOrigin, float3 WorldRayDirection, float TMin, float TMax, float SigmaBar, inout RandomSequence RandSequence)
+FTrackingResult RatioTracking(float3 WorldRayOrigin, float3 WorldRayDirection, float TMin, float TMax, float SigmaBar, inout FRandomSequence RandSequence)
 {
 	FTrackingResult Result = CreateTrackingResult(TMin);
 
@@ -372,8 +372,8 @@ FTrackingResult RatioTracking(float3 WorldRayOrigin, float3 WorldRayDirection, f
 	while (Result.Distance < TMax && all(Result.Throughput > Epsilon))
 	{
 		// Pick a distance, based on SigmaBar
-		float RandValue = RandomSequence_GenerateSample1D(RandSequence);
-		HitT = -log(RandValue) / SigmaBar;
+		float RandValue = RandSequence.Get1D();
+		HitT = -log(1 - RandValue) / SigmaBar;
 		Result.Distance += HitT;
 
 		// Evaluate SigmaT at that location
@@ -385,7 +385,7 @@ FTrackingResult RatioTracking(float3 WorldRayOrigin, float3 WorldRayDirection, f
 			float SigmaT = Luminance(Extinction);
 
 			float AcceptanceProbability = min(SigmaT / SigmaBar, 1);
-			if (RandomSequence_GenerateSample1D(RandSequence) < AcceptanceProbability)
+			if (RandSequence.Get1D() < AcceptanceProbability)
 			{
 				//Result.Emission = GetEmission(OrthoGridUniformBuffer.EmissionGridBuffer[LinearVoxelPos]) / Luminance(SigmaT);
 			}
@@ -398,7 +398,7 @@ FTrackingResult RatioTracking(float3 WorldRayOrigin, float3 WorldRayDirection, f
 	return Result;
 }
 
-FTrackingResult SpectralRatioTracking(float3 WorldRayOrigin, float3 WorldRayDirection, float TMin, float TMax, float SigmaBar, inout RandomSequence RandSequence)
+FTrackingResult SpectralRatioTracking(float3 WorldRayOrigin, float3 WorldRayDirection, float TMin, float TMax, float SigmaBar, inout FRandomSequence RandSequence)
 {
 	FTrackingResult Result = CreateTrackingResult(TMin);
 
@@ -407,8 +407,8 @@ FTrackingResult SpectralRatioTracking(float3 WorldRayOrigin, float3 WorldRayDire
 	while (Result.Distance < TMax && all(Result.Throughput > Epsilon))
 	{
 		// Pick a distance, based on SigmaBar
-		float RandValue = RandomSequence_GenerateSample1D(RandSequence);
-		HitT = -log(RandValue) / SigmaBar;
+		float RandValue = RandSequence.Get1D();
+		HitT = -log(1 - RandValue) / SigmaBar;
 		Result.Distance += HitT;
 
 		// Evaluate SigmaT at that location
@@ -420,7 +420,7 @@ FTrackingResult SpectralRatioTracking(float3 WorldRayOrigin, float3 WorldRayDire
 
 			// Pick a component, based on proportional weight
 			float3 SpectralCdf = float3(Extinction.x, Extinction.x + Extinction.y, Extinction.x + Extinction.y + Extinction.z);
-			float RandSpectra = RandomSequence_GenerateSample1D(RandSequence) * SpectralCdf.z;
+			float RandSpectra = RandSequence.Get1D() * SpectralCdf.z;
 			float SigmaT = (RandSpectra < SpectralCdf.x) ? Extinction.x : ((RandSpectra < SpectralCdf.y) ? Extinction.y : Extinction.z);
 
 			if (SigmaT > 0)
@@ -428,7 +428,7 @@ FTrackingResult SpectralRatioTracking(float3 WorldRayOrigin, float3 WorldRayDire
 				float3 SpectralWeight = Extinction / SigmaT;
 				float AcceptanceProbability = min(SigmaT / SigmaBar, 1);
 				float3 SpectralAcceptanceProbability = min(SpectralWeight * SigmaT / SigmaBar, 1);
-				if (RandomSequence_GenerateSample1D(RandSequence) < AcceptanceProbability)
+				if (RandSequence.Get1D() < AcceptanceProbability)
 				{
 					//Result.Emission = GetEmission(OrthoGridUniformBuffer.EmissionGridBuffer[LinearVoxelPos]) / Luminance(SigmaT);
 				}
@@ -442,7 +442,7 @@ FTrackingResult SpectralRatioTracking(float3 WorldRayOrigin, float3 WorldRayDire
 	return Result;
 }
 
-FTrackingResult DDABasedMajorantMarch(float3 WorldRayOrigin, float3 Direction, float TMin, float TMax, int DistanceSamplingMethod, int CollisionSamplingMethod, inout RandomSequence RandSequence)
+FTrackingResult DDABasedMajorantMarch(float3 WorldRayOrigin, float3 Direction, float TMin, float TMax, int DistanceSamplingMethod, int CollisionSamplingMethod, inout FRandomSequence RandSequence)
 {
 	FTrackingResult TrackingResult = CreateTrackingResult(TMin);
 
@@ -589,10 +589,9 @@ void RenderTransmittanceWithVoxelGridCS(
 	float RayTMin = 0.0;
 	float RayTMax = length(RayDirection);
 	RayDirection /= RayTMax;
-
-	RandomSequence RandSequence;
-	RandomSequence_Initialize(RandSequence, LinearIndex, View.StateFrameIndex);
-	float Jitter = RandomSequence_GenerateSample1D(RandSequence);
+	
+	FRandomSequence RandSequence = RandomSequenceCreate(LinearIndex, View.StateFrameIndex);
+	float Jitter = RandSequence.Get1D();
 
 	float3 Emission = 0;
 	float3 Transmittance = 1;
@@ -634,7 +633,7 @@ void RenderTransmittanceWithVoxelGridCS(
 		int SPPCount = 1;
 		for (int Index = 0; Index < SPPCount; ++Index)
 		{
-			RandomSequence_Initialize(RandSequence, LinearIndex, View.StateFrameIndex * SPPCount + Index);
+			RandSequence = RandomSequenceCreate(LinearIndex, View.StateFrameIndex * SPPCount + Index);
 #if DEBUG_MODE == DEBUG_MODE_TRACKING_BASED_DELTA_TRACKING
 			FTrackingResult Sample = SpectralWeightedDeltaTracking(RayOrigin, RayDirection, ClipHitT.x, ClipHitT.y, SigmaBar, COLLISION_MODE_AVG_COMPONENT, RandSequence);
 #elif DEBUG_MODE == DEBUG_MODE_TRACKING_BASED_RATIO_TRACKING
