Copyright boilerplate build fix

diff --git a/Engine/Shaders/Private/SMAA/SMAA.ush b/Engine/Shaders/Private/SMAA/SMAA.ush
--- a/Engine/Shaders/Private/SMAA/SMAA.ush
+++ b/Engine/Shaders/Private/SMAA/SMAA.ush
@@ -1,3 +1,5 @@
+// Copyright Epic Games, Inc. All Rights Reserved.
+
 /**
  * Copyright (C) 2013 [REDACTED_PERSON] ([REDACTED_EMAIL])
  * Copyright (C) 2013 Jose I. Echevarria ([REDACTED_EMAIL])
