Fix PSO compilation failure when using stereo input.

diff --git a/Engine/Shaders/Private/PostProcessVisualizeHDR.usf b/Engine/Shaders/Private/PostProcessVisualizeHDR.usf
--- a/Engine/Shaders/Private/PostProcessVisualizeHDR.usf
+++ b/Engine/Shaders/Private/PostProcessVisualizeHDR.usf
@@ -187,7 +187,7 @@ void PrintMediumFloatWithSign(int2 PixelPos, inout float3 OutColor, float3 FontC
 }
 
 //
-void MainPS(noperspective float4 UVAndScreenPos : TEXCOORD0, float4 SvPosition : SV_POSITION, out float4 OutColor : SV_Target0)
+void MainPS(noperspective float4 UVAndScreenPos : TEXCOORD0, FStereoPSInput StereoInput, float4 SvPosition : SV_POSITION, out float4 OutColor : SV_Target0)
 {
 	float2 UV = UVAndScreenPos.xy;
 	int2 PixelPos = (int2)SvPosition.xy;
