Fix shader compilation with hair card factory.

diff --git a/Engine/Shaders/Private/HairStrands/HairCardsVertexFactory.ush b/Engine/Shaders/Private/HairStrands/HairCardsVertexFactory.ush
--- a/Engine/Shaders/Private/HairStrands/HairCardsVertexFactory.ush
+++ b/Engine/Shaders/Private/HairStrands/HairCardsVertexFactory.ush
@@ -446,19 +446,20 @@ FVertexFactoryInterpolantsVSToPS VertexFactoryGetInterpolantsVSToPS(FVertexFacto
 	GetRootUVsAtlasUVs(Input, Interpolants.HairPrimitiveUV, Interpolants.HairPrimitiveRootUV);
 
 	SetPrimitiveId(Interpolants, Intermediates.SceneData.PrimitiveId);
-#if NUM_TEX_COORD_INTERPOLATORS 
-	float2 TexCoords[NUM_TEX_COORD_INTERPOLATORS];
+
+#if NUM_TEX_COORD_INTERPOLATORS && NUM_MATERIAL_TEXCOORDS_VERTEX
+	float2 TexCoords[NUM_MATERIAL_TEXCOORDS_VERTEX];
 	GetTexCoords(Input, TexCoords);
 	UNROLL
-	for (int CoordinateIndex = 0; CoordinateIndex < NUM_TEX_COORD_INTERPOLATORS; CoordinateIndex++)
+	for (int CoordinateIndex = 0; CoordinateIndex < min(NUM_MATERIAL_TEXCOORDS_VERTEX, NUM_TEX_COORD_INTERPOLATORS); CoordinateIndex++)
 	{
 		SetUV(Interpolants, CoordinateIndex, TexCoords[CoordinateIndex]);
 	}
 #endif
 
-	#if NUM_TEX_COORD_INTERPOLATORS 
+#if NUM_TEX_COORD_INTERPOLATORS 
 	SetUV(Interpolants, 0, Interpolants.HairPrimitiveUV);
-	#endif
+#endif
 
 	return Interpolants;
 }
