Fix minor shader warning (shader does not write to all fields of target).

In this case the target has a 2 channel format, so change SV_Target0 to float2.

#rnx

diff --git a/Engine/Shaders/Private/SMAA/SMAAEdgeDetectionShader.usf b/Engine/Shaders/Private/SMAA/SMAAEdgeDetectionShader.usf
--- a/Engine/Shaders/Private/SMAA/SMAAEdgeDetectionShader.usf
+++ b/Engine/Shaders/Private/SMAA/SMAAEdgeDetectionShader.usf
@@ -35,7 +35,7 @@ struct FSMAAEdgeDetectionScreenVertexOutput
 
 	void MainPS(
 		FSMAAEdgeDetectionScreenVertexOutput Input,
-		out float4 OutColor : SV_Target0)
+		out float2 OutColor : SV_Target0)
 	{
 		OutColor.rg = SMAAColorEdgeDetectionPS(Input.UV, Input.Offset, Input_Texture);
 	}
