Path Tracer: Restore correct sampler

Was flipped to lattice sampling in previous commit accidentally

#rnx

diff --git a/Engine/Shaders/Private/PathTracing/PathTracingCore.ush b/Engine/Shaders/Private/PathTracing/PathTracingCore.ush
--- a/Engine/Shaders/Private/PathTracing/PathTracingCore.ush
+++ b/Engine/Shaders/Private/PathTracing/PathTracingCore.ush
@@ -3,7 +3,6 @@
 #pragma once
 
 #define RANDSEQ_UNROLL_SOBOL PATH_TRACER_USE_COMPACTION // unrolling sobol sampler is best when we can guarantee high occupancy
-#define RANDSEQ RANDSEQ_LATTICE
 
 #define USE_PATH_TRACING_LIGHT_GRID	1
 #define USE_RAY_TRACING_DECAL_GRID	1
