Path Tracer: Fix broken holdouts

This feature was accidentally broken when implementing volumetric lighting models since this re-used the UNLIT shading model. For holdouts, we need to make sure to zero out base color and specular accordingly to prevent lighting contributions on camera ray hits.

[RN] Fix for holdouts not removing lighting correctly in the path tracer

diff --git a/Engine/Shaders/Private/PathTracing/PathTracingMaterialHitShader.usf b/Engine/Shaders/Private/PathTracing/PathTracingMaterialHitShader.usf
--- a/Engine/Shaders/Private/PathTracing/PathTracingMaterialHitShader.usf
+++ b/Engine/Shaders/Private/PathTracing/PathTracingMaterialHitShader.usf
@@ -1249,8 +1249,11 @@ RAY_TRACING_ENTRY_CLOSEST_HIT(PathTracingMaterialCHS,
 			}
 
 			Payload.SetHoldout();
+			// zero-out the material contribution (including base-color and specular since UNLIT is used for volumetric lighting modes - see above)
 			Payload.ShadingModelID = SHADINGMODELID_UNLIT;
 			Payload.Radiance = 0;
+			Payload.Anisotropy = 0;
+			Payload.SetBaseColor(0.0);
 		}
 	}
 #endif
