Substrate - reworked shading quality to be more intuitive to tweak.

#tests PC editor with cvar test.
[RN] Reworked Substrate shading quality cvars to be 0=approximation, 1= more accurate. (r.Substrate.ShadingQuality and r.Substrate.SheenQuality)

diff --git a/Engine/Shaders/Private/Substrate/Substrate.ush b/Engine/Shaders/Private/Substrate/Substrate.ush
--- a/Engine/Shaders/Private/Substrate/Substrate.ush
+++ b/Engine/Shaders/Private/Substrate/Substrate.ush
@@ -51,8 +51,8 @@
 #endif
 
 // Substrate SUBSTRATE_SHADING_QUALITY define the lighting accuarcy/approximation
+// 0: approximation lighting
 // 1: accurante lighting
-// 2: approximation lighting
 #ifndef SUBSTRATE_SHADING_QUALITY
 #define SUBSTRATE_SHADING_QUALITY 1
 #endif
@@ -3713,7 +3713,7 @@ float3 sqrtFast3(float3 In) { return float3(sqrtFast(In.x), sqrtFast(In.y), sqrt
 
 half3 SubstrateLinearToSrgb(half3 In)
 {
-#if SUBSTRATE_SHADING_QUALITY <= 1
+#if SUBSTRATE_SHADING_QUALITY >= 1
 	return LinearToSrgb(In);
 #else
 	return sqrtFast3(In); // Simple gamma 2.0
@@ -3722,7 +3722,7 @@ half3 SubstrateLinearToSrgb(half3 In)
 
 half3 SubstrateSrgbToLinear(half3 In)
 {
-#if SUBSTRATE_SHADING_QUALITY <= 1
+#if SUBSTRATE_SHADING_QUALITY >= 1
 	return sRGBToLinear(In);
 #else
 	return In * In; // Simple gamma 2.0
diff --git a/Engine/Shaders/Private/Substrate/SubstrateEvaluation.ush b/Engine/Shaders/Private/Substrate/SubstrateEvaluation.ush
--- a/Engine/Shaders/Private/Substrate/SubstrateEvaluation.ush
+++ b/Engine/Shaders/Private/Substrate/SubstrateEvaluation.ush
@@ -268,7 +268,7 @@ half Substrate_D_GGX_Aniso(float ax, float ay, float NoH, float XoH, float YoH)
 
 half Substrate_Vis_GGX(half Roughness, half a2, half NoV, half NoL)
 {
-#if SUBSTRATE_SHADING_QUALITY > 1
+#if SUBSTRATE_SHADING_QUALITY == 0
 	return Vis_SmithJointApprox(a2, NoV, NoL);
 #else
 	return Vis_SmithJoint(a2, NoV, NoL);
@@ -791,13 +791,13 @@ FSubstrateEvaluateResult SubstrateEvaluateBSDFCommon(FSubstrateBSDFContext BSDFC
 
 			// Specular to Diffuse energy preservation
 			{
-				#if MATERIAL_ROUGHDIFFUSE && SUBSTRATE_SHADING_QUALITY <= 1	
+				#if MATERIAL_ROUGHDIFFUSE && SUBSTRATE_SHADING_QUALITY == 1	
 				// We remove the Chan model baked transmittance specular directional albedo assuming F=0.04.
 				// Then we apply the transmittance according to the current surface setup.
 				// To make this a safe operation:
 				//	- We clamp SpecularTChan to a threshold to avoid div by 0 when the material is a metal. 
 				//	- When the material tends to be metal, DirectionalAlbedo_SpecularTransmission will naturally tend to be 0
-				// Bypass this correction when SUBSTRATE_SHADING_QUALITY > 1 (i.e., med/low quality), for trading accuracy for speed
+				// Bypass this correction when SUBSTRATE_SHADING_QUALITY == 0 (i.e., med/low quality), for trading accuracy for speed
 				// We also do not apply that correction when rendering SIMPLE_VOLUME slabs:
 				//	- In this case there is no specular transmission baked in the participating media scattering model we use for coating layer.
 				//	- And specular throughput needs to be applied on the surface transmittance too to be energy conservative.
diff --git a/Engine/Source/Runtime/Engine/Private/ShaderCompiler/ShaderCompiler.cpp b/Engine/Source/Runtime/Engine/Private/ShaderCompiler/ShaderCompiler.cpp
--- a/Engine/Source/Runtime/Engine/Private/ShaderCompiler/ShaderCompiler.cpp
+++ b/Engine/Source/Runtime/Engine/Private/ShaderCompiler/ShaderCompiler.cpp
@@ -3876,14 +3876,14 @@ void GlobalBeginCompileShader(
 
 		if (bSubstrate)
 		{
-			const uint32 SubstrateShadingQuality = Substrate::GetShadingQuality(Target.GetPlatform());
+			const uint32 SubstrateShadingQuality = FMath::Clamp(Substrate::GetShadingQuality(Target.GetPlatform()), 0u, 1u);
 			SET_SHADER_DEFINE(Input.Environment, SUBSTRATE_SHADING_QUALITY, SubstrateShadingQuality);
 
-			const bool bLowQuality = SubstrateShadingQuality > 1;
+			const bool bLowQuality = SubstrateShadingQuality == 0;
 			SET_SHADER_DEFINE(Input.Environment, USE_ACHROMATIC_BXDF_ENERGY, bLowQuality ? 1u : 0u);
 
-			const uint32 SubstrateSheenQuality = Substrate::GetSheenQuality(Target.GetPlatform());
-			Input.Environment.SetDefine(TEXT("SUBSTRATE_SHEEN_QUALITY"), bLowQuality ? 2 : SubstrateSheenQuality);
+			const uint32 SubstrateSheenQuality = FMath::Clamp(Substrate::GetSheenQuality(Target.GetPlatform()), 0u, 1u);
+			Input.Environment.SetDefine(TEXT("SUBSTRATE_SHEEN_QUALITY"), bLowQuality ? 0 : SubstrateSheenQuality);
 
 			const uint32 SubstrateNormalQuality = Substrate::GetNormalQuality();
 			SET_SHADER_DEFINE(Input.Environment, SUBSTRATE_NORMAL_QUALITY, SubstrateNormalQuality);
diff --git a/Engine/Source/Runtime/RenderCore/Private/RenderUtils.cpp b/Engine/Source/Runtime/RenderCore/Private/RenderUtils.cpp
--- a/Engine/Source/Runtime/RenderCore/Private/RenderUtils.cpp
+++ b/Engine/Source/Runtime/RenderCore/Private/RenderUtils.cpp
@@ -1890,14 +1890,14 @@ static TAutoConsoleVariable<int32> CVarSubstrateOpaqueMaterialRoughRefraction(
 
 static TAutoConsoleVariable<int32> CVarSubstrateSheenQuality(
 	TEXT("r.Substrate.SheenQuality"),
-	2, // Substrate 60fps by default
-	TEXT("Define Substrate sheen quality (1: Disney-based sheen, 2: Charlie-based sheen, ). r.Substrate.ShadingQuality=2 forces SheenQuality to 2. This variable is read-only."),
+	0, // Substrate 60fps by default
+	TEXT("Define Substrate sheen quality (0: Charlie-based sheen, 1: Disney-based sheen). r.Substrate.ShadingQuality=0 forces SheenQuality to 0. This variable is read-only."),
 	ECVF_ReadOnly | ECVF_RenderThreadSafe | ECVF_Preview);
 
 static TAutoConsoleVariable<int32> CVarSubstrateShadingQuality(
 	TEXT("r.Substrate.ShadingQuality"),
-	2, // Substrate 60fps by default
-	TEXT("Define Substrate shading quality (1: accurate lighting, 2: approximate lighting). This variable is read-only."),
+	0, // Substrate 60fps by default
+	TEXT("Define Substrate shading quality (0: approximate lighting, 1: accurate lighting). This variable is read-only."),
 	ECVF_ReadOnly | ECVF_RenderThreadSafe | ECVF_Preview);
 
 static TAutoConsoleVariable<int32> CVarSubstrateDBufferPass(
