r.PostProcessing.GBufferPicking - added scene depth and fixed resolution scaling

diff --git a/Engine/Shaders/Private/PostProcessGBufferHints.usf b/Engine/Shaders/Private/PostProcessGBufferHints.usf
--- a/Engine/Shaders/Private/PostProcessGBufferHints.usf
+++ b/Engine/Shaders/Private/PostProcessGBufferHints.usf
@@ -171,7 +171,7 @@ void MainCS(uint3 DispatchThreadId : SV_DispatchThreadID)
 	if (all(DispatchThreadId == 0))
 	{
 		FShaderPrintContext Ctx = InitShaderPrintContext(true, uint2(50, 50));
-		const FScreenSpaceData ScreenSpaceData = GetScreenSpaceDataUint(ShaderPrintData.CursorCoord, true);
+		const FScreenSpaceData ScreenSpaceData = GetScreenSpaceDataUint(ShaderPrintData.CursorCoord * View.ViewResolutionFraction, true);
 		const FGBufferData Data = ScreenSpaceData.GBuffer;
 
 		float2 RectMin = Ctx.Pos;
@@ -179,6 +179,7 @@ void MainCS(uint3 DispatchThreadId : SV_DispatchThreadID)
 
 		RectMin = Ctx.Pos;
 		Print(Ctx, TEXT("Global"), FontOrange);Newline(Ctx, RectMax);
+		Print(Ctx, TEXT("Scene Depth       : "), FontWhite); Print(Ctx, Data.Depth, FontYellow); Newline(Ctx, RectMax);
 		Print(Ctx, TEXT("Diffuse Color     : "), FontWhite); Print(Ctx, Data.DiffuseColor, FontYellow); Newline(Ctx, RectMax);
 		Print(Ctx, TEXT("Specular Color    : "), FontWhite); Print(Ctx, Data.SpecularColor, FontYellow); Newline(Ctx, RectMax);
 		PrintBackgroundRect(Ctx, RectMin, RectMax);
