[Substrate] Fix incorrect tangent basis type in debug visualization.

diff --git a/Engine/Shaders/Private/Substrate/Substrate.ush b/Engine/Shaders/Private/Substrate/Substrate.ush
--- a/Engine/Shaders/Private/Substrate/Substrate.ush
+++ b/Engine/Shaders/Private/Substrate/Substrate.ush
@@ -2538,6 +2538,19 @@ float3x3 SubstrateGetBSDFSharedBasis(in FSubstratePixelHeader SubstratePixelHead
 #endif
 }
 
+uint SubstrateGetBSDFBasisType(in FSubstratePixelHeader SubstratePixelHeader, in FSubstrateBSDF BSDF)
+{
+	uint BasisType = SUBSTRATE_BASIS_TYPE_NORMAL;
+	#if SUBSTRATE_DEFERRED_SHADING
+	const bool bHasNormalLayer = SubstratePixelHeader.IsSimpleMaterial() || SubstratePixelHeader.IsSingleMaterial() || SubstratePixelHeader.IsEye() || SubstratePixelHeader.IsHair() || SubstratePixelHeader.IsSingleLayerWater();
+	if (!bHasNormalLayer)
+	{
+		BasisType = SubstrateGetSharedLocalBasisType(HEADER_GETSHAREDLOCALBASESTYPE(SubstratePixelHeader.PackedHeader), BSDF_GETSHAREDLOCALBASISID(BSDF));
+	}
+	#endif
+	return BasisType;
+}
+
 ///////////////////////////////////////////////////////////////////////////////
 // Operator nodes
 
diff --git a/Engine/Shaders/Shared/SubstrateVisualizeDefinitions.h b/Engine/Shaders/Shared/SubstrateVisualizeDefinitions.h
--- a/Engine/Shaders/Shared/SubstrateVisualizeDefinitions.h
+++ b/Engine/Shaders/Shared/SubstrateVisualizeDefinitions.h
@@ -324,7 +324,7 @@ FSubstratePixelDebugData ConvertToSerializableSubstratePixelDebugData(uint2 InCo
 
 #if SUBSTRATE_DEFERRED_SHADING
 		Data.Closures[i].NormalID = BSDF_GETSHAREDLOCALBASISID(BSDF);
-		Data.Closures[i].BasisType = SubstrateGetSharedLocalBasisType(HEADER_GETSHAREDLOCALBASESTYPE(Header.PackedHeader), BSDF_GETSHAREDLOCALBASISID(BSDF));
+		Data.Closures[i].BasisType = SubstrateGetBSDFBasisType(Header, BSDF);
 #endif
 
 		Data.Closures[i].bIsTopLayer = BSDF_GETISTOPLAYER(BSDF);
