Substrate - fixed anisotropic lighting when blendable gbuffer is used.

#tests PC editor dx12
[RN] Substrate - fixed anisotropic lighting when blendable gbuffer is used.

diff --git a/Engine/Shaders/Private/AnisotropyPassShader.usf b/Engine/Shaders/Private/AnisotropyPassShader.usf
--- a/Engine/Shaders/Private/AnisotropyPassShader.usf
+++ b/Engine/Shaders/Private/AnisotropyPassShader.usf
@@ -9,6 +9,13 @@
 #include "/Engine/Generated/VertexFactory.ush"
 #include "DeferredShadingCommon.ush"
 
+
+#if SUBSTRATE_ENABLED && SUBSTRATE_GBUFFER_FORMAT==0
+#define MATERIAL_SUBSTRATE_OPAQUE_PRECOMPUTED_LIGHTING 0 
+#define SUBSTRATE_MATERIAL_EXPORT_REQUESTED 1
+#include "/Engine/Private/Substrate/SubstrateExport.ush"
+#endif
+
 struct FAnisotropyPassVSToPS
 {
 	INVARIANT_OUTPUT float4 Position : SV_POSITION;
@@ -117,9 +124,45 @@ void MainPixelShader(
 	OutCoverage = DiscardMaterialWithPixelCoverage(MaterialParameters, PixelMaterialInputs);
 #endif
 
+#if SUBSTRATE_ENABLED
+#if SUBSTRATE_GBUFFER_FORMAT==0
+		// Write legacy anisotropty data from Substrate
+
+	FSubstrateData SubstrateData = PixelMaterialInputs.GetFrontSubstrateData();
+	FSubstratePixelHeader SubstratePixelHeader = MaterialParameters.GetFrontSubstrateHeader();
+
+	FSubstrateIntegrationSettings Settings = InitSubstrateIntegrationSettings(false /*bForceFullyRough*/, false /*SubstrateStruct.bRoughDiffuse*/, 0/*SubstrateStruct.PeelLayersAboveDepth*/, false/*SubstrateStruct.bRoughnessTracking*/);
+
+	const float3 SurfaceWorldNormal = MaterialParameters.TangentToWorld[2].xyz;
+	FExportResult Export = SubstrateMaterialExportOut(
+		Settings,
+		SubstratePixelHeader,
+		SubstrateData,
+		SurfaceWorldNormal,
+		MaterialParameters.WorldPosition_CamRelative,
+		0.0f // Curvature
+		);
+
+	float Anisotropy = Export.Anisotropy;
+
+	// Similar to CalculateAnisotropyTangent. SUBSTRATE_TODO: merge the code and account for bottom normal.
+	float3 WorldTangent = normalize(Export.WorldTangent);
+	float3 WorldNormal = normalize(Export.WorldNormal);
+	float3 WorldBiTangent = cross(WorldNormal, WorldTangent);
+	WorldTangent = normalize(cross(WorldBiTangent, WorldNormal));
+
+
+#else // SUBSTRATE_GBUFFER_FORMAT
+	#error Anisotropic shaders should not be enabled when using the full Substrate GBuffer.
+#endif // SUBSTRATE_GBUFFER_FORMAT
+
+#else // SUBSTRATE_ENABLED
+
 	float Anisotropy = GetMaterialAnisotropy(PixelMaterialInputs);
 	float3 WorldTangent = MaterialParameters.WorldTangent;
 
+#endif // SUBSTRATE_ENABLED
+
 	GBufferF = EncodeWorldTangentAndAnisotropy(WorldTangent, Anisotropy);
 }
 #endif // PIXELSHADER
diff --git a/Engine/Source/Runtime/Core/Private/UObject/FortniteMainBranchObjectVersion.cpp b/Engine/Source/Runtime/Core/Private/UObject/FortniteMainBranchObjectVersion.cpp
--- a/Engine/Source/Runtime/Core/Private/UObject/FortniteMainBranchObjectVersion.cpp
+++ b/Engine/Source/Runtime/Core/Private/UObject/FortniteMainBranchObjectVersion.cpp
@@ -14,6 +14,6 @@ TMap<FGuid, FGuid> FFortniteMainBranchObjectVersion::GetSystemGuids()
 	SystemGuids.Add(DevGuids.Niagara_LatestScriptCompileVersion, FGuid("97F0FD7149D90D42B65D6557AF477D1A"));
 	SystemGuids.Add(DevGuids.SkeletalMeshDerivedDataVersion, FGuid("C998F2F3B2884A02B7E290B8C25835F0"));
 	SystemGuids.Add(DevGuids.STATICMESH_DERIVEDDATA_VER, FGuid("3ABF32149F9D4D838750368AB95573BE"));
-	SystemGuids.Add(DevGuids.MaterialTranslationDDCVersion, FGuid("BC673DED3CEC4D72B10DC7336930BF9B"));
+	SystemGuids.Add(DevGuids.MaterialTranslationDDCVersion, FGuid("73294AD79624440FAD3F53E9794C505A"));
 	return SystemGuids;
 }
diff --git a/Engine/Source/Runtime/Engine/Private/Materials/HLSLMaterialTranslator.cpp b/Engine/Source/Runtime/Engine/Private/Materials/HLSLMaterialTranslator.cpp
--- a/Engine/Source/Runtime/Engine/Private/Materials/HLSLMaterialTranslator.cpp
+++ b/Engine/Source/Runtime/Engine/Private/Materials/HLSLMaterialTranslator.cpp
@@ -1761,13 +1761,16 @@ void FHLSLMaterialTranslator::TranslateMaterial()
 		Chunk[MP_Metallic]		= Material->CompilePropertyAndSetMaterialProperty(MP_Metallic,		this);
 		Chunk[MP_Specular]		= Material->CompilePropertyAndSetMaterialProperty(MP_Specular,		this);
 		Chunk[MP_Roughness]		= Material->CompilePropertyAndSetMaterialProperty(MP_Roughness,		this);
-		Chunk[MP_Anisotropy]	= Material->CompilePropertyAndSetMaterialProperty(MP_Anisotropy,	this);
 		Chunk[MP_Tangent]		= Material->CompilePropertyAndSetMaterialProperty(MP_Tangent,		this);
 
 		// Make sure to compile this property before using ShadingModelsFromCompilation
 		Chunk[MP_ShadingModel]	= Material->CompilePropertyAndSetMaterialProperty(MP_ShadingModel,	this);
 	}
 
+	// This needs to be evaluated here for substrate hidden conversion just to check if we need anisotropy through attribute, code chunk of constant.
+	Chunk[MP_Anisotropy]		= Material->CompilePropertyAndSetMaterialProperty(MP_Anisotropy, this);
+
+	// Other root node inputs always evaluated.
 	Chunk[MP_Opacity]			= Material->CompilePropertyAndSetMaterialProperty(MP_Opacity,		this);
 	Chunk[MP_OpacityMask]		= Material->CompilePropertyAndSetMaterialProperty(MP_OpacityMask,	this);
 	Chunk[MP_WorldPositionOffset]= Material->CompilePropertyAndSetMaterialProperty(MP_WorldPositionOffset, this);
diff --git a/Engine/Source/Runtime/Engine/Private/ShaderCompiler/ShaderCompiler.cpp b/Engine/Source/Runtime/Engine/Private/ShaderCompiler/ShaderCompiler.cpp
--- a/Engine/Source/Runtime/Engine/Private/ShaderCompiler/ShaderCompiler.cpp
+++ b/Engine/Source/Runtime/Engine/Private/ShaderCompiler/ShaderCompiler.cpp
@@ -4738,7 +4738,7 @@ void VerifyGlobalShaders(EShaderPlatform Platform, const ITargetPlatform* Target
 		if (Substrate::IsSubstrateEnabled())
 		{
 			// SUBSTRATE_TODO reduce the number of permutation of FDeferredLightPS.
-			PermutationCountLimit = 1304;	// FDeferredLightPS as of today (2023-12-04)
+			PermutationCountLimit = 2200;	// FDeferredLightPS as of today (2025-06-13)
 		}
 		ensureMsgf(
 			PermutationCountToCompile <= PermutationCountLimit,
diff --git a/Engine/Source/Runtime/Renderer/Private/AnisotropyRendering.cpp b/Engine/Source/Runtime/Renderer/Private/AnisotropyRendering.cpp
--- a/Engine/Source/Runtime/Renderer/Private/AnisotropyRendering.cpp
+++ b/Engine/Source/Runtime/Renderer/Private/AnisotropyRendering.cpp
@@ -20,20 +20,22 @@ static FAutoConsoleVariableRef CVarAnisotropicMaterials(
 
 bool SupportsAnisotropicMaterials(ERHIFeatureLevel::Type FeatureLevel, EShaderPlatform ShaderPlatform)
 {
+	const bool bSubstrateEnabled = Substrate::IsSubstrateEnabled();
 	return GAnisotropicMaterials
 		&& FeatureLevel >= ERHIFeatureLevel::SM5
 		&& FDataDrivenShaderPlatformInfo::GetSupportsAnisotropicMaterials(ShaderPlatform)
-		&& !Substrate::IsSubstrateEnabled(); // Substrate renders anisotropy surface natively, without extra pass.;
+		&& (!bSubstrateEnabled || (bSubstrateEnabled && Substrate::IsSubstrateBlendableGBufferEnabled(ShaderPlatform))); // Substrate renders anisotropy surface natively, without extra pass. Unless blendable gbuffer is used.
 }
 
 static bool IsAnisotropyPassCompatible(const EShaderPlatform Platform, FMaterialShaderParameters MaterialParameters)
 {
-	return 
+	const bool bSubstrateEnabled = Substrate::IsSubstrateEnabled();
+	return
 		FDataDrivenShaderPlatformInfo::GetSupportsAnisotropicMaterials(Platform) &&
-		!Substrate::IsSubstrateEnabled() && // Substrate renders anisotropy surface natively, without extra pass.
 		MaterialParameters.bHasAnisotropyConnected &&
 		!IsTranslucentBlendMode(MaterialParameters) &&
-		MaterialParameters.ShadingModels.HasAnyShadingModel({ MSM_DefaultLit, MSM_ClearCoat });
+		MaterialParameters.ShadingModels.HasAnyShadingModel({ MSM_DefaultLit, MSM_ClearCoat }) &&
+		(!bSubstrateEnabled || (bSubstrateEnabled && Substrate::IsSubstrateBlendableGBufferEnabled(Platform))); // Substrate renders anisotropy surface natively, without extra pass. Unless blendable gbuffer is used.
 }
 
 class FAnisotropyVS : public FMeshMaterialShader
diff --git a/Engine/Source/Runtime/Renderer/Private/LightRendering.cpp b/Engine/Source/Runtime/Renderer/Private/LightRendering.cpp
--- a/Engine/Source/Runtime/Renderer/Private/LightRendering.cpp
+++ b/Engine/Source/Runtime/Renderer/Private/LightRendering.cpp
@@ -993,7 +993,7 @@ class FDeferredLightPS : public FGlobalShader
 		{
 			if (Substrate::IsSubstrateEnabled())
 			{
-				return false;
+				return Substrate::IsSubstrateBlendableGBufferEnabled(Parameters.Platform);
 			}
 
 			// Anisotropic materials do not currently support rect lights
@@ -2744,10 +2744,15 @@ static void RenderLight(
 	const FSphere LightBounds = LightProxy->GetBoundingSphere();
 	const ELightComponentType LightType = (ELightComponentType)LightProxy->GetLightType();
 	const bool bIsRadial = LightType != LightType_Directional;
-	const bool bSupportAnisotropyPermutation = ShouldRenderAnisotropyPass(View) && !Substrate::IsSubstrateEnabled(); // Substrate managed anisotropy differently than legacy path. No need for special permutation.
 	const bool bUseVirtualShadowMapMask = VirtualShadowMapId != INDEX_NONE && ShadowMaskBits;
 	const bool bNeedComplexTransmittanceSupport = View.HairCardsMeshElements.Num() && IsHairStrandsSupported(EHairStrandsShaderType::All, View.GetShaderPlatform());
 
+	bool bSupportAnisotropyPermutation = ShouldRenderAnisotropyPass(View);
+	if (Substrate::IsSubstrateEnabled())
+	{
+		bSupportAnisotropyPermutation = bSupportAnisotropyPermutation && Substrate::IsSubstrateBlendableGBufferEnabled(View.GetShaderPlatform());
+	}
+
 	check(!bUseVirtualShadowMapMask || bIsRadial);		// VSM mask only stores local lights
 
 	// Debug Overlap shader
