Fix Substrate clear coat with anisotropy base layer.

The clear coat specular was computed using anisotropy context value. However the normal/view/... of this context is alterated by the anisotropy update, causing incorrect value when computing the clear coat transmittance/specular direction.

The reorganization of the code also improves performance by 0.01ms during the DiffuseIndirectAO pass.

#test QAGame PS5/PC, SubstrateTest PC

diff --git a/Engine/Shaders/Private/Substrate/SubstrateEvaluation.ush b/Engine/Shaders/Private/Substrate/SubstrateEvaluation.ush
--- a/Engine/Shaders/Private/Substrate/SubstrateEvaluation.ush
+++ b/Engine/Shaders/Private/Substrate/SubstrateEvaluation.ush
@@ -1628,6 +1628,23 @@ FSubstrateEnvLightResult SubstrateEvaluateForEnvLight(FSubstrateBSDFContext BSDF
 			BRANCH
 			if (bEnableSpecular)
 			{
+				// Hazy specular is top specular, non hazy is bottom with potentially bottom normal.
+				// Save N/R/NoV before they are modified by anisotropy
+				#if SUBSTRATE_FASTPATH==0
+				const float3 ClearCoatTopContextN   = BSDFContext.N;
+				const float3 ClearCoatTopContextR   = BSDFContext.R;
+				const float  ClearCoatTopContextNoV = BSDFContext.Context.NoV;
+				float ClearCoatBottomContextNoL = 0;
+				float ClearCoatBottomContextNoV = 0;
+				const bool bHasHaziness = BSDF_GETHASHAZINESS(BSDFContext.BSDF);
+				if (bHasHaziness)
+				{
+					BxDFContext ClearCoatBottomContext = RefractClearCoatContext(BSDFContext.Context);
+					ClearCoatBottomContextNoL = ClearCoatBottomContext.NoL;
+					ClearCoatBottomContextNoV = ClearCoatBottomContext.NoV;
+				}
+				#endif
+
 				#if SUBSTRATE_COMPLEXPATH
 				const bool bHasAnisotropy = BSDF_GETHASANISOTROPY(BSDFContext.BSDF);
 				if (bHasAnisotropy)
@@ -1667,7 +1684,6 @@ FSubstrateEnvLightResult SubstrateEvaluateForEnvLight(FSubstrateBSDFContext BSDF
 				SubstrateEnvLightResult.SSRSpecularWeight = SubstrateEnvLightResult.SpecularWeight;
 
 				#if SUBSTRATE_FASTPATH==0
-				const bool bHasHaziness = BSDF_GETHASHAZINESS(BSDFContext.BSDF);
 				if (bHasHaziness)
 				{
 					FHaziness Haziness = UnpackHaziness(SLAB_HAZINESS(BSDFContext.BSDF));
@@ -1675,18 +1691,15 @@ FSubstrateEnvLightResult SubstrateEvaluateForEnvLight(FSubstrateBSDFContext BSDF
 					const float HazeSafeRoughness = MakeRoughnessSafe(Haziness.Roughness);
 					const bool  bHazeAsSimpleClearCoat = Haziness.bSimpleClearCoat;
 
-					SubstrateEnvLightResult.SpecularHazeDirection = SubstrateGetOffSpecularPeakReflectionDir(BSDFContext.N, BSDFContext.R, HazeSafeRoughness);
+					SubstrateEnvLightResult.SpecularHazeDirection = SubstrateGetOffSpecularPeakReflectionDir(ClearCoatTopContextN, ClearCoatTopContextR, HazeSafeRoughness);
 
 					BRANCH
 					if (bHazeAsSimpleClearCoat)
 					{
-						// Hazy specular is top specular, non hazy is bottom with potentially bottom normal.
-						BxDFContext ClearCoatBottomContext = RefractClearCoatContext(BSDFContext.Context);
-						const float3 HazeClearCoatTransmittance = SimpleClearCoatTransmittance(ClearCoatBottomContext.NoL, ClearCoatBottomContext.NoV, SubstrateGetBSDFMetallic(BSDFContext.BSDF), SubstrateGetBSDFBaseColor(BSDFContext.BSDF));
-
+						const float3 HazeClearCoatTransmittance = SimpleClearCoatTransmittance(ClearCoatBottomContextNoL, ClearCoatBottomContextNoV, SubstrateGetBSDFMetallic(BSDFContext.BSDF), SubstrateGetBSDFBaseColor(BSDFContext.BSDF));
 						const float HazeF0 = 0.04f;
 						const float HazeF90= 1.0f;
-						FBxDFEnergyTermsRGB EnergyTerms = ComputeGGXSpecEnergyTermsRGB(HazeSafeRoughness, BSDFContext.Context.NoV, HazeF0, HazeF90);
+						FBxDFEnergyTermsRGB EnergyTerms = ComputeGGXSpecEnergyTermsRGB(HazeSafeRoughness, ClearCoatTopContextNoV, HazeF0, HazeF90);
 						const float3 EvalEnvBRDFHaze = EnergyTerms.E;
 						const float3 TopLayerSpecularTransmittionApprox = saturate(1.0f - EvalEnvBRDFHaze); // Simple and do not use extra LUT for now
 
@@ -1703,7 +1716,7 @@ FSubstrateEnvLightResult SubstrateEvaluateForEnvLight(FSubstrateBSDFContext BSDF
 							SubstrateEnvLightResult.SpecularHazeDirection = SubstrateEnvLightResult.SpecularDirection;
 
 							// Now compute the bottom specular direction.
-							float3 SpecularBottomNormal = SubstrateGetSimpleCoatBottomNormal(Haziness.BottomNormalOct, BSDFContext.N);
+							float3 SpecularBottomNormal = SubstrateGetSimpleCoatBottomNormal(Haziness.BottomNormalOct, ClearCoatTopContextN);
 							float3 SpecularBottomR = 2 * dot(BSDFContext.V, SpecularBottomNormal) * SpecularBottomNormal - BSDFContext.V;
 							SubstrateEnvLightResult.SpecularDirection = SubstrateGetOffSpecularPeakReflectionDir(SpecularBottomNormal, SpecularBottomR, SubstrateEnvLightResult.SpecularSafeRoughness);
 						}
