Fixes for inline raytracing and non uniform resource indexing

diff --git a/Engine/Shaders/Private/RayTracing/TraceRayInlineD3D12.ush b/Engine/Shaders/Private/RayTracing/TraceRayInlineD3D12.ush
--- a/Engine/Shaders/Private/RayTracing/TraceRayInlineD3D12.ush
+++ b/Engine/Shaders/Private/RayTracing/TraceRayInlineD3D12.ush
@@ -61,8 +61,8 @@ FTriangleBaseAttributes LoadInlineRayTracingTriangleAttributes(
 	const uint IndexBufferHandle = Params.BindlessHitGroupSystemIndexBuffer;
 	const uint VertexBufferHandle = Params.BindlessHitGroupSystemVertexBuffer;
 
-	ByteAddressBuffer IndexBuffer = ResourceDescriptorHeap[IndexBufferHandle];
-	ByteAddressBuffer VertexBuffer = ResourceDescriptorHeap[VertexBufferHandle];	
+	ByteAddressBuffer IndexBuffer = ResourceDescriptorHeap[NonUniformResourceIndex(IndexBufferHandle)];
+	ByteAddressBuffer VertexBuffer = ResourceDescriptorHeap[NonUniformResourceIndex(VertexBufferHandle)];	
 
 	return LoadTriangleBaseAttributes(IndexBuffer, IndexBufferOffsetInBytes, IndexBufferStride, VertexBuffer, VertexBufferOffsetInBytes, VertexStride, PrimitiveIndex);
 }
