[SingleLayerWater] Fixed crash when using the SceneDepthWithoutWater node in a hitproxy shader, which may happen if it affects the OpacityMask pin in the material graph.

diff --git a/Engine/Shaders/Private/MaterialTemplate.ush b/Engine/Shaders/Private/MaterialTemplate.ush
--- a/Engine/Shaders/Private/MaterialTemplate.ush
+++ b/Engine/Shaders/Private/MaterialTemplate.ush
@@ -2434,7 +2434,7 @@ float2 SceneWithoutSingleLayerWaterInvTextureSize;
 */
 float MaterialExpressionSceneDepthWithoutWater(float2 ViewportUV, float FallbackDepth)
 {
-#if MATERIAL_SHADINGMODEL_SINGLELAYERWATER && !SCENE_TEXTURES_DISABLED && (SINGLE_LAYER_WATER_SHADING_QUALITY != SINGLE_LAYER_WATER_SHADING_QUALITY_MOBILE_WITH_DEPTH_BUFFER)
+#if MATERIAL_SHADINGMODEL_SINGLELAYERWATER && !HIT_PROXY_SHADER && !SCENE_TEXTURES_DISABLED && (SINGLE_LAYER_WATER_SHADING_QUALITY != SINGLE_LAYER_WATER_SHADING_QUALITY_MOBILE_WITH_DEPTH_BUFFER)
 	const float2 ClampedViewportUV = clamp(ViewportUV, SingleLayerWater.SceneWithoutSingleLayerWaterMinMaxUV.xy, SingleLayerWater.SceneWithoutSingleLayerWaterMinMaxUV.zw);
 	float WaterDeviceZ = WaterSampleSceneDepthWithoutWater(SingleLayerWater.SceneDepthWithoutSingleLayerWaterTexture, SingleLayerWaterSceneDepthSampler, ClampedViewportUV, 
 		SingleLayerWater.SceneWithoutSingleLayerWaterTextureSize, SingleLayerWater.SceneWithoutSingleLayerWaterInvTextureSize);
