Shader Print - fixed shader compilation error "use of undeclared identifier 'ShaderPrintData'"

diff --git a/Engine/Shaders/Private/ShaderPrintCommon.ush b/Engine/Shaders/Private/ShaderPrintCommon.ush
--- a/Engine/Shaders/Private/ShaderPrintCommon.ush
+++ b/Engine/Shaders/Private/ShaderPrintCommon.ush
@@ -1499,8 +1499,8 @@ void AddTextBackground(inout FShaderPrintContext Ctx, float4 Color)
 	{
 		AddFilledQuadSS(
 			Ctx,
-			(Ctx.RectMin - ShaderPrintData.FontSize) * ShaderPrintData.Resolution,
-			(Ctx.RectMax + ShaderPrintData.FontSize) * ShaderPrintData.Resolution,
+			(Ctx.RectMin - Ctx.Config.FontSize) * Ctx.Config.Resolution,
+			(Ctx.RectMax + Ctx.Config.FontSize) * Ctx.Config.Resolution,
 			Color);
 	}
 }
\ No newline at end of file
