Fix Primitive Flags unpacking in HairStrandsMaterialCommon.ush.

diff --git a/Engine/Shaders/Private/HairStrands/HairStrandsMaterialCommon.ush b/Engine/Shaders/Private/HairStrands/HairStrandsMaterialCommon.ush
--- a/Engine/Shaders/Private/HairStrands/HairStrandsMaterialCommon.ush
+++ b/Engine/Shaders/Private/HairStrands/HairStrandsMaterialCommon.ush
@@ -212,8 +212,8 @@ FEvaluateOutput Evaluate(
 
 #if SUBSTRATE_GBUFFER_FORMAT==1 && MATERIAL_IS_SUBSTRATE && MATERIAL_SHADINGMODEL_HAIR
 		FSubstratePixelHeader SubstratePixelHeader = MaterialParameters.GetFrontSubstrateHeader();
-		SubstratePixelHeader.SetCastContactShadow(GetPrimitiveData(MaterialParameters).Flags & 0x100);
-		SubstratePixelHeader.SetDynamicIndirectShadowCasterRepresentation(GetPrimitiveData(MaterialParameters).Flags & 0x80);
+		SubstratePixelHeader.SetCastContactShadow(GetPrimitiveData(MaterialParameters).Flags & PRIMITIVE_SCENE_DATA_FLAG_CAST_SHADOWS);
+		SubstratePixelHeader.SetDynamicIndirectShadowCasterRepresentation(GetPrimitiveData(MaterialParameters).Flags & PRIMITIVE_SCENE_DATA_FLAG_HAS_CAPSULE_REPRESENTATION);
 
 		float3 Tangent = float3(0,1,0);
 		//if (SubstratePixelHeader.SubstrateTree.BSDFCount > 0)
