New Material Translator:  File missing from previous check-in adding support for SceneTexture and UserSceneTexture.

#rnx

diff --git a/Engine/Shaders/Private/MaterialTemplate.ush b/Engine/Shaders/Private/MaterialTemplate.ush
--- a/Engine/Shaders/Private/MaterialTemplate.ush
+++ b/Engine/Shaders/Private/MaterialTemplate.ush
@@ -3370,6 +3370,23 @@ float4 SceneTextureFetchFunc(FMaterialPixelParameters Parameters, int SceneTextu
 {
 	return SceneTextureFetchFunc(Parameters, SceneTextureIndex, float2(PixelOffsetX, PixelOffsetY));
 }
+
+// SceneTextureLookup function used by new translator.  Shading path agnostic, and handles conditional application of exposure scale, so generated code doesn't need to vary.
+float4 SceneTextureLookup(FMaterialPixelParameters Parameters, float2 UV, int SceneTextureIndex, bool bFiltered)
+{
+#if SHADING_PATH_MOBILE
+	float4 Lookup = MobileSceneTextureLookup(Parameters, SceneTextureIndex, UV);
+#else
+	float4 Lookup = SceneTextureLookup(UV, SceneTextureIndex, bFiltered);
+#endif
+#if POST_PROCESS_MATERIAL && POST_PROCESS_MATERIAL_BEFORE_TONEMAP && !POST_PROCESS_DISABLE_PRE_EXPOSURE_SCALE
+	if (SceneTextureIndex >= PPI_PostProcessInput0 && SceneTextureIndex <= PPI_PostProcessInput6)
+	{
+		Lookup.xyz *= View.OneOverPreExposure;
+	}
+#endif
+	return Lookup;
+}
 #endif // NEEDS_SCENE_TEXTURES
 
 #if SHADING_PATH_DEFERRED || MOBILE_DEFERRED_SHADING
