Make DiffuseAlbedo/F0/Roughness packing shared between the different Substrate path (Simple/Legacy-/Cloth/Wrap/SSSProfile/ClearCoat).

This improves base pass cost by 5% on PS5.

diff --git a/Engine/Shaders/Private/PackUnpack.ush b/Engine/Shaders/Private/PackUnpack.ush
--- a/Engine/Shaders/Private/PackUnpack.ush
+++ b/Engine/Shaders/Private/PackUnpack.ush
@@ -104,10 +104,37 @@ float UnpackR8(uint In)
 #if COMPILER_SUPPORTS_UNPACKBYTEN
 	return UnpackByte0(In) * (1.0f / 255.0f);
 #else
-	return float(In & 0xFF) * (1.0f / 255.0f);
+	return float(BitFieldExtractU32(In, 8, 0)) * (1.0f / 255.0f);
 #endif
 }
 
+float UnpackG8(uint In)
+{
+	#if COMPILER_SUPPORTS_UNPACKBYTEN
+	return float(UnpackByte1(In)) * (1.0f / 255.0f);
+	#else
+	return float(BitFieldExtractU32(In, 8, 8)) * (1.0f / 255.0f);
+	#endif
+}
+
+float UnpackB8(uint In)
+{
+	#if COMPILER_SUPPORTS_UNPACKBYTEN
+	return float(UnpackByte2(In)) * (1.0f / 255.0f);
+	#else
+	return float(BitFieldExtractU32(In, 8, 16)) * (1.0f / 255.0f);
+	#endif
+}
+
+float UnpackA8(uint In)
+{
+	#if COMPILER_SUPPORTS_UNPACKBYTEN
+	return float(UnpackByte3(In)) * (1.0f / 255.0f);
+	#else
+	return float(BitFieldExtractU32(In, 8, 24)) * (1.0f / 255.0f);
+	#endif
+}
+
 //
 //	float <-> u6
 //
diff --git a/Engine/Shaders/Private/Substrate/Substrate.ush b/Engine/Shaders/Private/Substrate/Substrate.ush
--- a/Engine/Shaders/Private/Substrate/Substrate.ush
+++ b/Engine/Shaders/Private/Substrate/Substrate.ush
@@ -272,7 +272,7 @@ struct FSubstrateBSDF
 
 	void SubstrateSetBSDFRoughness(in float Roughness);
 
-	void UnpackFastPathSlabBSDF(uint2 PackedData01);
+	void UnpackSingleOrComplexPathSlabBSDF(uint Data0, uint Data1);
 
 	void PostProcessSlabBeforeLighting(bool bIsSubstrateOpaqueMaterial);
 
@@ -770,6 +770,16 @@ FHaziness UnpackHaziness(uint Data)
 	return Haziness;
 }
 
+float GetHazinessRoughness(uint Data)
+{
+	return UnpackR8(BitFieldExtractU32(Data, 8, 0));
+}
+
+uint SetHazinessRoughness(uint Data, float InRoughness)
+{
+	return (Data & 0xFFFFFF00) | PackR8(InRoughness);
+}
+
 uint LerpHazinessParameterBlending(FSubstrateBSDF A, FSubstrateBSDF B, float MixValue)
 {
 	FHaziness HazinessA = UnpackHaziness(SLAB_HAZINESS(A));
@@ -4144,20 +4154,53 @@ bool IsSubstrateSlabCompatible_ComplexSpecialPath(in FSubstrateBSDF BSDF)
 	return true;
 }
 
-void FSubstrateBSDF::UnpackFastPathSlabBSDF(uint2 PackedData01)
+void PackDiffuseAlbedoF0(float Dither, float3 DiffuseAlbedo, float3 F0, inout uint OutData0, inout uint OutData1)
+{
+	const uint PackedDiffuse20Bits	= PackR7G7B6Gamma2(DiffuseAlbedo, Dither);
+	const uint PackedF020Bits		= PackR7G7B6Gamma2(F0, Dither);
+
+	OutData0 = BitFieldExtractU32(PackedF020Bits,  8,12) << 24;
+	OutData1 =(BitFieldExtractU32(PackedF020Bits, 12, 0) << 20) | PackedDiffuse20Bits;
+
+}
+
+void PackDiffuseAlbedoF0Roughness(float Dither, float3 DiffuseAlbedo, float3 F0, float Roughness, inout uint OutData0, inout uint OutData1)
+{
+	PackDiffuseAlbedoF0(Dither, DiffuseAlbedo, F0, OutData0, OutData1);
+
+	const uint PackedRoughness8Bits	= PackR8(Roughness);
+	OutData0 = OutData0 | (PackedRoughness8Bits << 16);
+}
+
+void UnpackDiffuseAlbedoF0Roughness(uint InData0, uint InData1, inout float3 OutDiffuseAlbedo, inout float3 OutF0, inout float OutRoughness)
+{
+	const uint PackedRoughness8bits = BitFieldExtractU32(InData0, 8, HEADER_FASTENCODING_BIT_COUNT);
+	const uint PackedDiffuse20Bits	= BitFieldExtractU32(InData1,20, 0);
+	const uint PackedF020Bits		= BitFieldExtractU32(InData1,12, 20) | (BitFieldExtractU32(InData0, 8, HEADER_FASTENCODING_BIT_COUNT + 8)<<12);
+
+	OutRoughness = UnpackR8(PackedRoughness8bits);
+	OutDiffuseAlbedo = UnpackR7G7B6Gamma2(PackedDiffuse20Bits);
+	OutF0 = UnpackR7G7B6Gamma2(PackedF020Bits);
+}
+
+void UnpackDiffuseAlbedoF0(uint InData0, uint InData1, inout float3 OutDiffuseAlbedo, inout float3 OutF0)
+{
+	float OutDummyRoughness;
+	UnpackDiffuseAlbedoF0Roughness(InData0, InData1, OutDiffuseAlbedo, OutF0, OutDummyRoughness);
+}
+
+void FSubstrateBSDF::UnpackSingleOrComplexPathSlabBSDF(uint Data1, uint Data2)
 {
-	const uint PackedData8Bits		= PackedData01.y & 0xFF;
-	const uint PackedData32Bits		= PackedData01.x;
-	const uint PackedF020Bits		= ((PackedData8Bits << 12) & 0xFF000) | (PackedData32Bits & 0xFFF);
-	const uint PackedDiffuse20Bits	= (PackedData32Bits >> 12) & 0xFFFFF;
+	const uint PackedDiffuse20Bits	= BitFieldExtractU32(Data1, 20, 12);
+	const uint PackedF020Bits		= BitFieldExtractU32(Data1, 12,  0) | (BitFieldExtractU32(Data2,  8, 0)<<24);
 
 	SLAB_DIFFUSEALBEDO(this)	= UnpackR7G7B6Gamma2(PackedDiffuse20Bits);
 	SLAB_F0(this)				= UnpackR7G7B6Gamma2(PackedF020Bits);
 
-	float4 Data1 = UnpackRGBA8(PackedData01.y);
-	SLAB_ROUGHNESS(this)			= Data1.y;
-	SLAB_ANISOTROPY(this)			= Data1.z * 2.f - 1.f;
-	SLAB_SSSPHASEANISOTROPY(this)	= Data1.w * 2.f - 1.f;
+	float4 fData2 = UnpackRGBA8(Data2);
+	SLAB_ROUGHNESS(this)			= fData2.y;
+	SLAB_ANISOTROPY(this)			= fData2.z * 2.f - 1.f;
+	SLAB_SSSPHASEANISOTROPY(this)	= fData2.w * 2.f - 1.f;
 }
 
 FSubstrateBSDF UnpackFastPathSubstrateBSDFIn(uint Data0, uint Data1)
@@ -4170,15 +4213,14 @@ FSubstrateBSDF UnpackFastPathSubstrateBSDFIn(uint Data0, uint Data1)
 	OutBSDF.CoverageAboveAlongN = 0.0f;
 	OutBSDF.TransmittanceAboveAlongN = 1.0f;
 
-	const uint PackedRoughness8bits = 0xFF    & (Data0 >>  HEADER_FASTENCODING_BIT_COUNT);
-	const uint PackedData8Bits		= 0xFF    & (Data0 >> (HEADER_FASTENCODING_BIT_COUNT + 8));
-	const uint PackedData12Bits		= 0xFFF   & (Data1 >> 20);
-	const uint PackedF020Bits		= 0xFFFFF & (Data1);
-	const uint PackedDiffuse20Bits	= PackedData12Bits | (PackedData8Bits<<12);
+	float3 OutDiffuseAlbedo;
+	float3 OutF0; 
+	float OutRoughness;
+	UnpackDiffuseAlbedoF0Roughness(Data0, Data1, OutDiffuseAlbedo, OutF0, OutRoughness);
 
-	SLAB_DIFFUSEALBEDO(OutBSDF)	= UnpackR7G7B6Gamma2(PackedDiffuse20Bits);
-	SLAB_F0(OutBSDF)			= UnpackR7G7B6Gamma2(PackedF020Bits);
-	SLAB_ROUGHNESS(OutBSDF)		= UnpackR8(PackedRoughness8bits);
+	SLAB_DIFFUSEALBEDO(OutBSDF)	= OutDiffuseAlbedo;
+	SLAB_F0(OutBSDF)			= OutF0;
+	SLAB_ROUGHNESS(OutBSDF)		= OutRoughness;
 	SLAB_ANISOTROPY(OutBSDF)	= 0.0f;
 #endif
 
@@ -4333,20 +4375,21 @@ FSubstrateBSDF UnpackSubstrateBSDFIn(FSubstrateMaterialContainer SubstrateBuffer
 			OutBSDF.CoverageAboveAlongN = 0.0f;
 			OutBSDF.TransmittanceAboveAlongN = 1.0f;
 			SLAB_F90(OutBSDF) = 1.0f;
-
+	
 			if (OptimisedLegacyMode == SINGLE_OPTLEGACYMODE_CLEARCOAT)
 			{
 			#if SUBSTRATE_DEFERRED_SHADING
-				SLAB_ROUGHNESS(OutBSDF)		 = UnpackR7(Data0 >> (HEADER_SINGLE_TOTAL_BIT_COUNT + 0));
-
-				uint PackedDiffuse19Bits	 = (Data0 >> (HEADER_SINGLE_TOTAL_BIT_COUNT + 7)) & 0x3F;	// low  6  bits
-				PackedDiffuse19Bits			|= (Data1 & 0x1FFF) << 6;									// high 13 bits
-				SLAB_DIFFUSEALBEDO(OutBSDF)	 = UnpackR6G7B6Gamma2(PackedDiffuse19Bits);
-
-				uint PackedF019Bits			 = Data1 >> 13;
-				SLAB_F0(OutBSDF)			 = UnpackR6G7B6Gamma2(PackedF019Bits);
+				float3 OutDiffuseAlbedo;
+				float3 OutF0;
+				UnpackDiffuseAlbedoF0(Data0, Data1, OutDiffuseAlbedo, OutF0);
+				SLAB_DIFFUSEALBEDO(OutBSDF)	= OutDiffuseAlbedo;
+				SLAB_F0(OutBSDF)			= OutF0;
+				SLAB_ROUGHNESS(OutBSDF)		= SubstrateUnpackTopLayerData(SubstrateHeader.PackedTopLayerData).Roughness;
 
-				SLAB_HAZINESS(OutBSDF)		 = Data2;
+				const float ClearCoatRoughness = SLAB_ROUGHNESS(OutBSDF);
+				const float BaseRoughness = GetHazinessRoughness(Data2);
+				SLAB_HAZINESS(OutBSDF) = SetHazinessRoughness(Data2, ClearCoatRoughness);
+				SLAB_ROUGHNESS(OutBSDF) = BaseRoughness;
 			#endif
 
 				// Force a BSDF state to be a single layer on top.
@@ -4365,17 +4408,16 @@ FSubstrateBSDF UnpackSubstrateBSDFIn(FSubstrateMaterialContainer SubstrateBuffer
 			else if (OptimisedLegacyMode == SINGLE_OPTLEGACYMODE_CLOTH)
 			{
 			#if SUBSTRATE_DEFERRED_SHADING
-
-				const uint PackedDiffuse20Bits	= Data1 & 0x000FFFFF;
-				const uint PackedF020Bits		= Data2 & 0x000FFFFF;
-				const uint PackedFuzzColor20bits= ((Data1 >> 10) & 0xFFC00) | ((Data2 >> 20) & 0x03FF);
-				const uint PackedFuzzAmount8bits= (Data0 >> (HEADER_SINGLEENCODING_BIT_COUNT + HEADER_SINGLE_OPTLEGACYMODE_BIT_COUNT)) & 0xFF;
-
-				SLAB_DIFFUSEALBEDO(OutBSDF)	= UnpackR7G7B6Gamma2(PackedDiffuse20Bits);
-				SLAB_F0(OutBSDF)			= UnpackR7G7B6Gamma2(PackedF020Bits);
-				SLAB_FUZZ_COLOR(OutBSDF)	= UnpackR7G7B6Gamma2(PackedFuzzColor20bits);
-				SLAB_FUZZ_AMOUNT(OutBSDF)	= UnpackR8(PackedFuzzAmount8bits);
+				float3 OutDiffuseAlbedo;
+				float3 OutF0;
+				UnpackDiffuseAlbedoF0(Data0, Data1, OutDiffuseAlbedo, OutF0);
+				SLAB_DIFFUSEALBEDO(OutBSDF)	= OutDiffuseAlbedo;
+				SLAB_F0(OutBSDF)			= OutF0;
 				SLAB_ROUGHNESS(OutBSDF)		= SubstrateUnpackTopLayerData(SubstrateHeader.PackedTopLayerData).Roughness;
+
+				const float4 CustomData = UnpackRGBA8(Data2);
+				SLAB_FUZZ_COLOR(OutBSDF)	= CustomData.xyz;
+				SLAB_FUZZ_AMOUNT(OutBSDF)	= CustomData.w;
 				SLAB_FUZZ_ROUGHNESS(OutBSDF)= SLAB_ROUGHNESS(OutBSDF);
 			#endif
 
@@ -4395,21 +4437,17 @@ FSubstrateBSDF UnpackSubstrateBSDFIn(FSubstrateMaterialContainer SubstrateBuffer
 			else if (OptimisedLegacyMode == SINGLE_OPTLEGACYMODE_SSSWRAP || OptimisedLegacyMode == SINGLE_OPTLEGACYMODE_TWO_SIDED_SSSWRAP)
 			{
 			#if SUBSTRATE_DEFERRED_SHADING
-				const uint PackedDiffuse20Bits	= Data1 & 0x000FFFFF;
-				const uint PackedF020Bits		= Data2 & 0x000FFFFF;
-				
-				const uint PackedSSSMFP12BitsA = (Data1 >> 20) & 0xFFF;																				// 12 lower bits
-				const uint PackedSSSMFP12BitsB = (Data2 >> 20) & 0xFFF;																				// next 12 lower bits
-				const uint PackedSSSMFP6BitsC  = (Data0 >> (7 + HEADER_SINGLEENCODING_BIT_COUNT + HEADER_SINGLE_OPTLEGACYMODE_BIT_COUNT)) & 0x03F;	// 6 higher bits
-				const uint PackedSSSMFP30bits  = PackedSSSMFP12BitsA | (PackedSSSMFP12BitsB << 12) | (PackedSSSMFP6BitsC << 24);
-
-				const uint PackedSSSWOpacity7bits = (Data0 >> (HEADER_SINGLEENCODING_BIT_COUNT + HEADER_SINGLE_OPTLEGACYMODE_BIT_COUNT)) & 0x7F;
-				const float Opacity = UnpackR7(PackedSSSWOpacity7bits);
-				
-				SLAB_DIFFUSEALBEDO(OutBSDF)	= UnpackR7G7B6Gamma2(PackedDiffuse20Bits);
-				SLAB_F0(OutBSDF)			= UnpackR7G7B6Gamma2(PackedF020Bits);
+				float3 OutDiffuseAlbedo;
+				float3 OutF0;
+				UnpackDiffuseAlbedoF0(Data0, Data1, OutDiffuseAlbedo, OutF0);
+				SLAB_DIFFUSEALBEDO(OutBSDF)	= OutDiffuseAlbedo;
+				SLAB_F0(OutBSDF)			= OutF0;
 				SLAB_ROUGHNESS(OutBSDF)		= SubstrateUnpackTopLayerData(SubstrateHeader.PackedTopLayerData).Roughness;
-				SLAB_SSSMFP(OutBSDF)		= UnpackR10G10B10F(PackedSSSMFP30bits);
+
+				const uint PackedOpacity7bit = BitFieldExtractU32(Data0, 5, HEADER_SINGLEENCODING_BIT_COUNT + HEADER_SINGLE_OPTLEGACYMODE_BIT_COUNT) | (BitFieldExtractU32(Data2, 2, 30) << 5);
+				const float Opacity = UnpackR7(PackedOpacity7bit);
+
+				SLAB_SSSMFP(OutBSDF)		= UnpackR10G10B10F(BitFieldExtractU32(Data2, 30, 0));
 				SLAB_SSSPHASEANISOTROPY(OutBSDF)= saturate(1.f - Opacity);
 			#endif
 
@@ -4432,26 +4470,16 @@ FSubstrateBSDF UnpackSubstrateBSDFIn(FSubstrateMaterialContainer SubstrateBuffer
 			else if (OptimisedLegacyMode == SINGLE_OPTLEGACYMODE_SSSPROFILE)
 			{
 			#if SUBSTRATE_DEFERRED_SHADING
-				const uint PackedDiffuse20Bits	= Data1 & 0x000FFFFF;
-				const uint PackedF020Bits		= Data2 & 0x000FFFFF;
-				
-				float Roughness						= UnpackR8(Data0 >> 24);
-				const float SSSProfileRadiusScale	= UnpackR8(Data1 >> 24);
-				const float SSSProfileId			= UnpackR8(Data2 >> 24);
-				FHaziness Haziness 					= InitialiseHaziness();
-
-				// Average roughness for dual specular.
-				// When storing data as SINGLE_OPTLEGACYMODE_SSSPROFILE, the dual roughness is not translated/stored as haziness in order to save storage space. 
-				// Thus we need to (re)compute them directly from the SSS profile data.
-				const uint SubsurfaceProfileUInt = SubstrateSubsurfaceProfileIdTo8bits(SSSProfileId);
-				GetSubsurfaceProfileDualSpecular(SubsurfaceProfileUInt, Roughness, SSSProfileRadiusScale, Roughness, Haziness.Roughness, Haziness.Weight);
-
-				SLAB_DIFFUSEALBEDO(OutBSDF)			= UnpackR7G7B6Gamma2(PackedDiffuse20Bits);
-				SLAB_F0(OutBSDF)					= UnpackR7G7B6Gamma2(PackedF020Bits);
-				SLAB_ROUGHNESS(OutBSDF)				= Roughness;
-				SLAB_SSSPROFILEID(OutBSDF)			= SSSProfileId;
-				SLAB_SSSPROFILERADIUSSCALE(OutBSDF)	= SSSProfileRadiusScale;
-				SLAB_HAZINESS(OutBSDF)				= PackHaziness(Haziness);
+				float3 OutDiffuseAlbedo;
+				float3 OutF0;
+				UnpackDiffuseAlbedoF0(Data0, Data1, OutDiffuseAlbedo, OutF0);
+				SLAB_DIFFUSEALBEDO(OutBSDF)	= OutDiffuseAlbedo;
+				SLAB_F0(OutBSDF)			= OutF0;
+				SLAB_ROUGHNESS(OutBSDF)		= SubstrateUnpackTopLayerData(SubstrateHeader.PackedTopLayerData).Roughness;
+
+				SLAB_SSSPROFILEID(OutBSDF)			= UnpackR8(Data2);
+				SLAB_SSSPROFILERADIUSSCALE(OutBSDF)	= UnpackG8(Data2);
+				SLAB_HAZINESS(OutBSDF)				= BitFieldExtractU32(Data2, 16, 16);
 
 				// When loading Substrate BSDF which have SSS enable, BaseColor and (optionally) Specular value can be overriden 
 				// based on the SSS method used (screen space SSS / checherboarded screen space SSS / ...)
@@ -4494,9 +4522,7 @@ FSubstrateBSDF UnpackSubstrateBSDFIn(FSubstrateMaterialContainer SubstrateBuffer
 		{
 			//OutBSDF.State = ... already done above
 			#if SUBSTRATE_DEFERRED_SHADING
-			uint DataX = Data1;
-			uint DataY = Data2;
-			OutBSDF.UnpackFastPathSlabBSDF(uint2(DataX, DataY));
+			OutBSDF.UnpackSingleOrComplexPathSlabBSDF(Data1, Data2);
 			#endif
 		}
 		else
@@ -4522,7 +4548,7 @@ FSubstrateBSDF UnpackSubstrateBSDFIn(FSubstrateMaterialContainer SubstrateBuffer
 			}
 
 			#if SUBSTRATE_DEFERRED_SHADING
-			OutBSDF.UnpackFastPathSlabBSDF(uint2(DataX, DataY));
+			OutBSDF.UnpackSingleOrComplexPathSlabBSDF(DataX, DataY);
 			#endif
 		}
 
diff --git a/Engine/Shaders/Private/Substrate/SubstrateExport.ush b/Engine/Shaders/Private/Substrate/SubstrateExport.ush
--- a/Engine/Shaders/Private/Substrate/SubstrateExport.ush
+++ b/Engine/Shaders/Private/Substrate/SubstrateExport.ush
@@ -522,31 +522,24 @@ void PackSubstrateOut(
 					const bool bFastEncodedBSDF = bHasFastEncoding && BSDFIdx == OneBSDFMaterial_Index;
 					if (bFastEncodedBSDF)
 					{
-						// Now we pack diffuse and F0 in a special way: both encoded as R7G7B6 with gamma 2.0. The low bit cound will be hidden by dithering and TAA.
-						// We pack 32 bits in the first uint and the remaining bits are pack is the lowest significant bit of the second uint.
-						const uint PackedDiffuse20Bits = PackR7G7B6Gamma2(SLAB_DIFFUSEALBEDO(BSDF), Dither);
-						const uint PackedDiffuse12Bits = PackedDiffuse20Bits & 0xFFF;
-						const uint PackedDiffuse8Bits = (PackedDiffuse20Bits >> 12) & 0xFF;
-						const uint PackedF020Bits = PackR7G7B6Gamma2(SLAB_F0(BSDF), Dither);
-						const uint PackedRoughness8bits = PackR8(SLAB_ROUGHNESS(BSDF));
+						uint Data0 = 0;
+						uint Data1 = 0;
+						PackDiffuseAlbedoF0Roughness(Dither, SLAB_DIFFUSEALBEDO(BSDF), SLAB_F0(BSDF), SLAB_ROUGHNESS(BSDF), Data0, Data1);
 
 						// Data0 (Header_State|Header_AO|Roughness|Diffuse8bits)
 						{
 							uint Out = 0;
 							HEADER_SETCOMMONSTATES(Out, SubstratePixelHeader.State);
 
-							uint Packed16 = (PackedDiffuse8Bits << 8) | PackedRoughness8bits;
-							Out = (Out & HEADER_FASTENCODING_MASK) | (Packed16 << HEADER_FASTENCODING_BIT_COUNT);
-							SUBSTRATE_STORE_UINT1(Out);
+							SUBSTRATE_STORE_UINT1(Out | Data0);
 
 						#if (HEADER_FASTENCODING_BIT_COUNT + 16) > 32
 						#error Substrate fast path header is > 32bits
 						#endif
 						}
-						// Data1 (F0|Diffuse12bits)
+						// Data1
 						{
-							uint Out = PackedF020Bits | (PackedDiffuse12Bits << 20);
-							SUBSTRATE_STORE_UINT1(Out);
+							SUBSTRATE_STORE_UINT1(Data1);
 						}
 						// Ensure the rest of the BSDF is not stored with the regular path
 						BSDF.State = 0;
@@ -575,23 +568,12 @@ void PackSubstrateOut(
 						// Because we use 3 UINT MRT, this means that a total of 13+32+32 = 77 bits.
 						else if(OptimisedLegacyMode == SINGLE_OPTLEGACYMODE_CLEARCOAT)
 						{
-							// 77 bits with more aggressive bit reduction
-							const uint PackedRoughness7Bits		= PackR7(SLAB_ROUGHNESS(BSDF), Dither);
-							const uint PackedDiffuse19Bits		= PackR6G7B6Gamma2(SLAB_DIFFUSEALBEDO(BSDF), Dither);
-							const uint PackedF019Bits			= PackR6G7B6Gamma2(SLAB_F0(BSDF), Dither);
-							const uint PackedHaziness32Bits		= SLAB_HAZINESS(BSDF);
+							uint Data0 = 0; 
+							uint Data1 = 0;
+							uint Data2 = SetHazinessRoughness(SLAB_HAZINESS(BSDF), SLAB_ROUGHNESS(BSDF));
+							PackDiffuseAlbedoF0(Dither, SLAB_DIFFUSEALBEDO(BSDF), SLAB_F0(BSDF), Data0, Data1);
 
-							// Out has 13bits free
-							uint Data0 = Out;
-							Data0 |= ( PackedRoughness7Bits			<< (HEADER_SINGLE_TOTAL_BIT_COUNT+0));		// Append 7 bits of roughness
-							Data0 |= ((PackedDiffuse19Bits & 0x3F)	<< (HEADER_SINGLE_TOTAL_BIT_COUNT+7));		// Append 6 low bits of Diffuse
-
-							uint Data1 = PackedDiffuse19Bits		>> 6;	// Remaining 13bits of Diffuse
-							     Data1|= PackedF019Bits				<< 13;	// 19 bits of F0
-
-							uint Data2 = PackedHaziness32Bits;
-
-							SUBSTRATE_STORE_UINT1(Data0);
+							SUBSTRATE_STORE_UINT1(Out | Data0);
 							SUBSTRATE_STORE_UINT1(Data1);
 							SUBSTRATE_STORE_UINT1(Data2);
 
@@ -600,70 +582,44 @@ void PackSubstrateOut(
 						}
 						else if (OptimisedLegacyMode == SINGLE_OPTLEGACYMODE_CLOTH)
 						{
-							const uint PackedDiffuse20Bits	= PackR7G7B6Gamma2(	SLAB_DIFFUSEALBEDO(BSDF)	, Dither);
-							const uint PackedF020Bits		= PackR7G7B6Gamma2(	SLAB_F0(BSDF)				, Dither);
-							const uint PackedFuzzColor20bits= PackR7G7B6Gamma2(	SLAB_FUZZ_COLOR(BSDF)		, Dither);
-							const uint PackedFuzzAmount8bits= PackR8(			SLAB_FUZZ_AMOUNT(BSDF));
-							// Roughness is stored in the TopLayer texture. Fuzz roughness is identical to slab roughness in this path.
-
-							// MRT0 
-							#if (32 - (HEADER_SINGLE_TOTAL_BIT_COUNT)) < 8
-							#error Not Enough space to store fuzz amount in slab in Slab - Single - Header.
-							#endif
-							SUBSTRATE_STORE_UINT1(Out					| ((PackedFuzzAmount8bits & 0xFF) << (HEADER_SINGLE_TOTAL_BIT_COUNT)));
-							// MRT1 20 PackedDiffuse and 8 high bits from PackedClearCoat
-							SUBSTRATE_STORE_UINT1(PackedDiffuse20Bits	| ((PackedFuzzColor20bits & 0xFFC00) << 10));
-							// MRT2 24 PackedF0      and 8 low  bits from PackedClearCoat
-							SUBSTRATE_STORE_UINT1(PackedF020Bits		| ((PackedFuzzColor20bits & 0x3FF) << 20));
+							uint Data0 = 0; 
+							uint Data1 = 0;
+							uint Data2 = PackRGBA8(float4(SLAB_FUZZ_COLOR(BSDF), SLAB_FUZZ_AMOUNT(BSDF)));
+							PackDiffuseAlbedoF0(Dither, SLAB_DIFFUSEALBEDO(BSDF), SLAB_F0(BSDF), Data0, Data1);
+
+							SUBSTRATE_STORE_UINT1(Out | Data0);
+							SUBSTRATE_STORE_UINT1(Data1);
+							SUBSTRATE_STORE_UINT1(Data2);
 
 							// Ensure the rest of the BSDF is not stored with the regular path
 							BSDF.State = 0;
 						}
 						else if (OptimisedLegacyMode == SINGLE_OPTLEGACYMODE_SSSWRAP || OptimisedLegacyMode == SINGLE_OPTLEGACYMODE_TWO_SIDED_SSSWRAP)
 						{
-							const float SSSWOpacity = SubstrateSubSurfaceGetWrapOpacityFromAnisotropy(SLAB_SSSPHASEANISOTROPY(BSDF));
-
-							const uint PackedDiffuse20Bits		= PackR7G7B6Gamma2(	SLAB_DIFFUSEALBEDO(BSDF)	, Dither);
-							const uint PackedF020Bits			= PackR7G7B6Gamma2(	SLAB_F0(BSDF)				, Dither);
-							const uint PackedSSSWOpacity7bits	= PackR7(			SSSWOpacity					, 0.5);	// We should not dither here because it would have a large scale effect.
-							const uint PackedSSSMFP30bits		= PackR10G10B10F(	SLAB_SSSMFP(BSDF));
-							// Roughness is stored in the TopLayer texture.
-
-							const uint PackedSSSMFP12BitsA = (PackedSSSMFP30bits      ) & 0xFFF; // 12 lower bits
-							const uint PackedSSSMFP12BitsB = (PackedSSSMFP30bits >> 12) & 0xFFF; // next 12 lower bits
-							const uint PackedSSSMFP6BitsC  = (PackedSSSMFP30bits >> 24) & 0x03F; // 6 higher bits
-
-							// MRT0 
-							#if (32 - (HEADER_SINGLEENCODING_BIT_COUNT + HEADER_SINGLE_OPTLEGACYMODE_BIT_COUNT)) < (7+6)
-							#error Not Enough space to store wrapped SSS opacity in slab in Slab - Single - Header.
-							#endif
-							SUBSTRATE_STORE_UINT1(Out					| ((PackedSSSWOpacity7bits & 0x7F) << (HEADER_SINGLEENCODING_BIT_COUNT + HEADER_SINGLE_OPTLEGACYMODE_BIT_COUNT)) | (PackedSSSMFP6BitsC << (7 + HEADER_SINGLEENCODING_BIT_COUNT + HEADER_SINGLE_OPTLEGACYMODE_BIT_COUNT)));
-							// MRT1 20 PackedDiffuse and 8 high bits from PackedClearCoat
-							SUBSTRATE_STORE_UINT1(PackedDiffuse20Bits  | (PackedSSSMFP12BitsA << 20));
-							// MRT2 24 PackedF0      and 8 low  bits from PackedClearCoat
-							SUBSTRATE_STORE_UINT1(PackedF020Bits		| (PackedSSSMFP12BitsB << 20));
+							const uint PackedSSSWOpacity7bits = PackR7(SubstrateSubSurfaceGetWrapOpacityFromAnisotropy(SLAB_SSSPHASEANISOTROPY(BSDF)), 0.5);	// We should not dither here because it would have a large scale effect.
+
+							uint Data0 = 0; 
+							uint Data1 = 0;
+							uint Data2 = PackR10G10B10F(SLAB_SSSMFP(BSDF));
+							PackDiffuseAlbedoF0(Dither, SLAB_DIFFUSEALBEDO(BSDF), SLAB_F0(BSDF), Data0, Data1);
+
+							SUBSTRATE_STORE_UINT1(Out | Data0 | (BitFieldExtractU32(PackedSSSWOpacity7bits, 5, 0)<<(HEADER_SINGLEENCODING_BIT_COUNT + HEADER_SINGLE_OPTLEGACYMODE_BIT_COUNT)));
+							SUBSTRATE_STORE_UINT1(Data1);
+							SUBSTRATE_STORE_UINT1(Data2 | (BitFieldExtractU32(PackedSSSWOpacity7bits, 2, 5)<<30) );
 
 							// Ensure the rest of the BSDF is not stored with the regular path
 							BSDF.State = 0;
 						}
 						else if (OptimisedLegacyMode == SINGLE_OPTLEGACYMODE_SSSPROFILE)
 						{
-							const uint PackedDiffuse20Bits		= PackR7G7B6Gamma2(	SLAB_DIFFUSEALBEDO(BSDF)	, Dither);
-							const uint PackedF020Bits			= PackR7G7B6Gamma2(	SLAB_F0(BSDF)				, Dither);
-
-							// When creating a slab with GetSubstrateSlabBSDF(), the SSS profile's dual roughness is expanded into haziness. 
-							// But when storing data with SINGLE_OPTLEGACYMODE_SSSPROFILE we need to store the original roughness. 
-							// The dual roughnesses will be recomputed based on the SSS profile data at load time. To retrieve the 
-							// original roughness, we need to invert the computation done into GetSubstrateSlabBSDF()
-							const uint SubsurfaceProfileUInt = SubstrateSubsurfaceProfileIdTo8bits(SLAB_SSSPROFILEID(BSDF));
-							const float OriginalRoughness = GetSubsurfaceProfileOriginalRoughness(SubsurfaceProfileUInt, SLAB_ROUGHNESS(BSDF), SLAB_SSSPROFILERADIUSSCALE(BSDF));
-
-							// MRT0 
-							SUBSTRATE_STORE_UINT1(Out					| (PackR8(OriginalRoughness) << 24));
-							// MRT1 20 PackedDiffuse and Profile radius scale
-							SUBSTRATE_STORE_UINT1(PackedDiffuse20Bits  | (PackR8(SLAB_SSSPROFILERADIUSSCALE(BSDF)) << 24));
-							// MRT2 24 PackedF0      and Profile ID
-							SUBSTRATE_STORE_UINT1(PackedF020Bits		| (PackR8(SLAB_SSSPROFILEID(BSDF)) << 24));
+							uint Data0 = 0; 
+							uint Data1 = 0;
+							uint Data2 = PackRGBA8(float4(SLAB_SSSPROFILEID(BSDF), SLAB_SSSPROFILERADIUSSCALE(BSDF), 0, 0)) | ((SLAB_HAZINESS(BSDF) & 0xFFFF)<<16);
+							PackDiffuseAlbedoF0(Dither, SLAB_DIFFUSEALBEDO(BSDF), SLAB_F0(BSDF), Data0, Data1);
+
+							SUBSTRATE_STORE_UINT1(Out | Data0);
+							SUBSTRATE_STORE_UINT1(Data1);
+							SUBSTRATE_STORE_UINT1(Data2);
 
 							// Ensure the rest of the BSDF is not stored with the regular path
 							BSDF.State = 0;
diff --git a/Engine/Shaders/Private/Substrate/SubstrateMaterialClassification.usf b/Engine/Shaders/Private/Substrate/SubstrateMaterialClassification.usf
--- a/Engine/Shaders/Private/Substrate/SubstrateMaterialClassification.usf
+++ b/Engine/Shaders/Private/Substrate/SubstrateMaterialClassification.usf
@@ -175,7 +175,6 @@ void TileMainCS(uint2 DispatchThreadId : SV_DispatchThreadID, uint LinearIndex :
 			bContainsScreenSpaceSubsurfaceScattering = true;
 			if (bIsLegacyWrapOrWrapThin)
 			{
-				const uint PackedSSSWOpacity7bits = (PackedHeader >> (HEADER_SINGLEENCODING_BIT_COUNT + HEADER_SINGLE_OPTLEGACYMODE_BIT_COUNT)) & 0x7F;
 				FSubstrateSubsurfaceHeader SSSHeader = (FSubstrateSubsurfaceHeader)0;
 				if (OptimisedLegacyMode == SINGLE_OPTLEGACYMODE_SSSWRAP)
 				{
@@ -185,15 +184,23 @@ void TileMainCS(uint2 DispatchThreadId : SV_DispatchThreadID, uint LinearIndex :
 				{
 					SubstrateSubSurfaceHeaderSetSSSType(SSSHeader, SSS_TYPE_TWO_SIDED_WRAP);
 				}
-				SubstrateSubSurfaceHeaderSetWrapOpacity(SSSHeader, UnpackR7(PackedSSSWOpacity7bits));
+
+				const uint Data0 = PackedHeader;
+				const uint Data2 = MaterialTextureArrayUAV[uint3(PixelCoord, 2)];
+				const uint PackedSSSWOpacity7bits = BitFieldExtractU32(Data0, 5, HEADER_SINGLEENCODING_BIT_COUNT + HEADER_SINGLE_OPTLEGACYMODE_BIT_COUNT) | (BitFieldExtractU32(Data2, 2, 30) << 5);
+
+				const float Opacity = UnpackR7(PackedSSSWOpacity7bits);
+				SubstrateSubSurfaceHeaderSetWrapOpacity(SSSHeader, Opacity);
 				SubstrateStoreSubsurfaceHeader(MaterialTextureArrayUAV, FirstSliceStoringSubstrateSSSData, PixelCoord, SSSHeader.Bytes);
 			}
 			else if (bIsLegacySSSProfile)
 			{
 				const uint Data1 = MaterialTextureArrayUAV[uint3(PixelCoord, 1)];
 				const uint Data2 = MaterialTextureArrayUAV[uint3(PixelCoord, 2)];
-				float  RadiusScale = UnpackR8(Data1 >> 24);
-				float  ProfileId   = UnpackR8(Data2 >> 24);
+
+				const float4 CustomData = UnpackRGBA8(Data2);
+				float ProfileId = CustomData.x;
+				float RadiusScale = CustomData.y;
 
 				const uint PackedDiffuse20Bits = (Data1 & 0xFFFFF);
 				const float3 BaseColor = UnpackR7G7B6Gamma2(PackedDiffuse20Bits);
