Remove scale from vertex parameter tangent to world for gpu skin vertex factory.
Before this change the VertexNormalWS material node would apply component scale to the normal without normalization.
This is inconsistent with the skin cache path where the value will be normalized.
#tests Ran EngineTests for Rendering and Niagara with no new regressions.

diff --git a/Engine/Shaders/Private/GpuSkinVertexFactory.ush b/Engine/Shaders/Private/GpuSkinVertexFactory.ush
--- a/Engine/Shaders/Private/GpuSkinVertexFactory.ush
+++ b/Engine/Shaders/Private/GpuSkinVertexFactory.ush
@@ -311,7 +311,13 @@ FMaterialVertexParameters GetMaterialVertexParameters(
 	Result.PositionPrimitiveSpace = Result.PositionInstanceSpace; // No support for instancing, so instance == primitive
 
 	Result.VertexColor = Intermediates.Color;
-	Result.TangentToWorld = mul(TangentToLocal, DFToFloat3x3(Intermediates.LocalToWorld));
+
+	float3x3 LocalToWorldNoScale = DFToFloat3x3(Intermediates.LocalToWorld);
+	LocalToWorldNoScale[0] *= Intermediates.InvNonUniformScale.x;
+	LocalToWorldNoScale[1] *= Intermediates.InvNonUniformScale.y;
+	LocalToWorldNoScale[2] *= Intermediates.InvNonUniformScale.z;
+	Result.TangentToWorld = mul(TangentToLocal, LocalToWorldNoScale);
+
 	Result.PreSkinnedPosition = Intermediates.UnpackedPosition.xyz;
 	Result.PreSkinnedNormal = Input.TangentZ.xyz;
 
