VulkanRHI: Match array size in code.  Fix unbound parameter warning on startup.

#rnx

diff --git a/Engine/Shaders/Private/Landscape/LandscapeEditLayersUtils.usf b/Engine/Shaders/Private/Landscape/LandscapeEditLayersUtils.usf
--- a/Engine/Shaders/Private/Landscape/LandscapeEditLayersUtils.usf
+++ b/Engine/Shaders/Private/Landscape/LandscapeEditLayersUtils.usf
@@ -61,8 +61,7 @@ Texture2D<float4> InSourceTexture;
 	op(24) \
 	op(25) \
 	op(26) \
-	op(27) \
-	op(28) 
+	op(27) 
 #else // VULKAN_PROFILE_SM6 
 #define FOREACH_SOURCETEXTURE(op) \
 	op(0) \
