Temp fix for Instanced Stereo not showing proper scene texture in second eye
Using resolved view instead of view to properly account for right eye when transforming to buffer UVs

diff --git a/Engine/Shaders/Private/Common.ush b/Engine/Shaders/Private/Common.ush
--- a/Engine/Shaders/Private/Common.ush
+++ b/Engine/Shaders/Private/Common.ush
@@ -1639,8 +1639,13 @@ float2 BufferUVToViewportUV(float2 BufferUV)
 
 float2 ViewportUVToBufferUV(float2 ViewportUV)
 {
+#if INSTANCED_STEREO && !MOBILE_MULTI_VIEW
+	float2 PixelPos = ViewportUV * ResolvedView.ViewSizeAndInvSize.xy;
+	return (PixelPos + ResolvedView.ViewRectMin.xy) * View.BufferSizeAndInvSize.zw;
+#else
 	float2 PixelPos = ViewportUV * View.ViewSizeAndInvSize.xy;
 	return (PixelPos + View.ViewRectMin.xy) * View.BufferSizeAndInvSize.zw;
+#endif
 }
 
 // Maps standard viewport UV to screen position.
