Heterogeneous Volumes: Initialize LWCData with ResolvedView parameters

Provides extended shader support, such as proper evaluation of CameraPosition.

#tests Modified EngineTest material to threshold against distance to CameraPosition

diff --git a/Engine/Shaders/Private/HeterogeneousVolumes/HeterogeneousVolumesLiveShadingPipeline.usf b/Engine/Shaders/Private/HeterogeneousVolumes/HeterogeneousVolumesLiveShadingPipeline.usf
--- a/Engine/Shaders/Private/HeterogeneousVolumes/HeterogeneousVolumesLiveShadingPipeline.usf
+++ b/Engine/Shaders/Private/HeterogeneousVolumes/HeterogeneousVolumesLiveShadingPipeline.usf
@@ -257,8 +257,11 @@ FVolumeSampleContext CreateVolumeSampleContext(
 	float MipLevel
 )
 {
+	ResolvedView = ResolveView();
+
 	FVolumeSampleContext VolumeSampleContext;
 	VolumeSampleContext.MaterialParameters = MakeInitializedMaterialPixelParameters();
+	VolumeSampleContext.MaterialParameters.LWCData = MakeMaterialLWCData(VolumeSampleContext.MaterialParameters);
 	//VolumeSampleContext.MaterialParameters.PrimitiveId = PrimitiveId;
 	VolumeSampleContext.MaterialParameters.AbsoluteWorldPosition = DFPromote(WorldPosition);
 	VolumeSampleContext.MaterialParameters.LWCData.AbsoluteWorldPosition = WSPromote(WorldPosition);
@@ -967,6 +970,8 @@ void RenderSingleScatteringWithLiveShadingCS(
 	uint2 DispatchThreadId : SV_DispatchThreadID
 )
 {
+	ResolvedView = ResolveView();
+
 	RenderSingleScatteringWithLiveShading(DispatchThreadId);
 }
 
@@ -983,6 +988,8 @@ void RenderSingleScatteringWithLiveShadingIndirectCS(
 	uint2 GroupThreadId : SV_GroupThreadID
 )
 {
+	ResolvedView = ResolveView();
+
 	FScreenTile ScreenTile = ScreenTileBuffer[GroupId];
 	uint2 GroupId_2D = uint2(ScreenTile.Id % GroupCount.x, ScreenTile.Id / GroupCount.x);
 	uint2 DispatchThreadId = GroupId_2D * THREADGROUP_SIZE_2D + GroupThreadId;
diff --git a/Engine/Shaders/Private/HeterogeneousVolumes/HeterogeneousVolumesLiveShadingShadows.usf b/Engine/Shaders/Private/HeterogeneousVolumes/HeterogeneousVolumesLiveShadingShadows.usf
--- a/Engine/Shaders/Private/HeterogeneousVolumes/HeterogeneousVolumesLiveShadingShadows.usf
+++ b/Engine/Shaders/Private/HeterogeneousVolumes/HeterogeneousVolumesLiveShadingShadows.usf
@@ -166,8 +166,11 @@ FVolumeSampleContext CreateVolumeSampleContext(
 	float MipLevel
 )
 {
+	ResolvedView = ResolveView();
+
 	FVolumeSampleContext VolumeSampleContext;
 	VolumeSampleContext.MaterialParameters = MakeInitializedMaterialPixelParameters();
+	VolumeSampleContext.MaterialParameters.LWCData = MakeMaterialLWCData(VolumeSampleContext.MaterialParameters);
 	//VolumeSampleContext.MaterialParameters.PrimitiveId = PrimitiveId;
 	VolumeSampleContext.MaterialParameters.AbsoluteWorldPosition = DFPromote(WorldPosition);
 	VolumeSampleContext.MaterialParameters.LWCData.AbsoluteWorldPosition = WSPromote(WorldPosition);
diff --git a/Engine/Shaders/Private/HeterogeneousVolumes/HeterogeneousVolumesMaterialBakingPipeline.usf b/Engine/Shaders/Private/HeterogeneousVolumes/HeterogeneousVolumesMaterialBakingPipeline.usf
--- a/Engine/Shaders/Private/HeterogeneousVolumes/HeterogeneousVolumesMaterialBakingPipeline.usf
+++ b/Engine/Shaders/Private/HeterogeneousVolumes/HeterogeneousVolumesMaterialBakingPipeline.usf
@@ -50,6 +50,8 @@ void HeterogeneousVolumesBakeMaterialCS(
 	uint3 DispatchThreadId : SV_DispatchThreadID
 )
 {
+	ResolvedView = ResolveView();
+
 	uint3 VoxelIndex = DispatchThreadId;
 	if (all(VoxelIndex < VolumeResolution))
 	{
@@ -60,6 +62,7 @@ void HeterogeneousVolumesBakeMaterialCS(
 
 		// Setup evaluation context
 		FMaterialPixelParameters MaterialParameters = MakeInitializedMaterialPixelParameters();
+		MaterialParameters.LWCData = MakeMaterialLWCData(MaterialParameters);
 		MaterialParameters.PrimitiveId = PrimitiveId;
 		MaterialParameters.AbsoluteWorldPosition = DFPromote(WorldPosition);
 		MaterialParameters.LWCData.AbsoluteWorldPosition = WSPromote(WorldPosition);
@@ -209,6 +212,8 @@ void RasterizeBottomLevelOrthoGridCS(
 	uint3 GroupThreadId : SV_GroupThreadID
 )
 {
+	ResolvedView = ResolveView();
+
 	uint RasterTileIndex = GetUnWrappedDispatchGroupId(GroupId);
 	if (RasterTileIndex >= RasterTileAllocatorBuffer[0])
 	{
@@ -270,6 +275,7 @@ void RasterizeBottomLevelOrthoGridCS(
 			{
 				// Setup evaluation context
 				FMaterialPixelParameters MaterialParameters = MakeInitializedMaterialPixelParameters();
+				MaterialParameters.LWCData = MakeMaterialLWCData(MaterialParameters);
 				MaterialParameters.PrimitiveId = PrimitiveId;
 				MaterialParameters.AbsoluteWorldPosition = DFPromote(WorldPosition);
 				MaterialParameters.LWCData.AbsoluteWorldPosition = WSPromote(WorldPosition);
@@ -450,6 +456,8 @@ void RasterizeBottomLevelFrustumGridCS(
 	uint3 GroupThreadId : SV_GroupThreadID
 )
 {
+	ResolvedView = ResolveView();
+
 	uint RasterTileIndex = GetUnWrappedDispatchGroupId(GroupId);
 	if (RasterTileIndex >= RasterTileAllocatorBuffer[0])
 	{
@@ -483,6 +491,7 @@ void RasterizeBottomLevelFrustumGridCS(
 	if (WorldPositionIntersectsPrimitive(WorldPosition))
 	{
 		FMaterialPixelParameters MaterialParameters = MakeInitializedMaterialPixelParameters();
+		MaterialParameters.LWCData = MakeMaterialLWCData(MaterialParameters);
 		MaterialParameters.PrimitiveId = PrimitiveId;
 		MaterialParameters.AbsoluteWorldPosition = DFPromote(WorldPosition);
 		MaterialParameters.LWCData.AbsoluteWorldPosition = WSPromote(WorldPosition);
