Fix for RVT page corruption on some platforms.
Adding explicit clamp of DestRect array index calculation encourages compiler to avoid load scheduling issue.
#tests Cooked build with YCoCg RVT formats

diff --git a/Engine/Shaders/Private/VirtualTextureCompress.usf b/Engine/Shaders/Private/VirtualTextureCompress.usf
--- a/Engine/Shaders/Private/VirtualTextureCompress.usf
+++ b/Engine/Shaders/Private/VirtualTextureCompress.usf
@@ -96,8 +96,10 @@ uint3 GetDestPos(uint LayerIndex, uint3 ThreadId) { return ThreadId; }
 #else
 uint2 GetDestPos(uint LayerIndex, uint3 ThreadId) 
 {
-	uint Index = (ThreadId.z * MAX_DST_LAYERS + LayerIndex) * 2;
-	uint2 BasePos = uint2(GET_SCALAR_ARRAY_ELEMENT(DestPos, Index), GET_SCALAR_ARRAY_ELEMENT(DestPos, (Index + 1)));
+	// clamp() shouldn't be necessary here, but has the side effect of fixing a load scheduling bug for one shader compiler (UE-304803).
+	uint BatchIndex = clamp(ThreadId.z, 0, MAX_BATCH_SIZE - 1);
+	uint ArrayIndex = (BatchIndex * MAX_DST_LAYERS + LayerIndex) * 2;
+	uint2 BasePos = uint2(GET_SCALAR_ARRAY_ELEMENT(DestPos, ArrayIndex), GET_SCALAR_ARRAY_ELEMENT(DestPos, (ArrayIndex + 1)));
 	return BasePos + ThreadId.xy; 
 }
 #endif
