Fix Substrate F0/DiffuseAlbedo encoding for complex material which had mismatching logic between packing and unpacking.

diff --git a/Engine/Shaders/Private/Substrate/Substrate.ush b/Engine/Shaders/Private/Substrate/Substrate.ush
--- a/Engine/Shaders/Private/Substrate/Substrate.ush
+++ b/Engine/Shaders/Private/Substrate/Substrate.ush
@@ -272,6 +272,7 @@ struct FSubstrateBSDF
 
 	void SubstrateSetBSDFRoughness(in float Roughness);
 
+	void PackSingleOrComplexPathSlabBSDF(float Dither, inout uint OutData1, inout uint OutData2);
 	void UnpackSingleOrComplexPathSlabBSDF(uint Data0, uint Data1);
 
 	void PostProcessSlabBeforeLighting(bool bIsSubstrateOpaqueMaterial);
@@ -4189,10 +4190,19 @@ void UnpackDiffuseAlbedoF0(uint InData0, uint InData1, inout float3 OutDiffuseAl
 	UnpackDiffuseAlbedoF0Roughness(InData0, InData1, OutDiffuseAlbedo, OutF0, OutDummyRoughness);
 }
 
+void FSubstrateBSDF::PackSingleOrComplexPathSlabBSDF(float Dither, inout uint OutData1, inout uint OutData2)
+{
+	const uint PackedDiffuse20Bits	= PackR7G7B6Gamma2(SLAB_DIFFUSEALBEDO(this), Dither);
+	const uint PackedF020Bits		= PackR7G7B6Gamma2(SLAB_F0(this), Dither);
+
+	OutData1 =(BitFieldExtractU32(PackedF020Bits, 12, 0) << 20) | PackedDiffuse20Bits;
+	OutData2 = BitFieldExtractU32(PackedF020Bits, 8, 12) | PackRGBA8(float4(0.f, SLAB_ROUGHNESS(this), (SLAB_ANISOTROPY(this) + 1.f) * 0.5f, (SLAB_SSSPHASEANISOTROPY(this) + 1.f) * 0.5f));
+}
+
 void FSubstrateBSDF::UnpackSingleOrComplexPathSlabBSDF(uint Data1, uint Data2)
 {
-	const uint PackedDiffuse20Bits	= BitFieldExtractU32(Data1, 20, 12);
-	const uint PackedF020Bits		= BitFieldExtractU32(Data1, 12,  0) | (BitFieldExtractU32(Data2,  8, 0)<<24);
+	const uint PackedDiffuse20Bits	= BitFieldExtractU32(Data1, 20, 0);
+	const uint PackedF020Bits		= BitFieldExtractU32(Data1, 12,20) | (BitFieldExtractU32(Data2, 8, 0)<<12);
 
 	SLAB_DIFFUSEALBEDO(this)	= UnpackR7G7B6Gamma2(PackedDiffuse20Bits);
 	SLAB_F0(this)				= UnpackR7G7B6Gamma2(PackedF020Bits);
diff --git a/Engine/Shaders/Private/Substrate/SubstrateExport.ush b/Engine/Shaders/Private/Substrate/SubstrateExport.ush
--- a/Engine/Shaders/Private/Substrate/SubstrateExport.ush
+++ b/Engine/Shaders/Private/Substrate/SubstrateExport.ush
@@ -711,15 +711,14 @@ void PackSubstrateOut(
 						{
 						case SUBSTRATE_BSDF_TYPE_SLAB:
 						{
-							// Now we pack diffuse and F0 in a special way: both encoded as R7G7B6 with gamma 2.0. The low bit cound will be hidden by dithering and TAA.
-							// We pack 32 bits in the first uint and the remaining bits are pack is the lowest significant bit of the second uint.
-							const uint PackedDiffuse20Bits = PackR7G7B6Gamma2(SLAB_DIFFUSEALBEDO(BSDF), Dither);
-							const uint PackedF020Bits = PackR7G7B6Gamma2(SLAB_F0(BSDF), Dither);
-							const uint PackedData32Bits = ((PackedDiffuse20Bits << 12) & 0xFFFFF000) | (PackedF020Bits & 0xFFF);
-							const uint PackedData8Bits = (PackedF020Bits >> 12) & 0xFF;
-
-							SUBSTRATE_STORE_UINT1(PackedData32Bits);
-							SUBSTRATE_STORE_UINT1(PackedData8Bits | PackRGBA8(float4(0.0f, SLAB_ROUGHNESS(BSDF), (SLAB_ANISOTROPY(BSDF) + 1.f) * 0.5f, (SLAB_SSSPHASEANISOTROPY(BSDF) + 1.f) * 0.5f)));
+							// DiffuseAlbedo / F0 / Roughness / Aniso / SSSAniso
+							{
+								uint Data1 = 0; 
+								uint Data2 = 0;
+								BSDF.PackSingleOrComplexPathSlabBSDF(Dither, Data1, Data2);
+								SUBSTRATE_STORE_UINT1(Data1);
+								SUBSTRATE_STORE_UINT1(Data2);
+							}
 
 							if (BSDF_GETHASF90(BSDF) || BSDF_GETHASHAZINESS(BSDF))
 							{
