Mobile: Output 1 from tonemapping pass into alpha channel as mobile device composers expect that

diff --git a/Engine/Shaders/Private/PostProcessTonemap.usf b/Engine/Shaders/Private/PostProcessTonemap.usf
--- a/Engine/Shaders/Private/PostProcessTonemap.usf
+++ b/Engine/Shaders/Private/PostProcessTonemap.usf
@@ -486,7 +486,13 @@ float4 TonemapCommonPS(
 	}
 	#endif // USE_SHARPEN
 
+#if FEATURE_LEVEL == FEATURE_LEVEL_ES3_1
+	// Mobile device composers expect 1.0 in alpha
+	float4 OutColor = float4(0.0, 0.0, 0.0, 1.0);
+#else
 	float4 OutColor = 0;
+#endif
+
 #if USE_GAMMA_ONLY
 	{
 		#if LOCAL_EXPOSURE_MODE != LOCAL_EXPOSURE_MODE_NONE
