Fix for crash with AmbienCubemapComposite and instance stereo.

#tests PC editor dx12
#rnx

diff --git a/Engine/Shaders/Private/AmbientCubemapComposite.usf b/Engine/Shaders/Private/AmbientCubemapComposite.usf
--- a/Engine/Shaders/Private/AmbientCubemapComposite.usf
+++ b/Engine/Shaders/Private/AmbientCubemapComposite.usf
@@ -53,7 +53,9 @@ float3 GetEnvSpecularLighting(
 #include "Substrate/SubstrateLightingCommon.ush"
 
 void MainPS(
-	float4 SvPosition : SV_POSITION, 
+	in noperspective float4 UVAndScreenPos : TEXCOORD0,
+	in FStereoPSInput StereoInput,
+	in float4 SvPosition : SV_POSITION,
 	out float4 OutColor : SV_Target0
 	#if SUBSTRATE_OPAQUE_ROUGH_REFRACTION_ENABLED
 	, out float3 OutOpaqueRoughRefractionSceneColor : SV_Target1
@@ -61,6 +63,8 @@ void MainPS(
 	#endif
 	)
 {
+	StereoSetupPS(StereoInput);
+
 	const int2 PixelPos = int2(SvPosition.xy);
 	const float2 BufferUV = SvPositionToBufferUV(SvPosition);
 	const float2 ScreenPosition = SvPositionToScreenPosition(SvPosition).xy;
@@ -439,8 +443,14 @@ float3 ImageBasedLightingHair( FGBufferData GBuffer, float3 V, float3 N, uint2 R
 	return Lighting;
 }
 
-void MainPS(in noperspective float4 UVAndScreenPos : TEXCOORD0, float4 SvPosition : SV_POSITION, out float4 OutColor : SV_Target0)
+void MainPS(
+	in noperspective float4 UVAndScreenPos : TEXCOORD0, 
+	in FStereoPSInput StereoInput,
+	in float4 SvPosition : SV_POSITION,
+	out float4 OutColor : SV_Target0)
 {
+	StereoSetupPS(StereoInput);
+
 	float2 BufferUV = UVAndScreenPos.xy;
 	float2 ScreenSpacePos = UVAndScreenPos.zw;
 	int2 PixelPos = int2(SvPosition.xy);
