Substrate - follow up to 44127087 where Substrate was not tested correctly.

#tests PC editor dx12
#rnx

diff --git a/Engine/Shaders/Private/AmbientCubemapComposite.usf b/Engine/Shaders/Private/AmbientCubemapComposite.usf
--- a/Engine/Shaders/Private/AmbientCubemapComposite.usf
+++ b/Engine/Shaders/Private/AmbientCubemapComposite.usf
@@ -53,8 +53,6 @@ float3 GetEnvSpecularLighting(
 #include "Substrate/SubstrateLightingCommon.ush"
 
 void MainPS(
-	in noperspective float4 UVAndScreenPos : TEXCOORD0,
-	in FStereoPSInput StereoInput,
 	in float4 SvPosition : SV_POSITION,
 	out float4 OutColor : SV_Target0
 	#if SUBSTRATE_OPAQUE_ROUGH_REFRACTION_ENABLED
@@ -63,8 +61,6 @@ void MainPS(
 	#endif
 	)
 {
-	StereoSetupPS(StereoInput);
-
 	const int2 PixelPos = int2(SvPosition.xy);
 	const float2 BufferUV = SvPositionToBufferUV(SvPosition);
 	const float2 ScreenPosition = SvPositionToScreenPosition(SvPosition).xy;
