Volumetric Cloud - Fixed bad disoclusion happening for animated object over background.

#tests PCeditor, green cloud engine tests.
#rnx
[RN] Volumetric Cloud - Fixed bad disoclusion happening for animated object over background in some cases..

diff --git a/Engine/Shaders/Private/VolumetricRenderTarget.usf b/Engine/Shaders/Private/VolumetricRenderTarget.usf
--- a/Engine/Shaders/Private/VolumetricRenderTarget.usf
+++ b/Engine/Shaders/Private/VolumetricRenderTarget.usf
@@ -372,9 +372,33 @@ void ReconstructVolumetricRenderTargetPS(
 			}
 			else if (bValidPreviousUVs)
 			{
-				RGBA = SafeSamplePreviousFrameVolumetricTexture(PrevScreenUVs);
-				RGBA2 = SafeSamplePreviousFrameVolumetricSecondaryTexture(PrevScreenUVs);
-				DepthData = HistoryDepthData;
+				float4 ReprojRGBA = SafeSamplePreviousFrameVolumetricTexture(PrevScreenUVs);
+				float4 ReprojRGBA2 = SafeSamplePreviousFrameVolumetricSecondaryTexture(PrevScreenUVs);
+				FDepthData ReprojDepthData = HistoryDepthData;
+
+				// Min and Max the reprojected data to avoid reprojection issues. But this is too much constraint to use only the new sample as it causes jittering...
+				//if(false)
+				//{
+				//	float4 MinRGBA = min(NewRGBA, NewRGBA2);
+				//	float4 MaxRGBA = max(NewRGBA, NewRGBA2);
+				//	ReprojRGBA = clamp(ReprojRGBA, MinRGBA, MaxRGBA);
+				//	ReprojRGBA2 = clamp(ReprojRGBA2, MinRGBA, MaxRGBA);
+				//}
+
+				// Ignore reprojected samples with too large of a min depth difference that have been missed by dissoclusion.
+				// Those samples are considered as an edge so we skip because it might reproject bad bilinearly filtered data over the background and use the new incoming data.
+				// This is an attempt to fix missing cloud over sky background around close moving meshes. 
+				// It should hopefully not disrupt forced bilinear filtering of background clouds (when depth is > MinimumDistanceKmToDisableDisoclusion).
+				if (abs(NewDepthData.MinMaxViewDepthKm.x - HistoryDepthData.MinMaxViewDepthKm.x) > 1.0f)
+				{
+					ReprojRGBA = NewRGBA;// +float4(0.2, 0, 0, 0);
+					ReprojRGBA2 = NewRGBA2;
+					ReprojDepthData = NewDepthData;
+				}
+
+				RGBA = ReprojRGBA;
+				RGBA2 = ReprojRGBA2;
+				DepthData = ReprojDepthData;
 
 				if (!(all(IsFinite(RGBA)) && all(IsFinite(float2(DepthData.CloudFrontDepthFromViewKm, DepthData.SceneDepthFromViewKm)))))
 				{
