VolumetricCloud - fixed nan appearing in cloud trace for the far luminance/transmittance data.

#rnx
#tests PC editor
[RN] VolumetricCloud - fixed nan appearing in cloud trace for the far luminance/transmittance data.

diff --git a/Engine/Shaders/Private/VolumetricRenderTarget.usf b/Engine/Shaders/Private/VolumetricRenderTarget.usf
--- a/Engine/Shaders/Private/VolumetricRenderTarget.usf
+++ b/Engine/Shaders/Private/VolumetricRenderTarget.usf
@@ -400,7 +400,7 @@ void ReconstructVolumetricRenderTargetPS(
 				RGBA2 = ReprojRGBA2;
 				DepthData = ReprojDepthData;
 
-				if (!(all(IsFinite(RGBA)) && all(IsFinite(float2(DepthData.CloudFrontDepthFromViewKm, DepthData.SceneDepthFromViewKm)))))
+				if (!(all(IsFinite(RGBA)) && all(IsFinite(RGBA2)) && all(IsFinite(float2(DepthData.CloudFrontDepthFromViewKm, DepthData.SceneDepthFromViewKm)))))
 				{
 					RGBA = float4(0.0f, 0.0f, 0.0f, 1.0f);
 					RGBA2 = float4(0.0f, 0.0f, 0.0f, 1.0f);
