Path Tracer: Improve random sequence quality

Improve Sobol matrices for dimensions (2,3) so that this pair is a (0,2) sequence and the overall set forms a (0,4) sequence.

Improve the quality of Owen Scrambling for a similar instruction count by improving how the seed is incorporated and by a better choice of instructions to mix the bits.

Change EvolveSeed to a simpler LCG to evolve the seed across calls (combined with the owen scrambling change, this keeps the instruction count similar overall).

Add a precomputed texture version of the space filling curve used for error diffusion sampling using a Sierpinski curve which avoids some of the axis-aligned artifacts that can occur with the hilbert curve.

diff --git a/Engine/Shaders/Private/PathTracing/PathTracingCore.ush b/Engine/Shaders/Private/PathTracing/PathTracingCore.ush
--- a/Engine/Shaders/Private/PathTracing/PathTracingCore.ush
+++ b/Engine/Shaders/Private/PathTracing/PathTracingCore.ush
@@ -2,6 +2,9 @@
 
 #pragma once
 
+#define RANDSEQ_VERSION 1
+#define RANDSEQ_SFC_TEXTURE 1
+
 #define RANDSEQ_UNROLL_SOBOL PATH_TRACER_USE_COMPACTION // unrolling sobol sampler is best when we can guarantee high occupancy
 
 #define USE_PATH_TRACING_LIGHT_GRID	1
diff --git a/Engine/Shaders/Private/PathTracing/PathTracingDebug.usf b/Engine/Shaders/Private/PathTracing/PathTracingDebug.usf
--- a/Engine/Shaders/Private/PathTracing/PathTracingDebug.usf
+++ b/Engine/Shaders/Private/PathTracing/PathTracingDebug.usf
@@ -1,5 +1,8 @@
 // Copyright Epic Games, Inc. All Rights Reserved.
 
+#define RANDSEQ_VERSION 1
+#define RANDSEQ_SFC_TEXTURE 1
+
 #define LIGHT_LOOP_METHOD 	2 // 0: stochastic (like ref PT) 1: loop over all lights 2: RIS (take N samples but trace a single shadow ray)
 
 #define USE_BSDF_SAMPLING 1 // 0: light sampling only (no MIS) 1: light sampling and BSDF sampling (with MIS)
@@ -47,11 +50,11 @@ int EnableEmissive;
 
 uint4 GridHash(uint4 P)
 {
-	const uint4 S = uint4(14, 15, 16, 17);
-	const uint K = 1664525u, O = 1013904223u; // pcg4d constants
-	P ^= P >> S; P = P * K + O; P ^= P.yzwx ^ P.zwxy;
-	P ^= P >> S; P = P * K + O; P ^= P.yzwx ^ P.zwxy;
-	P ^= P >> S; P = P * K + O; P ^= P.yzwx ^ P.zwxy;
+	const uint4 S = uint4(16, 17, 18, 19);
+	const uint K = 1159349557u;
+	P ^= P >> S; P *= K; P ^= P.yzwx ^ P.zwxy;
+	P ^= P >> S; P *= K; P ^= P.yzwx ^ P.zwxy;
+	P ^= P >> S; P *= K; P ^= P.yzwx ^ P.zwxy;
 	return P;
 }
 
@@ -198,8 +201,7 @@ FPathTracingResult PathTracingDebugCore(FRayDesc Ray, uint2 LaunchIndex, int Sam
 	float3 PathThroughput = 1;
 	float PathRoughness = 0;
 	uint ScatterType = PATHTRACER_SCATTER_CAMERA;
-	RandomSequence RandSequence;
-	RandomSequence_Initialize(RandSequence, LaunchIndex, SampleIndex, View.StateFrameIndex, SamplesPerPixel);
+	FRandomSequence RandSequence = RandomSequenceCreate(uint3(LaunchIndex, View.StateFrameIndex), SampleIndex, SamplesPerPixel);
 
 	// State about the markov chain selection that needs to be kept from one bounce to the next
 	float ScoresSum = 0;
@@ -228,7 +230,7 @@ FPathTracingResult PathTracingDebugCore(FRayDesc Ray, uint2 LaunchIndex, int Sam
 		const uint MissShaderIndex = 0;
 
 		FPackedPathTracingPayload PackedPayload = InitPathTracingPayload(ScatterType, PathRoughness);
-		PackedPayload.SetStochasticSlabRand(RandomSequence_GenerateSample1D(RandSequence));
+		PackedPayload.SetStochasticSlabRand(RandSequence.Get1D());
 #if !INCLUDE_TRANSLUCENT
 		PackedPayload.SetFlag(PATH_TRACING_PAYLOAD_INPUT_FLAG_IGNORE_TRANSLUCENT);
 #endif
@@ -301,7 +303,7 @@ FPathTracingResult PathTracingDebugCore(FRayDesc Ray, uint2 LaunchIndex, int Sam
 		AdjustPayloadAfterUnpack(Payload, true);
 
 		// Generate main random number for direct+indirect sampling
-		float4 RandSample = RandomSequence_GenerateSample4D(RandSequence);
+		float4 RandSample = RandSequence.Get4D();
 
 #if LIGHT_LOOP_METHOD == 0
 		float LightPickingCdf[RAY_TRACING_LIGHT_COUNT_MAXIMUM];
@@ -448,7 +450,7 @@ FPathTracingResult PathTracingDebugCore(FRayDesc Ray, uint2 LaunchIndex, int Sam
 			// Writes to grid for the parent path vertex
 			const float3 Contrib = LocalRadiance;
 			const float MCf = dot(Contrib, float3(0.25, 0.5, 0.25));
-			if (MCf > 0 && RandomSequence_GenerateSample1D(RandSequence) * ScoresSum < MCf * N_MC)
+			if (MCf > 0 && RandSequence.Get1D() * ScoresSum < MCf * N_MC)
 			{
 				const float3 AbsoluteWorldPosition = DFDemote(DFFastSubtract(Ray.Origin, ResolvedView.PreViewTranslation));
 				const float Distance = length(Ray.Origin);
@@ -470,7 +472,7 @@ FPathTracingResult PathTracingDebugCore(FRayDesc Ray, uint2 LaunchIndex, int Sam
 				if (AdaptiveGridProb < 1)
 				{
 					// Save to Static Grid
-					float3 RandSample = RandomSequence_GenerateSample3D(RandSequence);
+					float3 RandSample = RandSequence.Get3D();
 					int3 Coord = int3(floor(AbsoluteWorldPosition / StaticGridWidth + RandSample));
 					int Side = 0;
 					uint2 HashIndex = GridHash(int4(Coord, -6 + Side)).xy; // two-hashes for collision detection
@@ -480,7 +482,7 @@ FPathTracingResult PathTracingDebugCore(FRayDesc Ray, uint2 LaunchIndex, int Sam
 				if (AdaptiveGridProb > 0)
 				{
 					// Save to adaptive grid
-					float4 RandSample = RandomSequence_GenerateSample4D(RandSequence);
+					float4 RandSample = RandSequence.Get4D();
 					float TargetWidth = TanFactor * Distance;
 					float Level = round(LevelScale * log2(max(TargetWidth, MinWidth) / MinWidth)) + floor(-log2(1.0 - RandSample.w));
 					float CellWidth = MinWidth * exp2(Level / LevelScale);
@@ -517,11 +519,9 @@ FPathTracingResult PathTracingDebugCore(FRayDesc Ray, uint2 LaunchIndex, int Sam
 			UNROLL
 			for (int Index = 0; Index < N_MC; Index++)
 			{
-				// TODO: Add splitting API so we can re-use the outer RandSequence instead
-				RandomSequence SplitRandSequence;
-				RandomSequence_Initialize(SplitRandSequence, LaunchIndex, SampleIndex * N_MC + Index, View.StateFrameIndex, SamplesPerPixel * N_MC);
-				float4 RandSample1 = RandomSequence_GenerateSample4D(SplitRandSequence);
-				float4 RandSample2 = RandomSequence_GenerateSample4D(SplitRandSequence);
+				RandomSequence SplitRandSequence = RandSequence.Split(Index, N_MC);
+				float4 RandSample1 = SplitRandSequence.Get4D();
+				float4 RandSample2 = SplitRandSequence.Get4D();
 
 				bool bUseAdaptiveGrid = RandSample2.z < AdaptiveGridProb;
 
@@ -628,11 +628,8 @@ FPathTracingResult PathTracingDebugCore(FRayDesc Ray, uint2 LaunchIndex)
 	const uint MissShaderIndex = 0;
 
 	FPackedPathTracingPayload PackedPayload = InitPathTracingPayload(PATHTRACER_SCATTER_CAMERA, 0.0);
-	{
-		RandomSequence RandSequence;
-		RandomSequence_Initialize(RandSequence, LaunchIndex, 0, View.StateFrameIndex, 1);
-		PackedPayload.SetStochasticSlabRand(RandomSequence_GenerateSample1D(RandSequence));
-	}
+	PackedPayload.SetStochasticSlabRand(RandomSequenceCreate(uint3(LaunchIndex, View.StateFrameIndex), 0, 1).Get1D());
+
 #if !INCLUDE_TRANSLUCENT
 	PackedPayload.SetFlag(PATH_TRACING_PAYLOAD_INPUT_FLAG_IGNORE_TRANSLUCENT);
 #endif
diff --git a/Engine/Shaders/Private/PathTracing/Utilities/PathTracingRandomSequence.ush b/Engine/Shaders/Private/PathTracing/Utilities/PathTracingRandomSequence.ush
--- a/Engine/Shaders/Private/PathTracing/Utilities/PathTracingRandomSequence.ush
+++ b/Engine/Shaders/Private/PathTracing/Utilities/PathTracingRandomSequence.ush
@@ -29,6 +29,14 @@
 #define RANDSEQ_UNROLL_SOBOL 1 // Use unrolled Sobol loop such that table can be inline constants.
 #endif
 
+#ifndef RANDSEQ_VERSION
+#define RANDSEQ_VERSION 1 // TODO: Remove code for old version (0) once all testsuite images have been updated
+#endif
+
+#ifndef RANDSEQ_SFC_TEXTURE
+#define RANDSEQ_SFC_TEXTURE 0 // 0 uses a Hilbert curve (ALU heavy) while 1 uses a precomputed texture (must be passed as a shader parameter)
+#endif
+
 // This hash mixes bits at the theoretical bias limit
 uint PerfectIntegerHash(uint x)
 {
@@ -48,22 +56,22 @@ uint PerfectIntegerHash(uint x)
 uint StrongIntegerHash(uint x)
 {
 	// From [REDACTED_PROFILE_URL]/hash-prospector
+#if RANDSEQ_VERSION == 1
+	// Current best hash in this form: [REDACTED_PROFILE_URL]/hash-prospector/issues/19#issuecomment-1105792898
+	// bias = 0.10734781817103507
+	x ^= x >> 16;
+	x *= 0x21f0aaad;
+	x ^= x >> 15;
+	x *= 0xf35a2d97;
+	x ^= x >> 15;
+#else
 	// bias = 0.16540778981744320
 	x ^= x >> 16;
 	x *= 0xa812d533;
 	x ^= x >> 15;
 	x *= 0xb278e4ad;
 	x ^= x >> 17;
-	return x;
-}
-
-// This is a much weaker hash, but is faster and can be used to drive other hashes
-uint WeakIntegerHash(uint x)
-{
-	// Generated using [REDACTED_PROFILE_URL]/hash-prospector
-	// Estimated Bias ~583
-	x *= 0x92955555u;
-	x ^= x >> 15;
+#endif
 	return x;
 }
 
@@ -94,6 +102,12 @@ uint FastOwenScramblingCore(uint Index, uint Seed)
 		Seed ^= Input & Mask;  // flip state  (depending on input)
 	}
 	return Index;
+#elif RANDSEQ_VERSION == 1
+	// This simple tweak appears to be nearly as good as the reference, while being cheaper to compute
+	Index ^= Index * 0xe0705d72u;
+	Index += Seed;
+	Seed ^= Seed >> 16;
+	Index *= Seed | 1;
 #elif 1
 	// Laine and Karras / Stratified Sampling for Stochastic Transparency / EGSR 2011
 
@@ -141,25 +155,40 @@ uint FastOwenScrambling(uint Index, uint Seed)
 	return reversebits(Index);
 }
 
+#if RANDSEQ_SFC_TEXTURE == 1
+Texture2D<uint> RandomSequenceSpaceFillingCurve;
+// Given a point (X,Y) on the unit square [0,2^B)^2, return the index along a space filling curve in [0,2^2B)
+// The space filling curve index is encoded in a texture which must be passed in as a shader parameter.
+
+uint SFCInverse(uint X, uint Y, int NumBits)
+{
+	const uint Mask = (1u << NumBits) - 1;
+	return RandomSequenceSpaceFillingCurve.Load(uint3(X & Mask, Y & Mask, 0));
+}
+#else
 // Given a point (X,Y) on the unit square [0,2^B)^2, return the index along the hilbert curve in [0,2^2B)
 // Because this function only loops over the bits it uses, there is no need to mask out bits beyond the
-// bit count, the pattern will simply tile.
-// TODO: This procedure should probably be turned into a small lookup texture
-uint HilbertCurveInverse(uint X, uint Y, int NumBits)
+// bit count, the pattern will simply tile. This is meant to be a fallback if you don't want/need to pass in
+// the space filling curve texture.
+uint SFCInverse(uint X, uint Y, int NumBits)
 {
 	// [REDACTED_PROFILE_URL]/Hackers-Delight/blob/master/hilbert/hil_s_from_xy.c.txt
-	uint TileID = 0, HilbertState = 0;
+	uint HilbertIndex = 0, HilbertState = 0;
 	for (int i = NumBits; i--;) {
 		uint xi = (X >> i) & 1;
 		uint yi = (Y >> i) & 1;
 		uint Row = 8 * HilbertState + 4 * xi + 2 * yi;
-		TileID = TileID * 4 + ((0x361E9CB4u >> Row) & 3);
+		HilbertIndex = HilbertIndex * 4 + ((0x361E9CB4u >> Row) & 3);
 		HilbertState = (0x8FE65831u >> Row) & 3;
 	}
-	return TileID;
+	return HilbertIndex;
 }
+#endif
 
-// 32-bit Sobol matrices for dimension 1,2,3 from:
+#if RANDSEQ_VERSION == 1
+// Matrices handled explicitly in SamplerCore
+#else
+// 32-bit Sobol matrices for dimension 2,3 from:
 // S. Joe and F. Y. Kuo, Constructing Sobol sequences with better two-dimensional projections, SIAM J. Sci. Comput. 30, 2635-2654 (2008)
 //    https://web.maths.unsw.edu.au/~fkuo/sobol/
 // NOTE: we don't bother storing dimension 0 since it is just a bit reversal
@@ -199,17 +228,35 @@ static const uint2 SobolMatrices[32] = {
 	uint2(0x77770071, 0x45004004),
 	uint2(0xaaaa00a3, 0xc900a00a)
 };
+#endif
 
 uint EvolveSobolSeed(inout uint Seed)
 {
+#if RANDSEQ_VERSION == 1
+	// LCG parameters from
+	// "Computationally Easy, Spectrally Good Multipliers for Congruential Pseudorandom Number Generators"
+	//   Guy Steele, Sebastiano Vigna
+	//  Journal of Software: Practice and Experience, Volume 52, Issue 2, Feb 2022
+	return Seed = Seed * 0x915f77f5u + 0x93d765ddu;
+#else
 	// constant from: https://www.pcg-random.org/posts/does-it-beat-the-minimal-standard.html
 	const uint MCG_C = 2739110765;
 	// a slightly weaker hash is ok since this drives FastOwenScrambling which is itself a hash
 	// Note that the Seed evolution is just an integer addition and the hash should optimize away
 	// when a particular dimension is not used
-	return WeakIntegerHash(Seed += MCG_C);
+	uint X = (Seed += MCG_C);
+	// Generated using [REDACTED_PROFILE_URL]/hash-prospector
+	// Estimated Bias ~583
+	X *= 0x92955555u;
+	X ^= X >> 15;
+	return X;
+#endif
 }
 
+#if RANDSEQ == RANDSEQ_PIX_LATTICE
+Texture2D<uint> PixelLatticeShift;
+#endif
+
 uint4 SamplerCore(uint SampleIndex, inout uint Seed)
 {
 #if RANDSEQ == RANDSEQ_PURERANDOM
@@ -224,6 +271,57 @@ uint4 SamplerCore(uint SampleIndex, inout uint Seed)
 	uint SobolIndex = FastOwenScrambling(SampleIndex, EvolveSobolSeed(Seed));
 	// now get Sobol' point from this index
 	uint4 Result = uint4(SobolIndex, SobolIndex, 0, 0);
+#if RANDSEQ_VERSION == 1
+	// "SZ Sequences: Binary-Based (0,2^q)-Sequences"
+	//  Abdalla G. M. Ahmed, Matt Pharr, Victor Ostromoukhov, Hui Huang
+	// 	  https://arxiv.org/pdf/2505.20434
+	//    Preprint, May 2025
+	// We only need dimension 2 because 3 can be derived via the formula Uy = P * Ux
+	// TODO: The paper suggests a faster method via Equation 42
+	Result.z = SobolIndex & 3;
+	Result.z ^= ((SobolIndex >>  2) & 1) * 6u;
+	Result.z ^= ((SobolIndex >>  3) & 1) * 11u;
+	Result.z ^= ((SobolIndex >>  4) & 1) * 19u;
+	Result.z ^= ((SobolIndex >>  5) & 1) * 33u;
+	Result.z ^= ((SobolIndex >>  6) & 1) * 109u;
+	Result.z ^= ((SobolIndex >>  7) & 1) * 182u;
+	Result.z ^= ((SobolIndex >>  8) & 1) * 258u;
+	Result.z ^= ((SobolIndex >>  9) & 1) * 515u;
+	Result.z ^= ((SobolIndex >> 10) & 1) * 1547u;
+	Result.z ^= ((SobolIndex >> 11) & 1) * 2829u;
+	Result.z ^= ((SobolIndex >> 12) & 1) * 4897u;
+	Result.z ^= ((SobolIndex >> 13) & 1) * 8498u;
+	Result.z ^= ((SobolIndex >> 14) & 1) * 28086u;
+	Result.z ^= ((SobolIndex >> 15) & 1) * 46811u;
+	Result.z ^= ((SobolIndex >> 16) & 1) * 65539u;
+	Result.z ^= ((SobolIndex >> 17) & 1) * 131073u;
+	Result.z ^= ((SobolIndex >> 18) & 1) * 393229u;
+	Result.z ^= ((SobolIndex >> 19) & 1) * 720902u;
+	Result.z ^= ((SobolIndex >> 20) & 1) * 1245234u;
+	Result.z ^= ((SobolIndex >> 21) & 1) * 2162707u;
+	Result.z ^= ((SobolIndex >> 22) & 1) * 7143643u;
+	Result.z ^= ((SobolIndex >> 23) & 1) * 11927661u;
+	Result.z ^= ((SobolIndex >> 24) & 1) * 16909057u;
+	Result.z ^= ((SobolIndex >> 25) & 1) * 33751298u;
+	Result.z ^= ((SobolIndex >> 26) & 1) * 101387526u;
+	Result.z ^= ((SobolIndex >> 27) & 1) * 185402891u;
+	Result.z ^= ((SobolIndex >> 28) & 1) * 320942611u;
+	Result.z ^= ((SobolIndex >> 29) & 1) * 556929825u;
+	Result.z ^= ((SobolIndex >> 30) & 1) * 1840700269u;
+	Result.z ^= ((SobolIndex >> 31) & 1) * 3067833782u;
+   	// y and w component can be computed without iteration by starting with Result.xz
+	// "An Implementation Algorithm of 2D Sobol Sequence Fast, Elegant, and Compact"
+	// Abdalla Ahmed, EGSR 2024
+	// See listing (19) in the paper
+	// The code is different here because we want the output to be bit-reversed, but
+	// the methodology is the same
+	Result.w = Result.z;
+	Result.yw ^=  Result.yw               >> 16;
+	Result.yw ^= (Result.yw & 0xFF00FF00) >>  8;
+	Result.yw ^= (Result.yw & 0xF0F0F0F0) >>  4;
+	Result.yw ^= (Result.yw & 0xCCCCCCCC) >>  2;
+	Result.yw ^= (Result.yw & 0xAAAAAAAA) >>  1;
+#else
 	// y component can be computed without iteration
 	// "An Implementation Algorithm of 2D Sobol Sequence Fast, Elegant, and Compact"
 	// Abdalla Ahmed, EGSR 2024
@@ -247,6 +345,7 @@ uint4 SamplerCore(uint SampleIndex, inout uint Seed)
 #endif
 		Result.zw ^= IndexBit * SobolMatrices[b];
 	}
+#endif
 	// finally scramble the points to avoid structured artifacts
 	Result.x = FastOwenScrambling(Result.x, EvolveSobolSeed(Seed));
 	Result.y = FastOwenScrambling(Result.y, EvolveSobolSeed(Seed));
@@ -257,7 +356,9 @@ uint4 SamplerCore(uint SampleIndex, inout uint Seed)
 	// Simplified algorithm using rank-1 lattice sequences. It avoids the extra call to reverse bits
 	// Scrambling the _output_ does not appear to help, so we only need to scramble the index, leading to a very fast algorithm
 	// Unfortunately, the resulting sequence does not have quite as good quality as Sobol in practice
+
 	uint LatticeIndex = FastOwenScramblingCore(SampleIndex, EvolveSobolSeed(Seed));
+
 	// Lattice parameters taken from:
 	// Weighted compound integration rules with higher order convergence for all N
 	// Fred J. Hickernell, Peter Kritzer, Frances Y. Kuo, Dirk Nuyens
@@ -341,17 +442,26 @@ FRandomSequence RandomSequenceCreate(uint3 PixelCoordAndFrame, uint SampleIndex,
 	//   Siggraph Asia 2020
 	//   http://abdallagafar.com/publications/zsampler/
 	// Improvements made in this version are: 
-	//   - switching the z-order curve for a hilbert curve
+	//   - switching the z-order curve for a sierpinski curve
 	//   - switching sobol points for owen-sobol points
 	//   - recognizing that the hierarchical scrambling proposed in the original paper is equivalent to Owen-Scrambling of the index
 	// Downsides of this method:
 	//   - quality does not improve beyond the target sample count, making usage with adaptive sampling difficult
 	//   - error-diffusion property is only achieved at the target sample count, earlier samples look random (sometimes slightly worse than random)
 
-	uint TileID = HilbertCurveInverse(PixelCoordAndFrame.x, PixelCoordAndFrame.y, 8);
+	uint TileID = SFCInverse(PixelCoordAndFrame.x, PixelCoordAndFrame.y, 8);
+
+#if 0
+	// Combine frame/sample index into the spatial index in a way that diffuses the error spatially and temporally
+	// However this appears to be worse overall, averaging successive frames does not converge quickly
+	// instead there are artifacts resulting from the space filling curve being explored in linear order.
+	TileID ^= FastOwenScrambling(reversebits(PixelCoordAndFrame.z), 0xcafef00du);
+#else
+	TileID += PixelCoordAndFrame.z * 65536;
+#endif
 
 	// Combine pixel-level and sample-level bits into the sample index (visible structure will be hidden by owen scrambling of the index)
-	RandSequence.SampleIndex = reversebits((65536 * PixelCoordAndFrame.z + TileID) * MaxSamples + SampleIndex);
+	RandSequence.SampleIndex = reversebits(TileID * MaxSamples + SampleIndex);
 
 	// progressive encodings (these kind of work, but quality is much worse)
 	//RandSequence.SampleIndex = reversebits(MortonEncode(uint2(TimeSeed0, TileID)) * MaxSamples + SampleIndex);
diff --git a/Engine/Source/Runtime/Renderer/Private/PathTracing.cpp b/Engine/Source/Runtime/Renderer/Private/PathTracing.cpp
--- a/Engine/Source/Runtime/Renderer/Private/PathTracing.cpp
+++ b/Engine/Source/Runtime/Renderer/Private/PathTracing.cpp
@@ -1354,6 +1354,7 @@ class FPathTracingRG : public FGlobalShader
 		SHADER_PARAMETER_RDG_UNIFORM_BUFFER(FNaniteRayTracingUniformParameters, NaniteRayTracing)
 		SHADER_PARAMETER_STRUCT_REF(FViewUniformShaderParameters, ViewUniformBuffer)
 		SHADER_PARAMETER_STRUCT_INCLUDE(FPathTracingData, PathTracingData)
+		SHADER_PARAMETER_RDG_TEXTURE(Texture2D<uint>, RandomSequenceSpaceFillingCurve)
 		SHADER_PARAMETER_STRUCT_INCLUDE(ShaderPrint::FShaderParameters, ShaderPrint)
 
 		// scene lights
@@ -1460,6 +1461,7 @@ class FPathTracingDebugRG : public FGlobalShader
 		SHADER_PARAMETER_RDG_BUFFER_SRV(RaytracingAccelerationStructure, TLAS)
 		SHADER_PARAMETER_RDG_BUFFER_UAV(RWStructuredBuffer<FMarkovChainState>, MarkovChainStateGrid)
 		SHADER_PARAMETER_STRUCT_REF(FViewUniformShaderParameters, ViewUniformBuffer)
+		SHADER_PARAMETER_RDG_TEXTURE(Texture2D<uint>, RandomSequenceSpaceFillingCurve)
 		SHADER_PARAMETER(int32, DebugMode)
 		SHADER_PARAMETER(int32, SamplesPerPixel)
 		SHADER_PARAMETER(int32, MaxBounces)
@@ -3504,6 +3506,7 @@ void FDeferredShadingSceneRenderer::RenderPathTracing(
 		PassParameters->MarkovChainStateGrid = GraphBuilder.CreateUAV(MarkovChainStateGrid);
 		PassParameters->MarkovChainStateGridSize = MarkovChainStateGridHashSize;
 		PassParameters->ViewUniformBuffer = View.ViewUniformBuffer;
+		PassParameters->RandomSequenceSpaceFillingCurve = GSystemTextures.GetSpaceFillingCurveTexture(GraphBuilder);
 		PassParameters->DebugMode = DebugMode;
 
 		PassParameters->SamplesPerPixel = CVarPathTracingVisualizeSamplesPerPixel.GetValueOnRenderThread();
@@ -4074,6 +4077,7 @@ void FDeferredShadingSceneRenderer::RenderPathTracing(
 							PassParameters->NaniteRayTracing = Nanite::GRayTracingManager.GetUniformBuffer();
 							PassParameters->ViewUniformBuffer = View.ViewUniformBuffer;
 							PassParameters->PathTracingData = Config.PathTracingData;
+							PassParameters->RandomSequenceSpaceFillingCurve = GSystemTextures.GetSpaceFillingCurveTexture(GraphBuilder);
 							PassParameters->StartingExtinctionCoefficient = GraphBuilder.CreateSRV(StartingExtinctionCoefficient, PF_R32_FLOAT);
 							if (PreviousPassParameters == nullptr)
 							{
diff --git a/Engine/Source/Runtime/Renderer/Private/SystemTextures.cpp b/Engine/Source/Runtime/Renderer/Private/SystemTextures.cpp
--- a/Engine/Source/Runtime/Renderer/Private/SystemTextures.cpp
+++ b/Engine/Source/Runtime/Renderer/Private/SystemTextures.cpp
@@ -1036,6 +1036,7 @@ void FSystemTextures::ReleaseRHI()
 	BlackDepthCube.SafeRelease();
 	GTAOPreIntegrated.SafeRelease();
 	AsciiTexture.SafeRelease();
+	SpaceFillingCurve.SafeRelease();
 
 	DefaultTextures.Empty();
 	DefaultBuffers.Empty();
@@ -1133,6 +1134,74 @@ FRDGTextureRef FSystemTextures::GetZeroUIntArrayAtomicCompatDummy(FRDGBuilder& G
     return GraphBuilder.RegisterExternalTexture(ZeroUIntArrayAtomicCompatDummy, TEXT("ZeroUIntArrayAtomicCompatDummy"), ERDGTextureFlags::SkipTracking);
 }
 
+FRDGTextureRef FSystemTextures::GetSpaceFillingCurveTexture(FRDGBuilder& GraphBuilder)
+{
+	if (!SpaceFillingCurve.IsValid())
+	{
+		// Initialize the texture the first time it is used
+		constexpr int Res = 256;
+		FRHITextureCreateDesc Desc = FRHITextureCreateDesc::Create2D(TEXT("SpaceFillingCurve"), Res, Res, PF_R16_UINT)
+			.SetFlags(ETextureCreateFlags::ShaderResource);
+		TArray<FIntVector> Temp;
+		Temp.Reserve(Res * Res);
+		for (int Dy = 0; Dy < Res; Dy++)
+		{
+			for (int Dx = 0; Dx < Res; Dx++)
+			{
+				unsigned X = Dx, Y = Dy;
+				// "A ROUTING SYSTEM BASED ON SPACEFILLING CURVES"
+				// https://www2.isye.gatech.edu/~jjb/research/mow/mow.pdf
+				// The paper describes a space filling curve from Sierpinski that passes through each pixel in a Res x Res square.
+				// This routine only gives results that are in the right _order_, but do not form a strict permutation
+				// so we need to store the inverse mapping and explicitly number the points after sorting them.
+				const unsigned M = Res - 1;
+				unsigned Result = 0;
+				if (X > Y)
+				{
+					Result++;
+					X = M - X;
+					Y = M - Y;
+				}
+				for (unsigned L = M; L > 0; L /= 2)
+				{
+					Result += Result;
+					if (X + Y > unsigned(M))
+					{
+						Result++;
+						unsigned OldX = X;
+						X = M - Y;
+						Y = OldX;
+					}
+					X += X;
+					Y += Y;
+					Result += Result;
+					if (Y > unsigned(M))
+					{
+						Result++;
+						unsigned OldX = X;
+						X = Y - M;
+						Y = M - OldX;
+					}
+				}
+				Temp.Add(FIntVector(Dx, Dy, Result));
+			}
+		}
+		Temp.Sort([](FIntVector A, FIntVector B) { return A.Z < B.Z; });
+		GRenderTargetPool.FindFreeElement(GraphBuilder.RHICmdList, Desc, SpaceFillingCurve, Desc.DebugName);
+		check(SpaceFillingCurve.IsValid());
+		uint32 DestStride;
+		uint16* DestBuffer = (uint16*) GraphBuilder.RHICmdList.LockTexture2D(SpaceFillingCurve->GetRHI(), 0, RLM_WriteOnly, DestStride, false);
+		DestStride /= sizeof(uint16);
+		for (int Index = 0; Index < Res * Res; Index++)
+		{
+			FIntVector P = Temp[Index];
+			DestBuffer[P.X + DestStride * P.Y] = Index;
+		}
+		GraphBuilder.RHICmdList.UnlockTexture2D(SpaceFillingCurve->GetRHI(), 0, false);
+	}
+	return GraphBuilder.RegisterExternalTexture(GSystemTextures.SpaceFillingCurve, TEXT("SpaceFillingCurve"), ERDGTextureFlags::SkipTracking);
+}
+
 //////////////////////////////////////////////////////////////////////////////////////////////////////////
 // Default textures
 
diff --git a/Engine/Source/Runtime/Renderer/Public/SystemTextures.h b/Engine/Source/Runtime/Renderer/Public/SystemTextures.h
--- a/Engine/Source/Runtime/Renderer/Public/SystemTextures.h
+++ b/Engine/Source/Runtime/Renderer/Public/SystemTextures.h
@@ -170,6 +170,9 @@ public:
 	// ASCII Standard character set - IBM code page 437 (character 32-127)
 	TRefCountPtr<IPooledRenderTarget> AsciiTexture;
 
+	// Space Filling Curve (used in FRandomSequence shader API)
+	TRefCountPtr<IPooledRenderTarget> SpaceFillingCurve;
+
 	// Create simple default texture
 	FRDGTextureRef RENDERER_API GetWhiteDummy(FRDGBuilder& GraphBuilder) const;
 	FRDGTextureRef RENDERER_API GetBlackDummy(FRDGBuilder& GraphBuilder) const;
@@ -189,6 +192,7 @@ public:
 	FRDGTextureRef RENDERER_API GetVolumetricBlackUintDummy(FRDGBuilder& GraphBuilder) const;
 	FRDGTextureRef RENDERER_API GetCubeBlackDummy(FRDGBuilder& GraphBuilder) const;
 	FRDGTextureRef RENDERER_API GetCubeArrayBlackDummy(FRDGBuilder& GraphBuilder) const;
+	FRDGTextureRef RENDERER_API GetSpaceFillingCurveTexture(FRDGBuilder& GraphBuilder);
 
 	// Create default 2D texture (1x1) with specific format and initialize value 
 	FRDGTextureRef RENDERER_API GetDefaultTexture2D(FRDGBuilder& GraphBuilder, EPixelFormat Format, float Value);
