Substrate - updated info text when debuging substrate material with legacy gbuffer.

#tests PC editor dx12
#rnx

diff --git a/Engine/Shaders/Private/Substrate/SubstrateVisualize.usf b/Engine/Shaders/Private/Substrate/SubstrateVisualize.usf
--- a/Engine/Shaders/Private/Substrate/SubstrateVisualize.usf
+++ b/Engine/Shaders/Private/Substrate/SubstrateVisualize.usf
@@ -60,7 +60,7 @@ void MaterialPrintInfoCS(uint3 DispatchThreadId : SV_DispatchThreadID)
 	SerializeNullPixelDebugDataEntry();
 
 	Context = InitShaderPrintContext(true, uint2(50, 25));
-	Print(Context, TEXT("Blendable GBuffer does not support pixel inspection yet."), FontOrange);
+	Print(Context, TEXT("Blendable GBuffer does not support pixel inspection yet. Please use `r.PostProcessing.GBufferPicking 1` instead."), FontOrange);
 
 #else
 
