Fix discrepancy look with anisotropy and clear coat under environement lighting when Lumen is disabled and Substrate is enabled.

diff --git a/Engine/Shaders/Private/Substrate/SubstrateLightingCommon.ush b/Engine/Shaders/Private/Substrate/SubstrateLightingCommon.ush
--- a/Engine/Shaders/Private/Substrate/SubstrateLightingCommon.ush
+++ b/Engine/Shaders/Private/Substrate/SubstrateLightingCommon.ush
@@ -254,6 +254,14 @@ void SubstrateEnvLightingCommon(
 
 	// Specular component
 	const bool bIsTopLayer = BSDF_GETISTOPLAYER(BSDF);
+	bool bHazeAsSimpleClearCoat = false;
+	#if SUBSTRATE_FASTPATH==0
+	if (BSDF_GETHASHAZINESS(BSDF))
+	{
+		bHazeAsSimpleClearCoat = UnpackHaziness(SLAB_HAZINESS(BSDF)).bSimpleClearCoat;
+	}
+	#endif
+
 	SpecularLighting = 0;
 #if SUBSTRATE_FASTPATH==0
 	if (any((SubstrateEnvLight.SpecularWeight + SubstrateEnvLight.SpecularHazeWeight) > 0.0f))
@@ -286,9 +294,11 @@ void SubstrateEnvLightingCommon(
 		float RoughnessSquared = SubstrateEnvLight.SpecularSafeRoughness * SubstrateEnvLight.SpecularSafeRoughness;
 		float SpecularOcclusion = IndirectSpecularOcclusion * GetSpecularOcclusion(SubstrateBSDFContext.SatNoV, RoughnessSquared, CombinedScreenAndMaterialAO);
 
+		// When using bHazeAsSimpleClearCoat, the main specular lobe represents the bottom layer's specular, on which 
+		// the TopLayerSpecularContribution (i.e. SSR) shouldn't be applied to
 		SpecularLighting += BSDFThroughput * SubstrateEnvLight.SpecularWeight *
 			GetEnvSpecularLighting(
-				(bIsTopLayer ? TopLayerSpecularContributionFactor : 1.0f) * SpecularOcclusion,
+				(bIsTopLayer && !bHazeAsSimpleClearCoat ? TopLayerSpecularContributionFactor : 1.0f) * SpecularOcclusion,
 				TranslatedWorldPosition,
 				SubstrateEnvLight.SpecularDirection,
 				SubstrateEnvLight.SpecularSafeRoughness,
