Substrate - fixed negative luminance when using anisotropic lighting especially with dual normal after refraction of the context. In a similar fashion to RefractClearCoatContext.

#rnx
#tests PC editor

diff --git a/Engine/Shaders/Private/Substrate/SubstrateEvaluation.ush b/Engine/Shaders/Private/Substrate/SubstrateEvaluation.ush
--- a/Engine/Shaders/Private/Substrate/SubstrateEvaluation.ush
+++ b/Engine/Shaders/Private/Substrate/SubstrateEvaluation.ush
@@ -278,7 +278,8 @@ half Substrate_Vis_GGX(half Roughness, half a2, half NoV, half NoL)
 half Substrate_Vis_GGX_Aniso(float ax, float ay, float NoV, float NoL, float XoV, float XoL, float YoV, float YoL)
 {
 	// SUBSTRATE_TODO: build approximation for the aniso version
-	return Vis_SmithJointAniso(ax, ay, NoV, NoL, XoV, XoL, YoV, YoL);
+	// Saturate NoV and NoL to avoid having negative value.
+	return Vis_SmithJointAniso(ax, ay, saturate(NoV), saturate(NoL), XoV, XoL, YoV, YoL);
 }
 
 float3 Substrate_F_GGX(half3 F0, half3 F90, half VoH)
