Fix voxel rendering for local light VSM.
Update SVPositionToTranslatedWorld in so CompactViewsVSM_CS so it matches the pixel size of the corresponding mip.

diff --git a/Engine/Shaders/Private/VirtualShadowMaps/VirtualShadowMapCompactViews.usf b/Engine/Shaders/Private/VirtualShadowMaps/VirtualShadowMapCompactViews.usf
--- a/Engine/Shaders/Private/VirtualShadowMaps/VirtualShadowMapCompactViews.usf
+++ b/Engine/Shaders/Private/VirtualShadowMaps/VirtualShadowMapCompactViews.usf
@@ -98,13 +98,17 @@ void CompactViewsVSM_CS(
 		MipView.HZBTestViewRect = MipView.ViewRect;
 		MipView.ViewSizeAndInvSize = float4(MipDim, MipDim, 1.0f / MipDim, 1.0f / MipDim);
 
-		float ScaleFactor = 1.0f / float(1U << MipLevel);
+		float ScaleFactor = float(1U << MipLevel);
+		float RcpScaleFactor = 1.0f / ScaleFactor;
 		MipView.LODScales = ScaleFactor * PrimaryNaniteView.LODScales;
-		MipView.ClipSpaceScaleOffset.x =  PrimaryNaniteView.ClipSpaceScaleOffset.x * ScaleFactor;
-		MipView.ClipSpaceScaleOffset.y =  PrimaryNaniteView.ClipSpaceScaleOffset.y * ScaleFactor;
+		MipView.ClipSpaceScaleOffset.x =  PrimaryNaniteView.ClipSpaceScaleOffset.x * RcpScaleFactor;
+		MipView.ClipSpaceScaleOffset.y =  PrimaryNaniteView.ClipSpaceScaleOffset.y * RcpScaleFactor;
 		MipView.ClipSpaceScaleOffset.z =  MipView.ClipSpaceScaleOffset.x - 1.0f;
 		MipView.ClipSpaceScaleOffset.w = -MipView.ClipSpaceScaleOffset.y + 1.0f;
 
+		MipView.SVPositionToTranslatedWorld[0] *= ScaleFactor;
+		MipView.SVPositionToTranslatedWorld[1] *= ScaleFactor;
+
 		// Write final mip view
 		CompactedViewsOut[CompactedOutputOffset] = MipView;
 	}
