Fixed incorrect LODScale for VSM mips introduced in CL44233573

diff --git a/Engine/Shaders/Private/VirtualShadowMaps/VirtualShadowMapCompactViews.usf b/Engine/Shaders/Private/VirtualShadowMaps/VirtualShadowMapCompactViews.usf
--- a/Engine/Shaders/Private/VirtualShadowMaps/VirtualShadowMapCompactViews.usf
+++ b/Engine/Shaders/Private/VirtualShadowMaps/VirtualShadowMapCompactViews.usf
@@ -100,7 +100,7 @@ void CompactViewsVSM_CS(
 
 		float ScaleFactor = float(1U << MipLevel);
 		float RcpScaleFactor = 1.0f / ScaleFactor;
-		MipView.LODScales = ScaleFactor * PrimaryNaniteView.LODScales;
+		MipView.LODScales = RcpScaleFactor * PrimaryNaniteView.LODScales;
 		MipView.ClipSpaceScaleOffset.x =  PrimaryNaniteView.ClipSpaceScaleOffset.x * RcpScaleFactor;
 		MipView.ClipSpaceScaleOffset.y =  PrimaryNaniteView.ClipSpaceScaleOffset.y * RcpScaleFactor;
 		MipView.ClipSpaceScaleOffset.z =  MipView.ClipSpaceScaleOffset.x - 1.0f;
