Substrate - Fixed bad clear coat indirect lighting due to swapped normals, when bottom normal is specified.

#tests PC Editor
#rnx

diff --git a/Engine/Shaders/Private/Substrate/SubstrateEvaluation.ush b/Engine/Shaders/Private/Substrate/SubstrateEvaluation.ush
--- a/Engine/Shaders/Private/Substrate/SubstrateEvaluation.ush
+++ b/Engine/Shaders/Private/Substrate/SubstrateEvaluation.ush
@@ -1643,6 +1643,25 @@ FSubstrateEnvLightResult SubstrateEvaluateForEnvLight(FSubstrateBSDFContext BSDF
 					BxDFContext ClearCoatBottomContext = RefractClearCoatContext(BSDFContext.Context);
 					ClearCoatBottomContextNoL = ClearCoatBottomContext.NoL;
 					ClearCoatBottomContextNoV = ClearCoatBottomContext.NoV;
+
+				#if SUBSTRATE_FASTPATH==0 && CLEAR_COAT_BOTTOM_NORMAL
+					BRANCH
+					if (BSDF_GETHASHAZINESS(BSDFContext.BSDF))
+					{
+						FHaziness Haziness = UnpackHaziness(SLAB_HAZINESS(BSDFContext.BSDF));
+						const bool  bHazeAsSimpleClearCoat = Haziness.bSimpleClearCoat;
+						BRANCH
+						if (bHazeAsSimpleClearCoat && Haziness.HasBottomNormal())
+						{
+							float3 BottomNormal = SubstrateGetSimpleCoatBottomNormal(Haziness.BottomNormalOct, BSDFContext.N);
+
+							BSDFContext.TangentBasis[2] = BottomNormal;
+							BSDFContext = SubstrateCreateBSDFContext(BSDFContext.BSDF, BSDFContext.TangentBasis, BSDFContext.V, BSDFContext.R/*L*/);
+
+							SubstrateEnvLightResult.DiffuseNormal = BottomNormal;
+						}
+					}
+				#endif
 				}
 				#endif
 
@@ -1658,6 +1677,7 @@ FSubstrateEnvLightResult SubstrateEvaluateForEnvLight(FSubstrateBSDFContext BSDF
 					BSDFContext.Context.NoL = dot(BSDFContext.N, BSDFContext.L);
 					BSDFContext.SatNoL = saturate(BSDFContext.Context.NoL);
 					BSDFContext.R = 2 * dot(BSDFContext.V, BSDFContext.N) * BSDFContext.N - BSDFContext.V;
+					BSDFContext.TangentBasis[2] = BSDFContext.N;
 				}
 				#endif
 
@@ -1709,19 +1729,6 @@ FSubstrateEnvLightResult SubstrateEvaluateForEnvLight(FSubstrateBSDFContext BSDF
 
 						SubstrateEnvLightResult.SpecularWeight *= BottomLayerFactors;
 						SubstrateEnvLightResult.DiffuseWeight  *= BottomLayerFactors;
-						
-					#if CLEAR_COAT_BOTTOM_NORMAL
-						if (Haziness.HasBottomNormal())
-						{
-							// Haze is top specular
-							SubstrateEnvLightResult.SpecularHazeDirection = SubstrateEnvLightResult.SpecularDirection;
-
-							// Now compute the bottom specular direction.
-							float3 SpecularBottomNormal = SubstrateGetSimpleCoatBottomNormal(Haziness.BottomNormalOct, ClearCoatTopContextN);
-							float3 SpecularBottomR = 2 * dot(BSDFContext.V, SpecularBottomNormal) * SpecularBottomNormal - BSDFContext.V;
-							SubstrateEnvLightResult.SpecularDirection = SubstrateGetOffSpecularPeakReflectionDir(SpecularBottomNormal, SpecularBottomR, SubstrateEnvLightResult.SpecularSafeRoughness);
-						}
-					#endif
 
 						SubstrateEnvLightResult.SpecularHazeWeight = TopLayerCoverage * EvalEnvBRDFHaze;
 						SubstrateEnvLightResult.SpecularHazeSafeRoughness = HazeSafeRoughness;
