Substrate - Fixed HasBottomNormal for the non complex case.

#rnx
#tests PC editor

diff --git a/Engine/Shaders/Private/Substrate/Substrate.ush b/Engine/Shaders/Private/Substrate/Substrate.ush
--- a/Engine/Shaders/Private/Substrate/Substrate.ush
+++ b/Engine/Shaders/Private/Substrate/Substrate.ush
@@ -720,7 +720,7 @@ bool   IsSubstrateSlabCompatible_ComplexSpecialPath(in FSubstrateBSDF BSDF);
 ///////////////////////////////////////////////////////////////////////////////
 // Haziness
 
-#define DEFAULT_CLEAR_COAT_BOTTOM_NORMAL_OCT float2(128.0 / 255.0, 128.0 / 255.0)
+#define DEFAULT_CLEAR_COAT_BOTTOM_NORMAL_OCT float2(0.0, 0.0)
 
 // bSimpleClearCoat==false means that Weight is the blend in amount of the secondary specular lob over the first specular lob.
 // bSimpleClearCoat==true  means fast clear coat is used, Weight is coverage of a top layer specular. 
@@ -773,6 +773,8 @@ FHaziness UnpackHaziness(uint Data)
 	Haziness.Roughness			= UnpackR8(Data);
 #if CLEAR_COAT_BOTTOM_NORMAL
 	Haziness.BottomNormalOct = UnpackBottomNormalOct(Data >> 16);
+#else
+	Haziness.BottomNormalOct = DEFAULT_CLEAR_COAT_BOTTOM_NORMAL_OCT;
 #endif
 	return Haziness;
 }
