Fix missing Emissive output from materials when using GBufferFormat=0 with Substrate.

diff --git a/Engine/Shaders/Private/BasePassPixelShader.usf b/Engine/Shaders/Private/BasePassPixelShader.usf
--- a/Engine/Shaders/Private/BasePassPixelShader.usf
+++ b/Engine/Shaders/Private/BasePassPixelShader.usf
@@ -1811,6 +1811,7 @@ void FPixelShaderInOut_MainPS(
 		GBuffer.GBufferAO = MaterialAO;
 		GBuffer.PerObjectGBufferData = GetPrimitive_PerObjectGBufferData(MaterialParameters.PrimitiveId);
 		GBuffer.Depth = MaterialParameters.ScreenPosition.w;
+		EmissiveLuminance = Export.EmissiveLuminance;
 
 		#if USES_GBUFFER && !IS_NANITE_PASS
 		// This requires cleanup. Shader code that uses GBuffer.SelectiveOutputMask expects the outputmask to be in
