Substrate - Fixed anisotropic material following the complex path due to anisotropy being enabled in the legacy path.

#tests PC editor
#rnx

diff --git a/Engine/Shaders/Private/Substrate/Substrate.ush b/Engine/Shaders/Private/Substrate/Substrate.ush
--- a/Engine/Shaders/Private/Substrate/Substrate.ush
+++ b/Engine/Shaders/Private/Substrate/Substrate.ush
@@ -4632,21 +4632,31 @@ FSubstrateBSDF UnpackSubstrateBSDFIn(FSubstrateMaterialContainer SubstrateBuffer
 		if (bHasF90 || BSDF_GETHASHAZINESS(OutBSDF))
 		{
 			const uint RawData = SubstrateLoadUint1(SubstrateBuffer, SubstrateAddressing);
-			const float2 DataRGBA8 = UnpackRGBA8(RawData).xy;
-
-			const float F90Co = DataRGBA8.x;
-			const float F90Cg = DataRGBA8.y;
-			// Setting Y=1 when converting back from YCoCg can lead to components being greater than 1. We fix this using two steps:
-			// 1- We apply a simple scale determined manually on Y to keep Hue/Saturation expressivity large and similare to the input color from the base pass,
-			//    even though the mapping from saturate to white is not linear.
-			const float F90Y  = 1.0f / 3.0f; 
-			float3 F90 = bHasF90 ? saturate(NormalisedYCoCg_2_LinearRGB(float3(F90Y, F90Co, F90Cg))) : 1.0f;
-			// 2- We re-scale again the final recovered color according to its max component to maximise the brightness/hue similarity with the input color from the base pass.
-			const float Divisor = max(F90.r, max(F90.g, F90.b));
-			F90 = Divisor > 0.0f ? F90 / Divisor : F90;
-			SLAB_F90(OutBSDF) = F90;
-
-			SLAB_HAZINESS(OutBSDF) = RawData >> 16; // Keeping the top 16 bits
+
+			BRANCH
+			if (bHasF90)
+			{
+				const float2 DataRGBA8 = UnpackRGBA8(RawData).xy;
+
+				const float F90Co = DataRGBA8.x;
+				const float F90Cg = DataRGBA8.y;
+				// Setting Y=1 when converting back from YCoCg can lead to components being greater than 1. We fix this using two steps:
+				// 1- We apply a simple scale determined manually on Y to keep Hue/Saturation expressivity large and similare to the input color from the base pass,
+				//    even though the mapping from saturate to white is not linear.
+				const float F90Y = 1.0f / 3.0f;
+				float3 F90 = bHasF90 ? saturate(NormalisedYCoCg_2_LinearRGB(float3(F90Y, F90Co, F90Cg))) : 1.0f;
+				// 2- We re-scale again the final recovered color according to its max component to maximise the brightness/hue similarity with the input color from the base pass.
+				const float Divisor = max(F90.r, max(F90.g, F90.b));
+				F90 = Divisor > 0.0f ? F90 / Divisor : F90;
+				SLAB_F90(OutBSDF) = F90;
+
+				SLAB_HAZINESS(OutBSDF) = RawData >> 16; // Keeping the top 16 bits only,
+			}
+			else
+			{
+				// See PackSubstrateOut for explanation of the special Haziness data write.
+				SLAB_HAZINESS(OutBSDF) = RawData;
+			}
 		}
 
 		const bool bHasSSS = BSDF_GETSSSTYPE(OutBSDF) != SSS_TYPE_NONE;
diff --git a/Engine/Shaders/Private/Substrate/SubstrateExport.ush b/Engine/Shaders/Private/Substrate/SubstrateExport.ush
--- a/Engine/Shaders/Private/Substrate/SubstrateExport.ush
+++ b/Engine/Shaders/Private/Substrate/SubstrateExport.ush
@@ -720,17 +720,31 @@ void PackSubstrateOut(
 								SUBSTRATE_STORE_UINT1(Data2);
 							}
 
-							if (BSDF_GETHASF90(BSDF) || BSDF_GETHASHAZINESS(BSDF))
+							const bool bHasF90 = BSDF_GETHASF90(BSDF);
+							if (bHasF90 || BSDF_GETHASHAZINESS(BSDF))
 							{
-								// What is important is to maintain the hue and saturation, so we scale the color by the maximum of its components
-								float3 F90 = saturate(SLAB_F90(BSDF));
-								const float Divisor = max(F90.r, max(F90.g, F90.b));
-								F90 = Divisor > 0.0f ? F90 / Divisor : 1.0f;
-								float3 F90YCoCg = LinearRGB_2_NormalisedYCoCg(F90);
-
-								uint F90Data = PackRGBA8(float4(F90YCoCg.y, F90YCoCg.z, 0, 0));
-								uint HazinessData = SLAB_HAZINESS(BSDF);
-								SUBSTRATE_STORE_UINT1(HazinessData << 16 | F90Data);
+								BRANCH
+								if (bHasF90)
+								{
+									// What is important is to maintain the hue and saturation, so we scale the color by the maximum of its components
+									float3 F90 = saturate(SLAB_F90(BSDF));
+									const float Divisor = max(F90.r, max(F90.g, F90.b));
+									F90 = Divisor > 0.0f ? F90 / Divisor : 1.0f;
+									float3 F90YCoCg = LinearRGB_2_NormalisedYCoCg(F90);
+
+									uint F90Data = PackRGBA8(float4(F90YCoCg.y, F90YCoCg.z, 0, 0));
+									uint HazinessData = SLAB_HAZINESS(BSDF);
+									SUBSTRATE_STORE_UINT1(HazinessData << 16 | F90Data);
+								}
+								else
+								{
+									// We only need to render clear coat material with bottom normal when converting legacy materials.
+									// That can be a single material or a complex material such as when anisotropy is enabled on legacy material conversion.
+									// In this case we must write the 32 bits of HAZINESS to get the bottom normal.
+									// And that is fine because for legacy materials F90 is not allowed.
+									uint HazinessData = SLAB_HAZINESS(BSDF);
+									SUBSTRATE_STORE_UINT1(HazinessData);
+								}
 							}
 
 							if (BSDF_GETSSSTYPE(BSDF) != SSS_TYPE_NONE)
