Substrate - Fix for bad clear coat direct lighting.

#tests PC editor dx12
#rnx

diff --git a/Engine/Shaders/Private/Substrate/SubstrateEvaluation.ush b/Engine/Shaders/Private/Substrate/SubstrateEvaluation.ush
--- a/Engine/Shaders/Private/Substrate/SubstrateEvaluation.ush
+++ b/Engine/Shaders/Private/Substrate/SubstrateEvaluation.ush
@@ -479,15 +479,15 @@ FSubstrateEvaluateResult SubstrateEvaluateBSDFCommon(FSubstrateBSDFContext BSDFC
 			DiffuseColor = ApplyDiffuseColorBoost(DiffuseColor, LumenHardwareRayTracingUniformBuffer.DiffuseColorBoost);
 		#endif
 
-
-		#if SUBSTRATE_FASTPATH==0 && CLEAR_COAT_BOTTOM_NORMAL
+			bool bHazeAsSimpleClearCoat = false;
+		#if SUBSTRATE_FASTPATH==0 && CLEAR_COAT_BOTTOM_NORMAL 
 			float  bUseBottomNormal = false;
 			float3 TopNormal = BSDFContext.N;
 			BRANCH
 			if (BSDF_GETHASHAZINESS(BSDFContext.BSDF))
 			{
 				FHaziness Haziness = UnpackHaziness(SLAB_HAZINESS(BSDFContext.BSDF));
-				const bool  bHazeAsSimpleClearCoat = Haziness.bSimpleClearCoat;
+				bHazeAsSimpleClearCoat = Haziness.bSimpleClearCoat;
 				BRANCH
 				if (bHazeAsSimpleClearCoat && Haziness.HasBottomNormal())
 				{
@@ -660,13 +660,13 @@ FSubstrateEvaluateResult SubstrateEvaluateBSDFCommon(FSubstrateBSDFContext BSDFC
 				Sample.IntegratedSpecularValue = Sample.SpecularPathValue;
 			}
 
-		#if SUBSTRATE_FASTPATH==0 && CLEAR_COAT_BOTTOM_NORMAL
+			bool bHazeHasAnisotropy = bHasAnisotropy && !bHazeAsSimpleClearCoat; // Simple clear coat top do not receive anisotropy
+		#if SUBSTRATE_FASTPATH==0 && CLEAR_COAT_BOTTOM_NORMAL 
 			if (bUseBottomNormal)
 			{
-				SubstrateSetupBSDFContext(BSDFContext, bHasAnisotropy, NoV, VoH, NoH, AreaLightContext.L, AreaLightContext.AreaLight.SphereSinAlpha);
-
 				BSDFContext.N = TopNormal;
 				BSDFContext.TangentBasis[2] = TopNormal;
+				SubstrateSetupBSDFContext(BSDFContext, bHazeHasAnisotropy, NoV, VoH, NoH, AreaLightContext.L, AreaLightContext.AreaLight.SphereSinAlpha);
 			}
 		#endif
 
@@ -702,17 +702,15 @@ FSubstrateEvaluateResult SubstrateEvaluateBSDFCommon(FSubstrateBSDFContext BSDFC
 				float3 RectLightSpecHaze = 0;
 				#if SUBSTRATE_COMPLEXPATH
 				BRANCH
-				if (bHasAnisotropy)
+				if (bHazeHasAnisotropy)
 				{
-					{
-						float2 HazeAlpha = 0;
-						GetAnisotropicRoughness(HazeSafeRoughness, SLAB_ANISOTROPY(BSDFContext.BSDF), HazeAlpha.x, HazeAlpha.y);
+					float2 HazeAlpha = 0;
+					GetAnisotropicRoughness(HazeSafeRoughness, SLAB_ANISOTROPY(BSDFContext.BSDF), HazeAlpha.x, HazeAlpha.y);
 
-						HazeD   = Substrate_D_GGX_Aniso(HazeAlpha.x, HazeAlpha.y, BSDFContext.Context.NoH, BSDFContext.Context.XoH, BSDFContext.Context.YoH);
-						HazeVis = Substrate_Vis_GGX_Aniso(HazeAlpha.x, HazeAlpha.y, BSDFContext.Context.NoV, BSDFContext.Context.NoL, BSDFContext.Context.XoV, BSDFContext.Context.XoL, BSDFContext.Context.YoV, BSDFContext.Context.YoL);
-						const float H_PDF = VisibleGGXPDF_aniso(BSDFContext.TangentV, BSDFContext.TangentH, HazeAlpha);
-						HazePDF = RayPDFToReflectionRayPDF(BSDFContext.Context.VoH, H_PDF);
-					}
+					HazeD   = Substrate_D_GGX_Aniso(HazeAlpha.x, HazeAlpha.y, BSDFContext.Context.NoH, BSDFContext.Context.XoH, BSDFContext.Context.YoH);
+					HazeVis = Substrate_Vis_GGX_Aniso(HazeAlpha.x, HazeAlpha.y, BSDFContext.Context.NoV, BSDFContext.Context.NoL, BSDFContext.Context.XoV, BSDFContext.Context.XoL, BSDFContext.Context.YoV, BSDFContext.Context.YoL);
+					const float H_PDF = VisibleGGXPDF_aniso(BSDFContext.TangentV, BSDFContext.TangentH, HazeAlpha);
+					HazePDF = RayPDFToReflectionRayPDF(BSDFContext.Context.VoH, H_PDF);
 				}
 				else
 				#endif
@@ -1757,7 +1755,7 @@ FSubstrateEnvLightResult SubstrateEvaluateForEnvLight(FSubstrateBSDFContext BSDF
 						SubstrateEnvLightResult.SpecularHazeSafeRoughness = HazeSafeRoughness;
 
 						const float DirectionalAlbedo_SpecularHazeTransmission	= ComputeEnergyPreservation(EnergyTerms);
-						DirectionalAlbedo_SpecularTransmission	= lerp(DirectionalAlbedo_SpecularTransmission, DirectionalAlbedo_SpecularHazeTransmission, HazeWeightWithFadeIn);
+						DirectionalAlbedo_SpecularTransmission	= lerp(DirectionalAlbedo_SpecularTransmission, DirectionalAlbedo_SpecularHazeTransmission, HazeWeightWithFadeIn); // Should not that be HazeWeight?
 					}
 
 					// Specular profile for haze
