Path Tracer: Fix substrate codepath in debug visualization

The new denoisers outputs were not taking this case into account.

#rnx

diff --git a/Engine/Shaders/Private/PathTracing/PathTracingDebug.usf b/Engine/Shaders/Private/PathTracing/PathTracingDebug.usf
--- a/Engine/Shaders/Private/PathTracing/PathTracingDebug.usf
+++ b/Engine/Shaders/Private/PathTracing/PathTracingDebug.usf
@@ -303,10 +303,16 @@ FPathTracingResult PathTracingDebugCore(FRayDesc Ray, uint2 LaunchIndex, int Sam
 			Result.Radiance.a = 0; // opaque alpha (we aren't handling translucency right now)
 			// G-Buffer like data for denoisers
 			Result.WorldNormal = PackedPayload.UnpackWorldNormal();
+#if PATHTRACING_SUBSTRATE_PAYLOAD
+			float4 RoughnessData = PackedPayload.UnpackRoughnessAniso();
+			Result.Roughness = lerp(RoughnessData.x, RoughnessData.y, RoughnessData.z); // mix both lobes to get an average value
+			float3 DiffuseColor = PackedPayload.UnpackDiffuseColor();
+			float3 SpecularColor = PackedPayload.UnpackSpecularColor();
+#else
 			Result.Roughness = PackedPayload.UnpackRoughness();
-
 			Result.DiffuseColor = PackedPayload.UnpackBaseColor() * (1 - PackedPayload.UnpackMetallic());
 			Result.SpecularColor = lerp(0.08 * PackedPayload.UnpackSpecular(), PackedPayload.UnpackBaseColor(), PackedPayload.UnpackMetallic());
+#endif
 			Result.LinearDepth = LinearZ;
 		}
 
@@ -818,7 +824,7 @@ RAY_TRACING_ENTRY_RAYGEN(PathTracingDebugRG)
 	Result.LinearDepth *= NormalizationFactor;
 	Result.DiffuseColor *= NormalizationFactor;
 	Result.SpecularColor *= NormalizationFactor;
-	Result.WorldNormal = normalize(Result.WorldNormal); // averaged roughness better?
+	Result.WorldNormal = normalize(Result.WorldNormal); // averaged normal better?
 	Result.Roughness *= NormalizationFactor;
 
 #else
