Substrate - Legacy gbuffer format now uses Substrate lighting code.

#tests PC editor dx12, console
#rnx

diff --git a/Engine/Shaders/Private/DeferredLightPixelShaders.usf b/Engine/Shaders/Private/DeferredLightPixelShaders.usf
--- a/Engine/Shaders/Private/DeferredLightPixelShaders.usf
+++ b/Engine/Shaders/Private/DeferredLightPixelShaders.usf
@@ -26,6 +26,12 @@
   #endif
 #endif
 
+#if SUBSTRATE_GBUFFER_FORMAT==0
+#define SUBSTRATE_LOAD_FROM_MATERIALCONTAINER 0	// Load from a specified BSDF, not from MAterialBuffer.
+#define SUBSTRATE_INLINE_SHADING 1				// Inline shading for normal setup if needed.
+#define SUBSTRATE_SSS_MATERIAL_OVERRIDE 1		// Need for SSS dither.
+#define SUBSTRATE_OPAQUE_MATERIAL 1				// Needed for SSS to end up in alpha. 
+#endif
 
 //#include "HeterogeneousVolumes/HeterogeneousVolumesVoxelGridTypes.ush"
 #include "HairStrands/HairStrandsVisibilityCommonStruct.ush" 
@@ -262,15 +268,23 @@ void DeferredLightPixelMain(
 	InputParams.ScreenVector	= ScreenVector;
 #endif
 
-#if SUBSTRATE_GBUFFER_FORMAT==1
+#if SUBSTRATE_ENABLED==1
 
+	#if SUBSTRATE_LOAD_FROM_MATERIALCONTAINER
 	FSubstrateAddressing SubstrateAddressing = GetSubstratePixelDataByteOffset(PixelPos, uint2(View.BufferSizeAndInvSize.xy), Substrate.MaxBytesPerPixel);
 	FSubstratePixelHeader SubstratePixelHeader = UnpackSubstrateHeaderIn(Substrate.MaterialTextureArray, SubstrateAddressing, Substrate.TopLayerTexture);
+	const bool bMaterialIsLit = SubstratePixelHeader.ClosureCount > 0;	// This test is also enough to exclude sky pixels
+	#else
+	FScreenSpaceData ScreenSpaceData = GetScreenSpaceData(InputParams.ScreenUV);
+	const FGBufferData GBuffer = ScreenSpaceData.GBuffer;
+	const bool bMaterialIsLit = GBuffer.ShadingModelID != SHADINGMODELID_UNLIT;
+	#endif
 
 	float EstimatedCost = 0.3; // running the PixelShader at all has a cost
 
 	BRANCH
-	if (SubstratePixelHeader.ClosureCount > 0	// This test is also enough to exclude sky pixels
+	if (
+		bMaterialIsLit
 #if USE_LIGHTING_CHANNELS
 		&& (GetLightingChannelMask(InputParams.ScreenUV) & DeferredLightUniforms.LightingChannelMask)
 #endif
@@ -296,7 +310,11 @@ void DeferredLightPixelMain(
 			FSubstrateShadowTermInputParameters SubstrateShadowTermInputParameters	= GetInitialisedSubstrateShadowTermInputParameters();
 			SubstrateShadowTermInputParameters.bEvaluateShadowTerm					= true;
 			SubstrateShadowTermInputParameters.SceneDepth							= SceneDepth;
+		#if SUBSTRATE_LOAD_FROM_MATERIALCONTAINER
 			SubstrateShadowTermInputParameters.PrecomputedShadowFactors				= SubstrateReadPrecomputedShadowFactors(SubstratePixelHeader, PixelPos, SceneTexturesStruct.GBufferETexture);
+		#else
+			SubstrateShadowTermInputParameters.PrecomputedShadowFactors				= 1.0f;
+		#endif
 			SubstrateShadowTermInputParameters.TranslatedWorldPosition				= DerivedParams.TranslatedWorldPosition;
 			SubstrateShadowTermInputParameters.LightAttenuation						= GetLightAttenuationFromShadow(InputParams, SceneDepth, DerivedParams.TranslatedWorldPosition);
 #if USE_ADAPTIVE_VOLUMETRIC_SHADOW_MAP
@@ -304,6 +322,172 @@ void DeferredLightPixelMain(
 			SubstrateShadowTermInputParameters.LightAttenuation *= AVSM_SampleTransmittance(DerivedParams.TranslatedWorldPosition, LightData.TranslatedWorldPosition);
 #endif // USE_ADAPTIVE_VOLUMETRIC_SHADOW_MAP
 			SubstrateShadowTermInputParameters.Dither								= InterleavedGradientNoise(InputParams.PixelPos, View.StateFrameIndexMod8);
+			
+		#if SUBSTRATE_LOAD_FROM_MATERIALCONTAINER==0
+
+			float3x3 BSDFTangentBasis;
+			if (GBuffer.Anisotropy != 0)
+			{
+				BSDFTangentBasis[0] = GBuffer.WorldTangent; 
+				BSDFTangentBasis[1] = cross(GBuffer.WorldNormal, GBuffer.WorldTangent);
+				BSDFTangentBasis[2] = GBuffer.WorldNormal;
+			}
+			else
+			{
+				BSDFTangentBasis = GetTangentBasis(GBuffer.WorldNormal);
+			}
+
+			float BSDFAO = ScreenSpaceData.AmbientOcclusion; // is this the equivalent to SubstrateGetAO(SubstratePixelHeader); ?
+
+			FSubstrateBSDF BSDF = (FSubstrateBSDF)0;
+			{
+				FSubstratePixelFootprint Footprint = SubstrateGetPixelFootprint(1.0f, 1.0f, GBuffer.Curvature);
+
+				const float3 Emissive							= 0.0f;
+				const float  SpecularProfileId					= 0.0f;
+
+				const float  GlintValue							= 1.0f;
+				const float2 GlintUV							= 0.0f;
+				const float2 GlintUVdx							= 0.0f;
+				const float2 GlintUVdy							= 0.0f;
+
+				float3		SSSMFP								= 0.0f;
+				float		SSSMFPScale							= 1.0f;
+				float		SSSPhaseAnisotropy					= 0.0f;
+
+				float		SecondRoughness						= 0.0f;
+				float		SecondRoughnessWeight				= 0.0f;
+				float		SecondRoughnessAsSimpleClearCoat	= 0.0f;
+
+				float		FuzzAmount							= 0.0f;
+				float3		FuzzColor							= 0.0f;
+				float		FuzzRoughness						= 0.0f;
+
+				float		SSSProfileID						= SSS_PROFILE_ID_INVALID;
+				uint		SSSType								= SSS_TYPE_NONE;
+				float		ThicknessCm							= 1.0f;
+				bool		bIsThinSurface						= false;
+
+				float3 DiffuseAlbedo = ComputeDiffuseAlbedo(GBuffer.BaseColor, GBuffer.Metallic);
+				float3 F0 = ComputeF0(GBuffer.Specular, GBuffer.BaseColor, GBuffer.Metallic);
+				float3 F90 = 1.0f;
+
+				float ClearCoatUseSecondNormal					= 0.0f;
+				float3 ClearCoatBottomNormal					= GBuffer.WorldNormal;
+
+				bool   bIsAtTheBottomOfTopology = true;
+				uint   SharedLocalBasisIndex = 0;
+				uint   SharedLocalBasisTypes = 0;
+
+				if (GBuffer.ShadingModelID == SHADINGMODELID_HAIR)
+				{
+					const float Scatter = GBuffer.Metallic;
+					const float BackLit = GBuffer.CustomData.b;
+					BSDF = GetSubstrateHairBSDF(GBuffer.BaseColor, Scatter, GBuffer.Specular, GBuffer.Roughness, BackLit, Emissive, SharedLocalBasisIndex).InlinedBSDF;
+				
+				}
+				else if (GBuffer.ShadingModelID == SHADINGMODELID_EYE)
+				{
+					SSSProfileID = ExtractSubsurfaceProfileInt(GBuffer.CustomData.x);
+					const float IrisMask = 1.0f - GBuffer.CustomData.w;
+					float IrisDistance = 0.0;
+					float3 IrisNormal = GBuffer.WorldNormal;
+					float3 IrisPlaneNormal = GBuffer.WorldNormal;
+
+				#if IRIS_NORMAL
+					const float2 IrisNormalDelta = float2(GBuffer.CustomData.y, GBuffer.CustomData.z) * 2 - (256.0 / 255.0);
+					const float2 WorldNormalOct = UnitVectorToOctahedron(GBuffer.WorldNormal);
+					IrisNormal = OctahedronToUnitVector(WorldNormalOct + IrisNormalDelta);
+				#else
+					IrisNormal = OctahedronToUnitVector(GBuffer.CustomData.yz * 2 - 1);
+					IrisDistance = GBuffer.StoredMetallic;
+				#endif
+					BSDF = GetSubstrateEyeBSDF(GBuffer.StoredBaseColor, GBuffer.Roughness, IrisMask, IrisDistance, IrisNormal, IrisPlaneNormal, SSSProfileID, Emissive, SharedLocalBasisIndex).InlinedBSDF;
+				}
+				else
+				{
+					if (GBuffer.ShadingModelID == SHADINGMODELID_CLOTH)
+					{
+						FuzzColor = ExtractSubsurfaceColor(GBuffer);
+						FuzzAmount = GBuffer.CustomData.a;
+						FuzzRoughness = GBuffer.Roughness;
+					}
+					else if (GBuffer.ShadingModelID == SHADINGMODELID_SUBSURFACE_PROFILE)
+					{
+						SSSType = SSS_TYPE_DIFFUSION;	//  SSS_TYPE_DIFFUSION_PROFILE will be selected in InternalGetSubstrateSlabBSDF based on SSSProfileID and SSSMFPScale.
+						SSSMFP = 0;
+						SSSProfileID = ExtractSubsurfaceProfileInt(GBuffer.CustomData.x);
+						SSSMFPScale = GBuffer.CustomData.a;
+						SSSPhaseAnisotropy = 0.93f;
+					}
+					else if (GBuffer.ShadingModelID == SHADINGMODELID_SUBSURFACE)
+					{
+						SSSType = SSS_TYPE_WRAP;
+						{
+							// See SubstrateGetBSDFSubSurfaceColor
+							float3 Extinction = TransmittanceToExtinction(ExtractSubsurfaceColor(GBuffer), SUBSTRATE_SIMPLEVOLUME_THICKNESS_CM * CENTIMETER_TO_METER);
+							SSSMFP = (1.0f / max(SUBSTRATE_EPSILON, Extinction));
+						}
+						SSSPhaseAnisotropy = GBuffer.CustomData.a;		// TODO test, should be one minus opacity
+						SSSMFPScale = 1.0f;
+					}
+					else if (GBuffer.ShadingModelID == SHADINGMODELID_TWOSIDED_FOLIAGE)
+					{
+						SSSType = SSS_TYPE_TWO_SIDED_WRAP;
+						{
+							// See SubstrateGetBSDFSubSurfaceColor
+							float3 Extinction = TransmittanceToExtinction(ExtractSubsurfaceColor(GBuffer), SUBSTRATE_SIMPLEVOLUME_THICKNESS_CM * CENTIMETER_TO_METER);
+							SSSMFP = (1.0f / max(SUBSTRATE_EPSILON, Extinction));
+						}
+						SSSPhaseAnisotropy = GBuffer.CustomData.a;		// TODO test, should be one minus opacity
+						SSSMFPScale = 1.0f;
+					}
+					else if (GBuffer.ShadingModelID == SHADINGMODELID_CLEAR_COAT)
+					{
+					#if SUBSTRATE_FASTPATH==0 && CLEAR_COAT_BOTTOM_NORMAL 
+						ClearCoatUseSecondNormal = 1.0f;
+						const float2 oct1 = ((float2(GBuffer.CustomData.a, GBuffer.CustomData.z) * 4) - (512.0 / 255.0)) + UnitVectorToOctahedron(GBuffer.WorldNormal);
+						ClearCoatBottomNormal = OctahedronToUnitVector(oct1);
+					#endif
+
+						SecondRoughnessWeight = GBuffer.CustomData.x;
+						SecondRoughness = GBuffer.CustomData.y;
+						SecondRoughnessAsSimpleClearCoat = (SecondRoughnessWeight > 0.0f || any(GBuffer.WorldNormal != ClearCoatBottomNormal)) ? 1.0f : 0.0f;
+					}
+
+					BSDF = InternalGetSubstrateSlabBSDF(
+						Footprint,
+						GBuffer.WorldNormal,
+						DiffuseAlbedo, F0, F90,
+						GBuffer.Roughness, GBuffer.Anisotropy,
+						SSSProfileID, SSSMFP, SSSMFPScale, SSSPhaseAnisotropy, SSSType,
+						Emissive,
+						SecondRoughness, SecondRoughnessWeight, SecondRoughnessAsSimpleClearCoat,
+						ClearCoatUseSecondNormal, ClearCoatBottomNormal,
+						FuzzAmount, FuzzColor, FuzzRoughness,
+						GlintValue, GlintUV, GlintUVdx, GlintUVdy,
+						SpecularProfileId,
+						ThicknessCm,
+						bIsThinSurface,
+						bIsAtTheBottomOfTopology,
+						SharedLocalBasisIndex, SharedLocalBasisTypes).InlinedBSDF;
+				}
+
+				BSDF.Coverage = 1;
+				BSDF.Emissive = 0;
+				BSDF.ThicknessCm = SUBSTRATE_LAYER_DEFAULT_THICKNESS_CM;
+				BSDF.TmpMFP = SSSMFP;
+				BSDF.TopLayerDataWeight = 1;
+
+				BSDF.OperatorIndex = 0;
+				BSDF.LuminanceWeightV = 1.0f;
+				BSDF.CoverageAboveAlongN = 0.0f;
+				BSDF.TransmittanceAboveAlongN = 1.0f;
+				BSDF.bIsBottom = true;
+				BSDF.bIsTop = true;
+			}
+
+		#endif
 
 			FSubstrateDeferredLighting SubstrateLighting = SubstrateDeferredLighting(
 				LightData,
@@ -312,9 +496,16 @@ void DeferredLightPixelMain(
 				ToLight,
 				LightMask,
 				SubstrateShadowTermInputParameters,
+			#if SUBSTRATE_LOAD_FROM_MATERIALCONTAINER
 				Substrate.MaterialTextureArray,
 				SubstrateAddressing,
-				SubstratePixelHeader);
+				SubstratePixelHeader
+			#else
+				BSDF,
+				BSDFTangentBasis,
+				BSDFAO
+			#endif
+			);
 
 			EstimatedCost += SubstrateLighting.EstimatedCost;
 
@@ -334,7 +525,7 @@ void DeferredLightPixelMain(
 	#endif
 #endif
 
-#else // SUBSTRATE_GBUFFER_FORMAT==1
+#else // SUBSTRATE_ENABLED==1
 
 	FScreenSpaceData ScreenSpaceData = GetScreenSpaceData(InputParams.ScreenUV);
 	// Only light pixels marked as using deferred shading
@@ -385,7 +576,7 @@ void DeferredLightPixelMain(
 		OutColor += Radiance;
 	}
 
-#endif // SUBSTRATE_GBUFFER_FORMAT==1
+#endif // SUBSTRATE_ENABLED==1
 
 	// RGB:SceneColor Specular and Diffuse
 	// A:Non Specular SceneColor Luminance
diff --git a/Engine/Shaders/Private/Substrate/SubstrateEvaluation.ush b/Engine/Shaders/Private/Substrate/SubstrateEvaluation.ush
--- a/Engine/Shaders/Private/Substrate/SubstrateEvaluation.ush
+++ b/Engine/Shaders/Private/Substrate/SubstrateEvaluation.ush
@@ -536,7 +536,7 @@ FSubstrateEvaluateResult SubstrateEvaluateBSDFCommon(FSubstrateBSDFContext BSDFC
 			//// Note: anisotropy completely disables area integrations. Lights fall back to punctual.
 			////
 
-			// Primary highlight
+			// Primary highlight. For clear coat, this is the bottom layer.
 			float PDF = 0;
 			float DirectionalAlbedo_SpecularTransmission = 1.0f;
 			{
@@ -670,7 +670,7 @@ FSubstrateEvaluateResult SubstrateEvaluateBSDFCommon(FSubstrateBSDFContext BSDFC
 			}
 		#endif
 
-			// Secondary highlight
+			// Secondary highlight. For clear coat, this is the top layer.
 			float HazePDF = 0;
 			#if SUBSTRATE_FASTPATH==0
 			BRANCH
diff --git a/Engine/Shaders/Private/Substrate/SubstrateMaterialClassification.usf b/Engine/Shaders/Private/Substrate/SubstrateMaterialClassification.usf
--- a/Engine/Shaders/Private/Substrate/SubstrateMaterialClassification.usf
+++ b/Engine/Shaders/Private/Substrate/SubstrateMaterialClassification.usf
@@ -133,8 +133,17 @@ void TileMainCS(uint2 DispatchThreadId : SV_DispatchThreadID, uint LinearIndex :
 
 		// Control tiles using ShadingModelID
 		FGBufferData GBufferData   = GetGBufferData(BufferUV);
-		bContainsSubstrateMaterial = GBufferData.ShadingModelID != SHADINGMODELID_UNLIT; 
-		bContainsSimpleMaterial    = bContainsSubstrateMaterial;
+		bContainsSubstrateMaterial = GBufferData.ShadingModelID != SHADINGMODELID_UNLIT;
+
+		if (bContainsSubstrateMaterial)
+		{
+			bool bHasAnisotropy = abs(GBufferData.Anisotropy) > 0;
+
+			bContainsComplexMaterial = bHasAnisotropy;
+			bContainsSimpleMaterial = !bContainsComplexMaterial && GBufferData.ShadingModelID == SHADINGMODELID_DEFAULT_LIT;
+			bContainsSingleMaterial = !bContainsComplexMaterial && !bContainsSimpleMaterial;
+			bContainsScreenSpaceSubsurfaceScattering = GBufferData.ShadingModelID == SHADINGMODELID_SUBSURFACE_PROFILE || GBufferData.ShadingModelID == SHADINGMODELID_EYE;
+		}
 
 #else // SUBSTRATE_GBUFFER_FORMAT==0
 
