[MegaLights] Approximate lighting on Translucent Forward Shading surfaces using TLV.
- extract maximum direction from SH and use it to calculate specular.

diff --git a/Engine/Shaders/Private/BasePassCommon.ush b/Engine/Shaders/Private/BasePassCommon.ush
--- a/Engine/Shaders/Private/BasePassCommon.ush
+++ b/Engine/Shaders/Private/BasePassCommon.ush
@@ -122,6 +122,7 @@ struct FSharedBasePassInterpolants
 #define TranslucencyLightingVolumeSampler TranslucentBasePass.TranslucencyLightingVolumeSampler
 
 #define TranslucencyLightingRandomPositionOffsetRadius TranslucentBasePass.TranslucencyLightingRandomPositionOffsetRadius
+#define TranslucencyLightingForwardUsingVolume TranslucentBasePass.TranslucencyLightingForwardUsingVolume
 
 #include "TranslucencyVolumeCommon.ush"
 
diff --git a/Engine/Shaders/Private/BasePassPixelShader.usf b/Engine/Shaders/Private/BasePassPixelShader.usf
--- a/Engine/Shaders/Private/BasePassPixelShader.usf
+++ b/Engine/Shaders/Private/BasePassPixelShader.usf
@@ -241,15 +241,18 @@ bool ManualDepthTestEqual(float4 SvPosition, float PSDeviceZWithOffset)
 }
 
 /** Calculates lighting for translucency. */
-float3 GetTranslucencyVolumeLighting(
+void GetTranslucencyVolumeLighting(
 	FMaterialPixelParameters MaterialParameters, 
 	FPixelMaterialInputs PixelMaterialInputs,
 	FBasePassInterpolantsVSToPS BasePassInterpolants, 
 	FGBufferData GBuffer, 
-	float IndirectIrradiance)
+	float IndirectIrradiance,
+	out float3 DiffuseLighting,
+	out float3 SpecularLighting)
 {
 	float4 VolumeLighting;
-	float3 InterpolatedLighting = 0;
+	DiffuseLighting = 0;
+	SpecularLighting = 0;
 	
 	float3 LightingPositionOffset = MaterialParameters.LightingPositionOffset;
 
@@ -268,22 +271,62 @@ float3 GetTranslucencyVolumeLighting(
 
 	#if TRANSLUCENCY_LIGHTING_VOLUMETRIC_PERVERTEX_DIRECTIONAL
 
-		GetVolumeLightingDirectional(float4(BasePassInterpolants.AmbientLightingVector, 1), BasePassInterpolants.DirectionalLightingVector, MaterialParameters.WorldNormal, GBuffer.DiffuseColor, GetMaterialTranslucencyDirectionalLightingIntensity(), InterpolatedLighting, VolumeLighting);
+		GetVolumeLightingDirectional(float4(BasePassInterpolants.AmbientLightingVector, 1), BasePassInterpolants.DirectionalLightingVector, MaterialParameters.WorldNormal, GBuffer.DiffuseColor, GetMaterialTranslucencyDirectionalLightingIntensity(), DiffuseLighting, VolumeLighting);
 
 	#elif TRANSLUCENCY_LIGHTING_VOLUMETRIC_PERVERTEX_NONDIRECTIONAL
 
-		GetVolumeLightingNonDirectional(float4(BasePassInterpolants.AmbientLightingVector, 1), GBuffer.DiffuseColor, InterpolatedLighting, VolumeLighting);
+		GetVolumeLightingNonDirectional(float4(BasePassInterpolants.AmbientLightingVector, 1), GBuffer.DiffuseColor, DiffuseLighting, VolumeLighting);
 
 	#elif TRANSLUCENCY_LIGHTING_VOLUMETRIC_DIRECTIONAL || TRANSLUCENCY_LIGHTING_SURFACE_LIGHTINGVOLUME
 	
 		float4 AmbientLightingVector = GetAmbientLightingVectorFromTranslucentLightingVolume(InnerVolumeUVs, OuterVolumeUVs, FinalLerpFactor);
 		float3 DirectionalLightingVector = GetDirectionalLightingVectorFromTranslucentLightingVolume(InnerVolumeUVs, OuterVolumeUVs, FinalLerpFactor);
-		GetVolumeLightingDirectional(AmbientLightingVector, DirectionalLightingVector, MaterialParameters.WorldNormal, GBuffer.DiffuseColor, GetMaterialTranslucencyDirectionalLightingIntensity(), InterpolatedLighting, VolumeLighting);
+		GetVolumeLightingDirectional(AmbientLightingVector, DirectionalLightingVector, MaterialParameters.WorldNormal, GBuffer.DiffuseColor, GetMaterialTranslucencyDirectionalLightingIntensity(), DiffuseLighting, VolumeLighting);
 
 	#elif TRANSLUCENCY_LIGHTING_VOLUMETRIC_NONDIRECTIONAL
 
 		float4 AmbientLightingVector = GetAmbientLightingVectorFromTranslucentLightingVolume(InnerVolumeUVs, OuterVolumeUVs, FinalLerpFactor);
-		GetVolumeLightingNonDirectional(AmbientLightingVector, GBuffer.DiffuseColor, InterpolatedLighting, VolumeLighting);
+		GetVolumeLightingNonDirectional(AmbientLightingVector, GBuffer.DiffuseColor, DiffuseLighting, VolumeLighting);
+
+	#elif TRANSLUCENCY_LIGHTING_SURFACE_FORWARDSHADING
+	
+		if (TranslucencyLightingForwardUsingVolume)
+		{
+			float4 AmbientLightingVector = GetAmbientLightingVectorFromTranslucentLightingVolume(InnerVolumeUVs, OuterVolumeUVs, FinalLerpFactor);
+			float3 DirectionalLightingVector = GetDirectionalLightingVectorFromTranslucentLightingVolume(InnerVolumeUVs, OuterVolumeUVs, FinalLerpFactor);
+			FTwoBandSHVectorRGB TranslucentLighting = ReconstructSHCoefficients(AmbientLightingVector, DirectionalLightingVector);
+			
+			GetVolumeLightingDirectional(TranslucentLighting, AmbientLightingVector.a, MaterialParameters.WorldNormal, GBuffer.DiffuseColor, GetMaterialTranslucencyDirectionalLightingIntensity(), DiffuseLighting, VolumeLighting);
+
+			const float3 MaxDirection = GetMaximumDirection(GetLuminance(TranslucentLighting));
+
+			// TODO: need more accurate transfer coefficients to approximate specular with.
+			FTwoBandSHVector TransferSH = CalcDiffuseTransferSH(MaxDirection, /*Exponent*/ 0.0f); 
+
+			const float3 Color = max(float3(0,0,0), DotSH(TranslucentLighting, TransferSH));
+
+			FDeferredLightData LightData = (FDeferredLightData)0;
+			LightData.Color				= Color;
+			LightData.FalloffExponent	= 0;
+			LightData.Direction			= MaxDirection;
+			LightData.DistanceFadeMAD	= 0;
+			LightData.bRadialLight		= false;
+			LightData.SpecularScale		= 1.0f;
+			LightData.DiffuseScale		= 0.0f;
+			LightData.ShadowedBits		= 1;
+			LightData.HairTransmittance	= InitHairTransmittanceData();
+			LightData.SourceRadius		= 0;
+			LightData.SoftSourceRadius	= 0;
+			LightData.IESAtlasIndex		= -1;
+			LightData.LightFunctionAtlasLightIndex = 0;
+			LightData.bAffectsTranslucentLighting = true;
+			LightData.ShadowMapChannelMask = 0;
+
+			float SurfaceShadow = 1.0f;
+			FDeferredLightingSplit ApproximateLighting = GetDynamicLightingSplit(MaterialParameters.WorldPosition_CamRelative, -MaterialParameters.CameraVector, GBuffer, 1, LightData, /*LightAttenuation*/ 1.0f, /*Dither*/ 0.0f, uint2(0,0), SurfaceShadow);
+
+			SpecularLighting += ApproximateLighting.SpecularLighting.rgb;
+		}
 
 	#endif
 
@@ -344,12 +387,10 @@ float3 GetTranslucencyVolumeLighting(
 			float3 SelfShadowingCorrection = TranslucentSelfShadow.DirectionalLightColor.rgb * VolumeLighting.a * (1 - SelfShadowing);
 				
 			// Combine backscattering and directional light self shadowing
-			InterpolatedLighting = (BackscatteredLighting + GBuffer.DiffuseColor * max(VolumeLighting.rgb - SelfShadowingCorrection, 0));
+			DiffuseLighting = (BackscatteredLighting + GBuffer.DiffuseColor * max(VolumeLighting.rgb - SelfShadowingCorrection, 0));
 		}
 
 	#endif
-
-	return InterpolatedLighting;
 }
 
 #endif
@@ -1325,42 +1366,65 @@ void FPixelShaderInOut_MainPS(
 
 				#if FORWARD_SHADING || TRANSLUCENCY_LIGHTING_SURFACE_FORWARDSHADING || MATERIAL_SHADINGMODEL_SINGLELAYERWATER
 
-					float3 DirectionalLightAtmosphereTransmittance = 1.0f;
-
-					#if NEEDS_BASEPASS_PIXEL_FOGGING && PROJECT_SUPPORT_SKY_ATMOSPHERE
-					const uint LightIndex = 0;
-					if (ResolvedView.AtmosphereLightDiscCosHalfApexAngle_PPTrans[LightIndex].y > 0.0f)
+					#if TRANSLUCENCY_LIGHTING_SURFACE_FORWARDSHADING
+					if (!TranslucencyLightingForwardUsingVolume)
+					#endif // TRANSLUCENCY_LIGHTING_SURFACE_FORWARDSHADING
 					{
-						// Only when using forward shading, we can evaluate per pixel atmosphere transmittance.
-						const float3 PlanetCenterToTranslatedWorldPos = (GetTranslatedWorldPosition(MaterialParameters) - ResolvedView.SkyPlanetTranslatedWorldCenterAndViewHeight.xyz) * CM_TO_SKY_UNIT;
-						DirectionalLightAtmosphereTransmittance = GetAtmosphereTransmittance(
-							PlanetCenterToTranslatedWorldPos, ResolvedView.AtmosphereLightDirection[LightIndex].xyz, ResolvedView.SkyAtmosphereBottomRadiusKm, ResolvedView.SkyAtmosphereTopRadiusKm,
-							View.TransmittanceLutTexture, SkyAtmTransmittanceSharedSampler);
+						float3 DirectionalLightAtmosphereTransmittance = 1.0f;
+
+						#if NEEDS_BASEPASS_PIXEL_FOGGING && PROJECT_SUPPORT_SKY_ATMOSPHERE
+						const uint LightIndex = 0;
+						if (ResolvedView.AtmosphereLightDiscCosHalfApexAngle_PPTrans[LightIndex].y > 0.0f)
+						{
+							// Only when using forward shading, we can evaluate per pixel atmosphere transmittance.
+							const float3 PlanetCenterToTranslatedWorldPos = (GetTranslatedWorldPosition(MaterialParameters) - ResolvedView.SkyPlanetTranslatedWorldCenterAndViewHeight.xyz) * CM_TO_SKY_UNIT;
+							DirectionalLightAtmosphereTransmittance = GetAtmosphereTransmittance(
+								PlanetCenterToTranslatedWorldPos, ResolvedView.AtmosphereLightDirection[LightIndex].xyz, ResolvedView.SkyAtmosphereBottomRadiusKm, ResolvedView.SkyAtmosphereTopRadiusKm,
+								View.TransmittanceLutTexture, SkyAtmTransmittanceSharedSampler);
+						}
+						#endif // NEEDS_BASEPASS_PIXEL_FOGGING && PROJECT_SUPPORT_SKY_ATMOSPHERE
+
+						#if MATERIAL_SHADINGMODEL_SINGLELAYERWATER
+						const bool bSkipMainDirLightVirtualShadowMapEvaluation = SingleLayerWater.bMainDirectionalLightVSMFiltering;
+						const bool bSeparateWaterMainDirLightLuminance = (SINGLE_LAYER_WATER_SEPARATED_MAIN_LIGHT > 0) && SingleLayerWater.bSeparateMainDirLightLuminance;
+						#else
+						const bool bSkipMainDirLightVirtualShadowMapEvaluation = false;
+						const bool bSeparateWaterMainDirLightLuminance = false;
+						#endif
+
+						FDeferredLightingSplit ForwardDirectLighting = GetForwardDirectLightingSplit(
+							In.SvPosition.xy,
+							GridIndex, MaterialParameters.WorldPosition_CamRelative, MaterialParameters.CameraVector, GBuffer, NearestResolvedDepthScreenUV, MaterialParameters.PrimitiveId, EyeIndex, Dither, 
+							DirectionalLightCloudShadow, DirectionalLightAtmosphereTransmittance, DirectionalLightShadow,
+							bSeparateWaterMainDirLightLuminance, SeparatedWaterMainDirLightLuminance,
+							bSkipMainDirLightVirtualShadowMapEvaluation);
+					
+						#if MATERIAL_SHADINGMODEL_THIN_TRANSLUCENT
+							DiffuseColor += ForwardDirectLighting.DiffuseLighting.rgb;
+							ColorSeparateSpecular += ForwardDirectLighting.SpecularLighting.rgb;
+						#else
+							Color += ForwardDirectLighting.DiffuseLighting.rgb;
+							Color += ForwardDirectLighting.SpecularLighting.rgb;
+						#endif
 					}
-					#endif // NEEDS_BASEPASS_PIXEL_FOGGING && PROJECT_SUPPORT_SKY_ATMOSPHERE
-
-					#if MATERIAL_SHADINGMODEL_SINGLELAYERWATER
-					const bool bSkipMainDirLightVirtualShadowMapEvaluation = SingleLayerWater.bMainDirectionalLightVSMFiltering;
-					const bool bSeparateWaterMainDirLightLuminance = (SINGLE_LAYER_WATER_SEPARATED_MAIN_LIGHT > 0) && SingleLayerWater.bSeparateMainDirLightLuminance;
-					#else
-					const bool bSkipMainDirLightVirtualShadowMapEvaluation = false;
-					const bool bSeparateWaterMainDirLightLuminance = false;
-					#endif
-
-					FDeferredLightingSplit ForwardDirectLighting = GetForwardDirectLightingSplit(
-						In.SvPosition.xy,
-						GridIndex, MaterialParameters.WorldPosition_CamRelative, MaterialParameters.CameraVector, GBuffer, NearestResolvedDepthScreenUV, MaterialParameters.PrimitiveId, EyeIndex, Dither, 
-						DirectionalLightCloudShadow, DirectionalLightAtmosphereTransmittance, DirectionalLightShadow,
-						bSeparateWaterMainDirLightLuminance, SeparatedWaterMainDirLightLuminance,
-						bSkipMainDirLightVirtualShadowMapEvaluation);
+					#if TRANSLUCENCY_LIGHTING_SURFACE_FORWARDSHADING
+					else
+					{
+						float3 TLVDiffuseLighting;
+						float3 TLVSpecularLighting;
+						GetTranslucencyVolumeLighting(MaterialParameters, PixelMaterialInputs, BasePassInterpolants, GBuffer, IndirectIrradiance, TLVDiffuseLighting, TLVSpecularLighting);
+						DiffuseColor += TLVDiffuseLighting;
+						ColorSeparateSpecular += TLVSpecularLighting;
 					
-					#if MATERIAL_SHADINGMODEL_THIN_TRANSLUCENT
-						DiffuseColor += ForwardDirectLighting.DiffuseLighting.rgb;
-						ColorSeparateSpecular += ForwardDirectLighting.SpecularLighting.rgb;
-					#else
-						Color += ForwardDirectLighting.DiffuseLighting.rgb;
-						Color += ForwardDirectLighting.SpecularLighting.rgb;
-					#endif
+						#if MATERIAL_SHADINGMODEL_THIN_TRANSLUCENT
+							DiffuseColor += TLVDiffuseLighting;
+							ColorSeparateSpecular += TLVSpecularLighting;
+						#else
+							Color += TLVDiffuseLighting;
+							Color += TLVSpecularLighting;
+						#endif
+					}
+					#endif // TRANSLUCENCY_LIGHTING_SURFACE_FORWARDSHADING
 				#endif
 			#endif
 				// No IBL for water in deferred: that is skipped because it is done in the water composite pass. It should however be applied when using forward shading in order to get reflection without the water composite pass.
@@ -1478,10 +1542,14 @@ void FPixelShaderInOut_MainPS(
 #if !SUBSTRATE_ENABLED
 
 	// Volume lighting for lit translucency
-#if (MATERIAL_SHADINGMODEL_DEFAULT_LIT || MATERIAL_SHADINGMODEL_SUBSURFACE) && (MATERIALBLENDING_TRANSLUCENT || MATERIALBLENDING_ADDITIVE) && !FORWARD_SHADING
+#if (MATERIAL_SHADINGMODEL_DEFAULT_LIT || MATERIAL_SHADINGMODEL_SUBSURFACE) && (MATERIALBLENDING_TRANSLUCENT || MATERIALBLENDING_ADDITIVE) && !FORWARD_SHADING && !TRANSLUCENCY_LIGHTING_SURFACE_FORWARDSHADING
 	if (GBuffer.ShadingModelID == SHADINGMODELID_DEFAULT_LIT || GBuffer.ShadingModelID == SHADINGMODELID_SUBSURFACE)
 	{
-		Color += GetTranslucencyVolumeLighting(MaterialParameters, PixelMaterialInputs, BasePassInterpolants, GBuffer, IndirectIrradiance);
+		float3 TLVDiffuseLighting;
+		float3 TLVSpecularLighting;
+		GetTranslucencyVolumeLighting(MaterialParameters, PixelMaterialInputs, BasePassInterpolants, GBuffer, IndirectIrradiance, TLVDiffuseLighting, TLVSpecularLighting);
+		Color += TLVDiffuseLighting;
+		Color += TLVSpecularLighting;
 	}
 #endif
 
diff --git a/Engine/Shaders/Private/BasePassVertexShader.usf b/Engine/Shaders/Private/BasePassVertexShader.usf
--- a/Engine/Shaders/Private/BasePassVertexShader.usf
+++ b/Engine/Shaders/Private/BasePassVertexShader.usf
@@ -136,8 +136,6 @@ void Main(
 #endif
 
 #if TRANSLUCENCY_PERVERTEX_LIGHTING_VOLUME
-	float4 VolumeLighting;
-	float3 InterpolatedLighting = 0;
 
 	float3 InnerVolumeUVs;
 	float3 OuterVolumeUVs;
@@ -157,31 +155,51 @@ void Main(
 		Output.BasePassInterpolants.AmbientLightingVector = GetAmbientLightingVectorFromTranslucentLightingVolume(InnerVolumeUVs, OuterVolumeUVs, FinalLerpFactor).xyz;
 
 	#endif
+
 #elif TRANSLUCENCY_PERVERTEX_FORWARD_SHADING
 
-	float4 VertexLightingClipSpacePosition = mul(float4(WorldPositionForVertexLightingTranslated, 1), ResolvedView.TranslatedWorldToClip);
-	float2 SvPosition = (VertexLightingClipSpacePosition.xy / VertexLightingClipSpacePosition.w * float2(.5f, -.5f) + .5f) * ResolvedView.ViewSizeAndInvSize.xy;
-	uint GridIndex = ComputeLightGridCellIndex((uint2)(SvPosition* View.LightProbeSizeRatioAndInvSizeRatio.zw - ResolvedView.ViewRectMin.xy), VertexLightingClipSpacePosition.w, EyeIndex);
+	float3 InnerVolumeUVs;
+	float3 OuterVolumeUVs;
+	float FinalLerpFactor;
 
-	float DirectionalLightCloudShadow = 1.0f;
-	#if SUPPORT_CLOUD_SHADOW_ON_FORWARD_LIT_TRANSLUCENT
-		float OutOpticalDepth = 0.0f;
-		if (TranslucentBasePass.ForwardDirLightCloudShadow.CloudShadowmapStrength > 0.0f)
-		{
-			DirectionalLightCloudShadow = lerp(
-				1.0f,
-				GetCloudVolumetricShadow(
-					WorldPositionForVertexLightingTranslated,
-					TranslucentBasePass.ForwardDirLightCloudShadow.CloudShadowmapTranslatedWorldToLightClipMatrix,
-					TranslucentBasePass.ForwardDirLightCloudShadow.CloudShadowmapFarDepthKm,
-					TranslucentBasePass.ForwardDirLightCloudShadow.CloudShadowmapTexture,
-					TranslucentBasePass.ForwardDirLightCloudShadow.CloudShadowmapSampler,
-					OutOpticalDepth),
-				TranslucentBasePass.ForwardDirLightCloudShadow.CloudShadowmapStrength);
-		}
-	#endif
+	//@todo - get from VF
+	float3 LightingPositionOffset = 0;
+	ComputeVolumeUVs(WorldPositionForVertexLightingTranslated, LightingPositionOffset, InnerVolumeUVs, OuterVolumeUVs, FinalLerpFactor);
+
+	if (TranslucencyLightingForwardUsingVolume)
+	{
+		float4 AmbientLightingVector = GetAmbientLightingVectorFromTranslucentLightingVolume(InnerVolumeUVs, OuterVolumeUVs, FinalLerpFactor);
+		float3 DirectionalLightingVector = GetDirectionalLightingVectorFromTranslucentLightingVolume(InnerVolumeUVs, OuterVolumeUVs, FinalLerpFactor);
+		FTwoBandSHVectorRGB TranslucentLighting = ReconstructSHCoefficients(AmbientLightingVector, DirectionalLightingVector);
+
+		Output.BasePassInterpolants.VertexDiffuseLighting = GetVolumeLightingDirectional(AmbientLightingVector, DirectionalLightingVector, VertexParameters.TangentToWorld[2], GetMaterialTranslucencyDirectionalLightingIntensity()).rgb;
+	}
+	else
+	{
+		float4 VertexLightingClipSpacePosition = mul(float4(WorldPositionForVertexLightingTranslated, 1), ResolvedView.TranslatedWorldToClip);
+		float2 SvPosition = (VertexLightingClipSpacePosition.xy / VertexLightingClipSpacePosition.w * float2(.5f, -.5f) + .5f) * ResolvedView.ViewSizeAndInvSize.xy;
+		uint GridIndex = ComputeLightGridCellIndex((uint2)(SvPosition* View.LightProbeSizeRatioAndInvSizeRatio.zw - ResolvedView.ViewRectMin.xy), VertexLightingClipSpacePosition.w, EyeIndex);
+
+		float DirectionalLightCloudShadow = 1.0f;
+		#if SUPPORT_CLOUD_SHADOW_ON_FORWARD_LIT_TRANSLUCENT
+			float OutOpticalDepth = 0.0f;
+			if (TranslucentBasePass.ForwardDirLightCloudShadow.CloudShadowmapStrength > 0.0f)
+			{
+				DirectionalLightCloudShadow = lerp(
+					1.0f,
+					GetCloudVolumetricShadow(
+						WorldPositionForVertexLightingTranslated,
+						TranslucentBasePass.ForwardDirLightCloudShadow.CloudShadowmapTranslatedWorldToLightClipMatrix,
+						TranslucentBasePass.ForwardDirLightCloudShadow.CloudShadowmapFarDepthKm,
+						TranslucentBasePass.ForwardDirLightCloudShadow.CloudShadowmapTexture,
+						TranslucentBasePass.ForwardDirLightCloudShadow.CloudShadowmapSampler,
+						OutOpticalDepth),
+					TranslucentBasePass.ForwardDirLightCloudShadow.CloudShadowmapStrength);
+			}
+		#endif
 
-	Output.BasePassInterpolants.VertexDiffuseLighting = GetForwardDirectLightingForVertexLighting(GridIndex, WorldPositionForVertexLightingTranslated, Output.Position.w, VertexParameters.TangentToWorld[2], EyeIndex, DirectionalLightCloudShadow);
+		Output.BasePassInterpolants.VertexDiffuseLighting = GetForwardDirectLightingForVertexLighting(GridIndex, WorldPositionForVertexLightingTranslated, Output.Position.w, VertexParameters.TangentToWorld[2], EyeIndex, DirectionalLightCloudShadow);
+	}
 
 #elif NEEDS_BASEPASS_CLOUD_SHADOW_INTERPOLATOR
 
diff --git a/Engine/Shaders/Private/TranslucencyVolumeCommon.ush b/Engine/Shaders/Private/TranslucencyVolumeCommon.ush
--- a/Engine/Shaders/Private/TranslucencyVolumeCommon.ush
+++ b/Engine/Shaders/Private/TranslucencyVolumeCommon.ush
@@ -38,6 +38,23 @@ float3 GetDirectionalLightingVectorFromTranslucentLightingVolume(float3 InnerVol
 
 #include "SHCommon.ush"
 
+FTwoBandSHVectorRGB ReconstructSHCoefficients(float4 AmbientLightingVector, float3 DirectionalLightingVector)
+{
+	// Reconstruct the SH coefficients based on what was encoded
+	FTwoBandSHVectorRGB TranslucentLighting;
+	TranslucentLighting.R.V.x = AmbientLightingVector.r;
+	TranslucentLighting.G.V.x = AmbientLightingVector.g;
+	TranslucentLighting.B.V.x = AmbientLightingVector.b;
+	float3 NormalizedAmbientColor = AmbientLightingVector.rgb / (Luminance(AmbientLightingVector.rgb) + 0.00001f);
+
+	// Scale the monocrome directional coefficients with the normalzed ambient color as an approximation to the uncompressed values
+	TranslucentLighting.R.V.yzw = DirectionalLightingVector.rgb * NormalizedAmbientColor.r;
+	TranslucentLighting.G.V.yzw = DirectionalLightingVector.rgb * NormalizedAmbientColor.g;
+	TranslucentLighting.B.V.yzw = DirectionalLightingVector.rgb * NormalizedAmbientColor.b;
+
+	return TranslucentLighting;
+}
+
 void GetVolumeLightingNonDirectional(float4 AmbientLightingVector, float3 DiffuseColor, inout float3 InterpolatedLighting, out float4 VolumeLighting)
 {
 	// Normal is not taken into account with non directional lighting, and only the ambient term of the SH coefficients are needed
@@ -51,26 +68,21 @@ void GetVolumeLightingNonDirectional(float4 AmbientLightingVector, float3 Diffus
 	InterpolatedLighting = DiffuseColor * VolumeLighting.rgb;
 }
 
-void GetVolumeLightingDirectional(float4 AmbientLightingVector, float3 DirectionalLightingVector, float3 WorldNormal, float3 DiffuseColor, float DirectionalLightingIntensity, inout float3 InterpolatedLighting, out float4 VolumeLighting)
+void GetVolumeLightingDirectional(FTwoBandSHVectorRGB TranslucentLighting, float DirectionalLightContribution, float3 WorldNormal, float3 DiffuseColor, float DirectionalLightingIntensity, inout float3 InterpolatedLighting, out float4 VolumeLighting)
 {
-	// Reconstruct the SH coefficients based on what was encoded
-	FTwoBandSHVectorRGB TranslucentLighting;
-	TranslucentLighting.R.V.x = AmbientLightingVector.r;
-	TranslucentLighting.G.V.x = AmbientLightingVector.g;
-	TranslucentLighting.B.V.x = AmbientLightingVector.b;
-	float3 NormalizedAmbientColor = AmbientLightingVector.rgb / (Luminance(AmbientLightingVector.rgb) + 0.00001f);
-
-	// Scale the monocrome directional coefficients with the normalzed ambient color as an approximation to the uncompressed values
-	TranslucentLighting.R.V.yzw = DirectionalLightingVector.rgb * NormalizedAmbientColor.r;
-	TranslucentLighting.G.V.yzw = DirectionalLightingVector.rgb * NormalizedAmbientColor.g;
-	TranslucentLighting.B.V.yzw = DirectionalLightingVector.rgb * NormalizedAmbientColor.b;
-
 	// Compute diffuse lighting which takes the normal into account
 	FTwoBandSHVector DiffuseTransferSH = CalcDiffuseTransferSH(WorldNormal * DirectionalLightingIntensity, 1);
-	VolumeLighting = float4(max(half3(0, 0, 0), DotSH(TranslucentLighting, DiffuseTransferSH)), AmbientLightingVector.a);
+	VolumeLighting = float4(max(half3(0, 0, 0), DotSH(TranslucentLighting, DiffuseTransferSH)), DirectionalLightContribution);
 	InterpolatedLighting += DiffuseColor * VolumeLighting.rgb;
 }
 
+void GetVolumeLightingDirectional(float4 AmbientLightingVector, float3 DirectionalLightingVector, float3 WorldNormal, float3 DiffuseColor, float DirectionalLightingIntensity, inout float3 InterpolatedLighting, out float4 VolumeLighting)
+{
+	FTwoBandSHVectorRGB TranslucentLighting = ReconstructSHCoefficients(AmbientLightingVector, DirectionalLightingVector);
+	
+	GetVolumeLightingDirectional(TranslucentLighting, AmbientLightingVector.a, WorldNormal, DiffuseColor, DirectionalLightingIntensity, InterpolatedLighting, VolumeLighting);
+}
+
 float4 GetVolumeLightingDirectional(float4 AmbientLightingVector, float3 DirectionalLightingVector, float3 WorldNormal, float DirectionalLightingIntensity)
 {
 	float3 DiffuseColor = 1.0f;
diff --git a/Engine/Source/Runtime/Renderer/Private/TranslucentLighting.cpp b/Engine/Source/Runtime/Renderer/Private/TranslucentLighting.cpp
--- a/Engine/Source/Runtime/Renderer/Private/TranslucentLighting.cpp
+++ b/Engine/Source/Runtime/Renderer/Private/TranslucentLighting.cpp
@@ -1274,6 +1274,7 @@ FTranslucencyLightingVolumeParameters GetTranslucencyLightingVolumeParameters(FR
 		: TStaticSamplerState<SF_Bilinear, AM_Clamp, AM_Clamp, AM_Clamp>::GetRHI();
 
 	Parameters.TranslucencyLightingRandomPositionOffsetRadius = GetTranslucencyLightingVolumePositionOffsetRadius();
+	Parameters.TranslucencyLightingForwardUsingVolume = MegaLights::IsEnabled(*View.Family) && MegaLights::UseTranslucencyVolume();
 
 	return Parameters;
 }
diff --git a/Engine/Source/Runtime/Renderer/Private/TranslucentLighting.h b/Engine/Source/Runtime/Renderer/Private/TranslucentLighting.h
--- a/Engine/Source/Runtime/Renderer/Private/TranslucentLighting.h
+++ b/Engine/Source/Runtime/Renderer/Private/TranslucentLighting.h
@@ -54,6 +54,7 @@ BEGIN_SHADER_PARAMETER_STRUCT(FTranslucencyLightingVolumeParameters, )
 	SHADER_PARAMETER_RDG_TEXTURE(Texture3D, TranslucencyLightingVolumeDirectionalOuter)
 	SHADER_PARAMETER_SAMPLER(SamplerState, TranslucencyLightingVolumeSampler)
 	SHADER_PARAMETER(float, TranslucencyLightingRandomPositionOffsetRadius)
+	SHADER_PARAMETER(float, TranslucencyLightingForwardUsingVolume)
 END_SHADER_PARAMETER_STRUCT()
 
 /** Utility for batching together multiple lighting injections */
