Substrate - Legacy gbuffer format now uses Substrate lighting code for environment lighting.
BSDF loading from legacy GBuffer is now shared.
Fixed AO being applied twice in substrate path.

We still need to figure out why reflection do not works well on eye ou simple clear coat.

#tests PC and console performance
#rnx

diff --git a/Engine/Shaders/Private/AmbientCubemapComposite.usf b/Engine/Shaders/Private/AmbientCubemapComposite.usf
--- a/Engine/Shaders/Private/AmbientCubemapComposite.usf
+++ b/Engine/Shaders/Private/AmbientCubemapComposite.usf
@@ -99,12 +99,13 @@ void MainPS(
 
 			// Evaluate environment lighting
 			FSubstrateEnvLightResult SubstrateEnvLight = SubstrateEvaluateForEnvLight(SubstrateBSDFContext, true /*bEnableSpecular*/, Settings);
+			FSubstrateIrradianceAndOcclusion SubstrateIrradianceAndOcclusion = SubstrateGetIrradianceAndAO(SubstratePixelHeader);
 
 			float3 DiffuseLighting = 0;
 			float3 SpecularLighting = 0;
 			SubstrateEnvLightingCommon(
 				SubstrateEnvLight,
-				SubstratePixelHeader,
+				SubstrateIrradianceAndOcclusion,
 				SubstrateBSDFContext,
 				BSDF,
 				SubstrateEnvLight.DiffuseNormal,
diff --git a/Engine/Shaders/Private/DeferredLightPixelShaders.usf b/Engine/Shaders/Private/DeferredLightPixelShaders.usf
--- a/Engine/Shaders/Private/DeferredLightPixelShaders.usf
+++ b/Engine/Shaders/Private/DeferredLightPixelShaders.usf
@@ -31,6 +31,9 @@
 #define SUBSTRATE_INLINE_SHADING 1				// Inline shading for normal setup if needed.
 #define SUBSTRATE_SSS_MATERIAL_OVERRIDE 1		// Need for SSS dither.
 #define SUBSTRATE_OPAQUE_MATERIAL 1				// Needed for SSS to end up in alpha. 
+#include "Substrate/SubstrateRead.ush"
+#else 
+#define SUBSTRATE_LOAD_FROM_MATERIALCONTAINER 1
 #endif
 
 //#include "HeterogeneousVolumes/HeterogeneousVolumesVoxelGridTypes.ush"
@@ -325,167 +328,7 @@ void DeferredLightPixelMain(
 			
 		#if SUBSTRATE_LOAD_FROM_MATERIALCONTAINER==0
 
-			float3x3 BSDFTangentBasis;
-			if (GBuffer.Anisotropy != 0)
-			{
-				BSDFTangentBasis[0] = GBuffer.WorldTangent; 
-				BSDFTangentBasis[1] = cross(GBuffer.WorldNormal, GBuffer.WorldTangent);
-				BSDFTangentBasis[2] = GBuffer.WorldNormal;
-			}
-			else
-			{
-				BSDFTangentBasis = GetTangentBasis(GBuffer.WorldNormal);
-			}
-
-			float BSDFAO = ScreenSpaceData.AmbientOcclusion; // is this the equivalent to SubstrateGetAO(SubstratePixelHeader); ?
-
-			FSubstrateBSDF BSDF = (FSubstrateBSDF)0;
-			{
-				FSubstratePixelFootprint Footprint = SubstrateGetPixelFootprint(1.0f, 1.0f, GBuffer.Curvature);
-
-				const float3 Emissive							= 0.0f;
-				const float  SpecularProfileId					= 0.0f;
-
-				const float  GlintValue							= 1.0f;
-				const float2 GlintUV							= 0.0f;
-				const float2 GlintUVdx							= 0.0f;
-				const float2 GlintUVdy							= 0.0f;
-
-				float3		SSSMFP								= 0.0f;
-				float		SSSMFPScale							= 1.0f;
-				float		SSSPhaseAnisotropy					= 0.0f;
-
-				float		SecondRoughness						= 0.0f;
-				float		SecondRoughnessWeight				= 0.0f;
-				float		SecondRoughnessAsSimpleClearCoat	= 0.0f;
-
-				float		FuzzAmount							= 0.0f;
-				float3		FuzzColor							= 0.0f;
-				float		FuzzRoughness						= 0.0f;
-
-				float		SSSProfileID						= SSS_PROFILE_ID_INVALID;
-				uint		SSSType								= SSS_TYPE_NONE;
-				float		ThicknessCm							= 1.0f;
-				bool		bIsThinSurface						= false;
-
-				float3 DiffuseAlbedo = ComputeDiffuseAlbedo(GBuffer.BaseColor, GBuffer.Metallic);
-				float3 F0 = ComputeF0(GBuffer.Specular, GBuffer.BaseColor, GBuffer.Metallic);
-				float3 F90 = 1.0f;
-
-				float ClearCoatUseSecondNormal					= 0.0f;
-				float3 ClearCoatBottomNormal					= GBuffer.WorldNormal;
-
-				bool   bIsAtTheBottomOfTopology = true;
-				uint   SharedLocalBasisIndex = 0;
-				uint   SharedLocalBasisTypes = 0;
-
-				if (GBuffer.ShadingModelID == SHADINGMODELID_HAIR)
-				{
-					const float Scatter = GBuffer.Metallic;
-					const float BackLit = GBuffer.CustomData.b;
-					BSDF = GetSubstrateHairBSDF(GBuffer.BaseColor, Scatter, GBuffer.Specular, GBuffer.Roughness, BackLit, Emissive, SharedLocalBasisIndex).InlinedBSDF;
-				
-				}
-				else if (GBuffer.ShadingModelID == SHADINGMODELID_EYE)
-				{
-					SSSProfileID = ExtractSubsurfaceProfileInt(GBuffer.CustomData.x);
-					const float IrisMask = 1.0f - GBuffer.CustomData.w;
-					float IrisDistance = 0.0;
-					float3 IrisNormal = GBuffer.WorldNormal;
-					float3 IrisPlaneNormal = GBuffer.WorldNormal;
-
-				#if IRIS_NORMAL
-					const float2 IrisNormalDelta = float2(GBuffer.CustomData.y, GBuffer.CustomData.z) * 2 - (256.0 / 255.0);
-					const float2 WorldNormalOct = UnitVectorToOctahedron(GBuffer.WorldNormal);
-					IrisNormal = OctahedronToUnitVector(WorldNormalOct + IrisNormalDelta);
-				#else
-					IrisNormal = OctahedronToUnitVector(GBuffer.CustomData.yz * 2 - 1);
-					IrisDistance = GBuffer.StoredMetallic;
-				#endif
-					BSDF = GetSubstrateEyeBSDF(GBuffer.StoredBaseColor, GBuffer.Roughness, IrisMask, IrisDistance, IrisNormal, IrisPlaneNormal, SSSProfileID, Emissive, SharedLocalBasisIndex).InlinedBSDF;
-				}
-				else
-				{
-					if (GBuffer.ShadingModelID == SHADINGMODELID_CLOTH)
-					{
-						FuzzColor = ExtractSubsurfaceColor(GBuffer);
-						FuzzAmount = GBuffer.CustomData.a;
-						FuzzRoughness = GBuffer.Roughness;
-					}
-					else if (GBuffer.ShadingModelID == SHADINGMODELID_SUBSURFACE_PROFILE)
-					{
-						SSSType = SSS_TYPE_DIFFUSION;	//  SSS_TYPE_DIFFUSION_PROFILE will be selected in InternalGetSubstrateSlabBSDF based on SSSProfileID and SSSMFPScale.
-						SSSMFP = 0;
-						SSSProfileID = ExtractSubsurfaceProfileInt(GBuffer.CustomData.x);
-						SSSMFPScale = GBuffer.CustomData.a;
-						SSSPhaseAnisotropy = 0.93f;
-					}
-					else if (GBuffer.ShadingModelID == SHADINGMODELID_SUBSURFACE)
-					{
-						SSSType = SSS_TYPE_WRAP;
-						{
-							// See SubstrateGetBSDFSubSurfaceColor
-							float3 Extinction = TransmittanceToExtinction(ExtractSubsurfaceColor(GBuffer), SUBSTRATE_SIMPLEVOLUME_THICKNESS_CM * CENTIMETER_TO_METER);
-							SSSMFP = (1.0f / max(SUBSTRATE_EPSILON, Extinction));
-						}
-						SSSPhaseAnisotropy = GBuffer.CustomData.a;		// TODO test, should be one minus opacity
-						SSSMFPScale = 1.0f;
-					}
-					else if (GBuffer.ShadingModelID == SHADINGMODELID_TWOSIDED_FOLIAGE)
-					{
-						SSSType = SSS_TYPE_TWO_SIDED_WRAP;
-						{
-							// See SubstrateGetBSDFSubSurfaceColor
-							float3 Extinction = TransmittanceToExtinction(ExtractSubsurfaceColor(GBuffer), SUBSTRATE_SIMPLEVOLUME_THICKNESS_CM * CENTIMETER_TO_METER);
-							SSSMFP = (1.0f / max(SUBSTRATE_EPSILON, Extinction));
-						}
-						SSSPhaseAnisotropy = GBuffer.CustomData.a;		// TODO test, should be one minus opacity
-						SSSMFPScale = 1.0f;
-					}
-					else if (GBuffer.ShadingModelID == SHADINGMODELID_CLEAR_COAT)
-					{
-					#if SUBSTRATE_FASTPATH==0 && CLEAR_COAT_BOTTOM_NORMAL 
-						ClearCoatUseSecondNormal = 1.0f;
-						const float2 oct1 = ((float2(GBuffer.CustomData.a, GBuffer.CustomData.z) * 4) - (512.0 / 255.0)) + UnitVectorToOctahedron(GBuffer.WorldNormal);
-						ClearCoatBottomNormal = OctahedronToUnitVector(oct1);
-					#endif
-
-						SecondRoughnessWeight = GBuffer.CustomData.x;
-						SecondRoughness = GBuffer.CustomData.y;
-						SecondRoughnessAsSimpleClearCoat = (SecondRoughnessWeight > 0.0f || any(GBuffer.WorldNormal != ClearCoatBottomNormal)) ? 1.0f : 0.0f;
-					}
-
-					BSDF = InternalGetSubstrateSlabBSDF(
-						Footprint,
-						GBuffer.WorldNormal,
-						DiffuseAlbedo, F0, F90,
-						GBuffer.Roughness, GBuffer.Anisotropy,
-						SSSProfileID, SSSMFP, SSSMFPScale, SSSPhaseAnisotropy, SSSType,
-						Emissive,
-						SecondRoughness, SecondRoughnessWeight, SecondRoughnessAsSimpleClearCoat,
-						ClearCoatUseSecondNormal, ClearCoatBottomNormal,
-						FuzzAmount, FuzzColor, FuzzRoughness,
-						GlintValue, GlintUV, GlintUVdx, GlintUVdy,
-						SpecularProfileId,
-						ThicknessCm,
-						bIsThinSurface,
-						bIsAtTheBottomOfTopology,
-						SharedLocalBasisIndex, SharedLocalBasisTypes).InlinedBSDF;
-				}
-
-				BSDF.Coverage = 1;
-				BSDF.Emissive = 0;
-				BSDF.ThicknessCm = SUBSTRATE_LAYER_DEFAULT_THICKNESS_CM;
-				BSDF.TmpMFP = SSSMFP;
-				BSDF.TopLayerDataWeight = 1;
-
-				BSDF.OperatorIndex = 0;
-				BSDF.LuminanceWeightV = 1.0f;
-				BSDF.CoverageAboveAlongN = 0.0f;
-				BSDF.TransmittanceAboveAlongN = 1.0f;
-				BSDF.bIsBottom = true;
-				BSDF.bIsTop = true;
-			}
+			FSubstrateGBufferBSDF GBufferBSDF = SubstrateReadGBufferBSDF(ScreenSpaceData);
 
 		#endif
 
@@ -501,9 +344,9 @@ void DeferredLightPixelMain(
 				SubstrateAddressing,
 				SubstratePixelHeader
 			#else
-				BSDF,
-				BSDFTangentBasis,
-				BSDFAO
+				GBufferBSDF.BSDF,
+				GBufferBSDF.BSDFTangentBasis,
+				GBufferBSDF.BSDFAO
 			#endif
 			);
 
diff --git a/Engine/Shaders/Private/ReflectionEnvironmentPixelShader.usf b/Engine/Shaders/Private/ReflectionEnvironmentPixelShader.usf
--- a/Engine/Shaders/Private/ReflectionEnvironmentPixelShader.usf
+++ b/Engine/Shaders/Private/ReflectionEnvironmentPixelShader.usf
@@ -20,6 +20,18 @@
   #endif
 #endif
 
+//#define SceneTexturesStruct SceneTextures 
+
+#if SUBSTRATE_GBUFFER_FORMAT==0
+#define SUBSTRATE_LOAD_FROM_MATERIALCONTAINER 0	// Load from a specified BSDF, not from MAterialBuffer.
+#define SUBSTRATE_INLINE_SHADING 1				// Inline shading for normal setup if needed.
+#define SUBSTRATE_SSS_MATERIAL_OVERRIDE 1		// Need for SSS dither.
+#define SUBSTRATE_OPAQUE_MATERIAL 1				// Needed for SSS to end up in alpha. 
+#include "Substrate/SubstrateRead.ush"
+#else 
+#define SUBSTRATE_LOAD_FROM_MATERIALCONTAINER 1
+#endif
+
 #include "Common.ush" 
 #include "DeferredShadingCommon.ush"  
 #include "BRDF.ush"
@@ -296,7 +308,8 @@ float GetCloudVolumetricAOShadow(in float3 TranslatedWorldPosition)
 	return 1.0f;
 }
 
-#if SUBSTRATE_GBUFFER_FORMAT==1
+#if SUBSTRATE_ENABLED==1
+
 struct FSSRContribution
 {
 	float3 SSRTopLayerEnvBRDF;
@@ -387,8 +400,7 @@ FLumenReflectionOnlyOutput GetLumenReflectionOnly(
 	return Out;
 }
 
-#endif // SUBSTRATE_GBUFFER_FORMAT==1
-
+#endif // SUBSTRATE_ENABLED==1
 
 void ReflectionEnvironmentSkyLighting(
 	in float4 SvPosition : SV_Position,
@@ -411,12 +423,22 @@ void ReflectionEnvironmentSkyLighting(
 	OutSubSurfaceSceneColor = 0.0f;
 #endif
 
-#if SUBSTRATE_GBUFFER_FORMAT==1 
+#if SUBSTRATE_ENABLED==1
 
+	#if SUBSTRATE_LOAD_FROM_MATERIALCONTAINER
 	FSubstrateAddressing SubstrateAddressing = GetSubstratePixelDataByteOffset(PixelPos, uint2(View.BufferSizeAndInvSize.xy), Substrate.MaxBytesPerPixel);
 	FSubstratePixelHeader SubstratePixelHeader = UnpackSubstrateHeaderIn(Substrate.MaterialTextureArray, SubstrateAddressing, Substrate.TopLayerTexture);
+	const bool bMaterialIsLit = SubstratePixelHeader.ClosureCount > 0;	// This test is also enough to exclude sky pixels
+	#else
+	FScreenSpaceData ScreenSpaceData;
+	const FGBufferData GBuffer = GetGBufferDataFromSceneTextures(BufferUV);
+	ScreenSpaceData.GBuffer = GBuffer;
+	ScreenSpaceData.AmbientOcclusion = 1.0f; // Used, fetched from AmbientOcclusionTexture.
+	const bool bMaterialIsLit = GBuffer.ShadingModelID != SHADINGMODELID_UNLIT;
+	#endif
+
 	BRANCH
-	if (SubstratePixelHeader.ClosureCount > 0) // This test is also enough to exclude sky pixels
+	if (bMaterialIsLit) // This test is also enough to exclude sky pixels
 	{
 		float DeviceZ = SampleDeviceZFromSceneTextures(BufferUV);
 		float SceneDepth = ConvertFromDeviceZ(DeviceZ);
@@ -460,8 +482,6 @@ void ReflectionEnvironmentSkyLighting(
 		FUpsampleDFAOOutput UpsampleDFAOOutput = UpsampleDFAO(BentNormalUV, SceneDepth);
 	#endif
 	
-		float CombinedScreenAndMaterialAO = SubstrateGetAO(SubstratePixelHeader) * AmbientOcclusion;
-	
 		
 	#if USE_DEFAULT_ENV_LIGHTING_INPUT && FEATURE_LEVEL >= FEATURE_LEVEL_SM5
 		float2 LocalPosition = SvPosition.xy - View.ViewRectMin.xy;
@@ -478,12 +498,43 @@ void ReflectionEnvironmentSkyLighting(
 
 		const bool bPixelMayHaveSSR = REFLECTION_SOURCE_TYPE == REFLECTION_SOURCE_LUMEN_STANDALONE || TopLayerSpecularContributionFactor < 1.0f;
 		const bool bEnableSpecular = ReflectionStruct.SkyLightParameters.y > 0.0f || NumCulledReflectionCaptures > 0 || bPixelMayHaveSSR;
-	
+
+	#if SUBSTRATE_LOAD_FROM_MATERIALCONTAINER
 		Substrate_for (uint ClosureIndex = 0, ClosureIndex < SubstratePixelHeader.ClosureCount, ++ClosureIndex)
+	#else
+		uint ClosureIndex = 0;
+	#endif
 		{
+		#if SUBSTRATE_LOAD_FROM_MATERIALCONTAINER
 			FSubstrateBSDF BSDF = UnpackSubstrateBSDF(Substrate.MaterialTextureArray, SubstrateAddressing, SubstratePixelHeader);
 			FSubstrateBSDFContext SubstrateBSDFContext = SubstrateCreateBSDFContext(SubstratePixelHeader, BSDF, SubstrateAddressing, V);
-			const float3 BSDFThroughput = LuminanceWeight(SubstrateBSDFContext, BSDF); // Use the reflected direction
+			float CombinedScreenAndMaterialAO = SubstrateGetAO(SubstratePixelHeader) * AmbientOcclusion;
+			FSubstrateIrradianceAndOcclusion SubstrateIrradianceAndOcclusion = SubstrateGetIrradianceAndAO(SubstratePixelHeader);
+		#else
+
+			FSubstrateGBufferBSDF GBufferBSDF = SubstrateReadGBufferBSDF(ScreenSpaceData);
+			float CombinedScreenAndMaterialAO = 1.0; // AmbientOcclusion;
+			
+			FSubstrateBSDFContext SubstrateBSDFContext = SubstrateCreateBSDFContext(GBufferBSDF.BSDFTangentBasis, GBufferBSDF.BSDF, V, V/*L*/, false/*bHasValidL*/);
+
+			// Similar to SubstrateUnpackIrradianceAndOcclusion
+			FSubstrateIrradianceAndOcclusion SubstrateIrradianceAndOcclusion = InitIrradianceAndOcclusion();
+		#if GBUFFER_HAS_DIFFUSE_SAMPLE_OCCLUSION
+			SubstrateIrradianceAndOcclusion.DiffuseIndirectSampleOcclusion = GBuffer.DiffuseIndirectSampleOcclusion;
+			SubstrateIrradianceAndOcclusion.MaterialAO = saturate(1.0 - float(countbits(Out.DiffuseIndirectSampleOcclusion)) * rcp(float(INDIRECT_SAMPLE_COUNT)));
+			SubstrateIrradianceAndOcclusion.IndirectIrradiance = 1.f;
+		#elif ALLOW_STATIC_LIGHTING
+			SubstrateIrradianceAndOcclusion.MaterialAO = 1.f;
+			SubstrateIrradianceAndOcclusion.DiffuseIndirectSampleOcclusion = 0x0u;
+			SubstrateIrradianceAndOcclusion.IndirectIrradiance = GBuffer.IndirectIrradiance;
+		#else
+			SubstrateIrradianceAndOcclusion.MaterialAO = GBuffer.GBufferAO;
+			SubstrateIrradianceAndOcclusion.DiffuseIndirectSampleOcclusion = 0x0u;
+			SubstrateIrradianceAndOcclusion.IndirectIrradiance = 1.f;
+		#endif
+
+		#endif
+			const float3 BSDFThroughput = LuminanceWeight(SubstrateBSDFContext, SubstrateBSDFContext.BSDF); // Use the reflected direction
 	
 			// Evaluate environment lighting
 			FSubstrateEnvLightResult SubstrateEnvLight = SubstrateEvaluateForEnvLight(SubstrateBSDFContext, bEnableSpecular, Settings);
@@ -498,9 +549,9 @@ void ReflectionEnvironmentSkyLighting(
 			float3 SpecularLighting = 0;
 			SubstrateEnvLightingCommon(
 				SubstrateEnvLight,
-				SubstratePixelHeader,
+				SubstrateIrradianceAndOcclusion,
 				SubstrateBSDFContext,
-				BSDF,
+				SubstrateBSDFContext.BSDF,
 				BentNormal,
 				BSDFThroughput,
 				CaptureDataStartIndex,
@@ -515,7 +566,7 @@ void ReflectionEnvironmentSkyLighting(
 	
 		#if REFLECTION_SOURCE_TYPE == REFLECTION_SOURCE_LUMEN_STANDALONE
 			{
-				const FLumenReflectionOnlyOutput LumenReflections = GetLumenReflectionOnly(SubstrateEnvLight, BSDF, BSDFThroughput, Texture2DArraySample(LumenReflectionTexture, LumenReflectionTextureSampler, float3(BufferUV, ClosureIndex)));
+				const FLumenReflectionOnlyOutput LumenReflections = GetLumenReflectionOnly(SubstrateEnvLight, SubstrateBSDFContext.BSDF, BSDFThroughput, Texture2DArraySample(LumenReflectionTexture, LumenReflectionTextureSampler, float3(BufferUV, ClosureIndex)));
 
 				SpecularLighting *= min(LumenReflections.LightingLobe0.a, LumenReflections.LightingLobe1.a);
 
@@ -533,15 +584,15 @@ void ReflectionEnvironmentSkyLighting(
 				#endif 
 			}
 		#else
-			if (BSDF_GETISTOPLAYER(BSDF))
+			if (BSDF_GETISTOPLAYER(SubstrateBSDFContext.BSDF))
 			{
-				GetSSSRContribution(SubstrateEnvLight, BSDF, BSDFThroughput, SSRContribution);
+				GetSSSRContribution(SubstrateEnvLight, SubstrateBSDFContext.BSDF, BSDFThroughput, SSRContribution);
 			}
 		#endif
 			
 			FLightAccumulator Out = (FLightAccumulator)0;
 			LightAccumulator_AddSplit(Out, DiffuseLighting, SpecularLighting, DiffuseLighting, 1.f /*CommonMultiplier*/, SubstrateEnvLight.bPostProcessSubsurface);
-			AccumulateSubstrateDeferredLighting(SubstrateLighting, Out, SubstrateEnvLight.bPostProcessSubsurface, BSDF_GETISTOPLAYER(BSDF));
+			AccumulateSubstrateDeferredLighting(SubstrateLighting, Out, SubstrateEnvLight.bPostProcessSubsurface, BSDF_GETISTOPLAYER(SubstrateBSDFContext.BSDF));
 		}
 
 		OutColor = SubstrateLighting.SceneColor * (REFLECTION_SOURCE_TYPE == REFLECTION_SOURCE_LUMEN_STANDALONE ? 1.0f : View.PreExposure);
@@ -558,7 +609,7 @@ void ReflectionEnvironmentSkyLighting(
 		#endif
 	}
 
-#else // SUBSTRATE_GBUFFER_FORMAT==1
+#else // SUBSTRATE_ENABLED==1
 
 	// Sample scene textures.
 	FGBufferData GBuffer = GetGBufferDataFromSceneTextures(BufferUV);
@@ -605,6 +656,6 @@ void ReflectionEnvironmentSkyLighting(
 		OutColor.xyz += ReflectionEnvironment(GBuffer, AmbientOcclusion, BufferUV, ScreenPosition, SvPosition, BentNormal, SpecularColor, ShadingModelID);
 	}
 
-#endif // SUBSTRATE_GBUFFER_FORMAT==1
+#endif // SUBSTRATE_ENABLED==1
 }
 
diff --git a/Engine/Shaders/Private/Substrate/SubstrateEvaluation.ush b/Engine/Shaders/Private/Substrate/SubstrateEvaluation.ush
--- a/Engine/Shaders/Private/Substrate/SubstrateEvaluation.ush
+++ b/Engine/Shaders/Private/Substrate/SubstrateEvaluation.ush
@@ -1917,6 +1917,7 @@ FSubstrateEnvLightResult SubstrateEvaluateForEnvLight(FSubstrateBSDFContext BSDF
 				SubstrateEnvLightResult.SpecularSafeRoughness	= SafeRoughness;
 				SubstrateEnvLightResult.SpecularColor			= F0;
 			}
+			SubstrateEnvLightResult.SSRSpecularWeight = SubstrateEnvLightResult.SpecularWeight;
 
 			break;
 		}
diff --git a/Engine/Shaders/Private/Substrate/SubstrateLightingCommon.ush b/Engine/Shaders/Private/Substrate/SubstrateLightingCommon.ush
--- a/Engine/Shaders/Private/Substrate/SubstrateLightingCommon.ush
+++ b/Engine/Shaders/Private/Substrate/SubstrateLightingCommon.ush
@@ -203,7 +203,7 @@ float3 GetEnvSpecularLighting(
 
 void SubstrateEnvLightingCommon(
 	in FSubstrateEnvLightResult SubstrateEnvLight,
-	in FSubstratePixelHeader SubstratePixelHeader,
+	in FSubstrateIrradianceAndOcclusion SubstrateIrradianceAndOcclusion,
 	in FSubstrateBSDFContext SubstrateBSDFContext,
 	in FSubstrateBSDF BSDF,
 	in float3 BentNormal,
@@ -218,7 +218,6 @@ void SubstrateEnvLightingCommon(
 	inout float3 DiffuseLighting,
 	inout float3 SpecularLighting)
 {
-	FSubstrateIrradianceAndOcclusion SubstrateIrradianceAndOcclusion = SubstrateGetIrradianceAndAO(SubstratePixelHeader);
 
 	// Diffuse component
 	DiffuseLighting = 0;
@@ -292,7 +291,7 @@ void SubstrateEnvLightingCommon(
 	#endif
 
 		float RoughnessSquared = SubstrateEnvLight.SpecularSafeRoughness * SubstrateEnvLight.SpecularSafeRoughness;
-		float SpecularOcclusion = IndirectSpecularOcclusion * GetSpecularOcclusion(SubstrateBSDFContext.SatNoV, RoughnessSquared, CombinedScreenAndMaterialAO);
+		float SpecularOcclusion = GetSpecularOcclusion(SubstrateBSDFContext.SatNoV, RoughnessSquared, CombinedScreenAndMaterialAO); // Do not mul with IndirectSpecularOcclusion here otherwise BentNormalAO will be applied twice.
 
 		// When using bHazeAsSimpleClearCoat, the main specular lobe represents the bottom layer's specular, on which 
 		// the TopLayerSpecularContribution (i.e. SSR) shouldn't be applied to
diff --git a/Engine/Shaders/Private/Substrate/SubstrateRead.ush b/Engine/Shaders/Private/Substrate/SubstrateRead.ush
new file mode 100644
--- /dev/null
+++ b/Engine/Shaders/Private/Substrate/SubstrateRead.ush
@@ -0,0 +1,199 @@
+// Copyright Epic Games, Inc. All Rights Reserved.
+
+#pragma once
+
+// Sanity guard.  
+#ifndef SUBSTRATE_ENABLED
+#define SUBSTRATE_ENABLED 1
+#error SUBSTRATE_ENABLED needs to be defined
+#endif
+
+#if SUBSTRATE_ENABLED
+
+#include "../Common.ush"
+#include "Substrate.ush"
+
+struct FSubstrateGBufferBSDF
+{
+	FSubstrateBSDF BSDF;
+	float3x3 BSDFTangentBasis;
+	float BSDFAO;
+};
+
+FSubstrateGBufferBSDF SubstrateReadGBufferBSDF(in FScreenSpaceData ScreenSpaceData)
+{
+	const FGBufferData GBuffer = ScreenSpaceData.GBuffer;
+
+	FSubstrateBSDF BSDF = (FSubstrateBSDF)0;
+	float3x3 BSDFTangentBasis = (float3x3)0;
+	float BSDFAO = 0.0f;
+
+	if (GBuffer.Anisotropy != 0)
+	{
+		BSDFTangentBasis[0] = GBuffer.WorldTangent;
+		BSDFTangentBasis[1] = cross(GBuffer.WorldNormal, GBuffer.WorldTangent);
+		BSDFTangentBasis[2] = GBuffer.WorldNormal;
+	}
+	else
+	{
+		BSDFTangentBasis = GetTangentBasis(GBuffer.WorldNormal);
+	}
+
+	BSDFAO = ScreenSpaceData.AmbientOcclusion;
+
+	{
+		FSubstratePixelFootprint Footprint = SubstrateGetPixelFootprint(1.0f, 1.0f, GBuffer.Curvature);
+
+		const float3 Emissive = 0.0f;
+		const float  SpecularProfileId = 0.0f;
+
+		const float  GlintValue = 1.0f;
+		const float2 GlintUV = 0.0f;
+		const float2 GlintUVdx = 0.0f;
+		const float2 GlintUVdy = 0.0f;
+
+		float3		SSSMFP = 0.0f;
+		float		SSSMFPScale = 1.0f;
+		float		SSSPhaseAnisotropy = 0.0f;
+
+		float		SecondRoughness = 0.0f;
+		float		SecondRoughnessWeight = 0.0f;
+		float		SecondRoughnessAsSimpleClearCoat = 0.0f;
+
+		float		FuzzAmount = 0.0f;
+		float3		FuzzColor = 0.0f;
+		float		FuzzRoughness = 0.0f;
+
+		float		SSSProfileID = SSS_PROFILE_ID_INVALID;
+		uint		SSSType = SSS_TYPE_NONE;
+		float		ThicknessCm = 1.0f;
+		bool		bIsThinSurface = false;
+
+		float3 DiffuseAlbedo = ComputeDiffuseAlbedo(GBuffer.BaseColor, GBuffer.Metallic);
+		float3 F0 = ComputeF0(GBuffer.Specular, GBuffer.BaseColor, GBuffer.Metallic);
+		float3 F90 = 1.0f;
+
+		float ClearCoatUseSecondNormal = 0.0f;
+		float3 ClearCoatBottomNormal = GBuffer.WorldNormal;
+
+		bool   bIsAtTheBottomOfTopology = true;
+		uint   SharedLocalBasisIndex = 0;
+		uint   SharedLocalBasisTypes = 0;
+
+		if (GBuffer.ShadingModelID == SHADINGMODELID_HAIR)
+		{
+			const float Scatter = GBuffer.Metallic;
+			const float BackLit = GBuffer.CustomData.b;
+			BSDF = GetSubstrateHairBSDF(GBuffer.BaseColor, Scatter, GBuffer.Specular, GBuffer.Roughness, BackLit, Emissive, SharedLocalBasisIndex).InlinedBSDF;
+
+		}
+		else if (GBuffer.ShadingModelID == SHADINGMODELID_EYE)
+		{
+			SSSProfileID = ExtractSubsurfaceProfileInt(GBuffer.CustomData.x);
+			const float IrisMask = 1.0f - GBuffer.CustomData.w;
+			float IrisDistance = 0.0;
+			float3 IrisNormal = GBuffer.WorldNormal;
+			float3 IrisPlaneNormal = GBuffer.WorldNormal;
+
+		#if IRIS_NORMAL
+			const float2 IrisNormalDelta = float2(GBuffer.CustomData.y, GBuffer.CustomData.z) * 2 - (256.0 / 255.0);
+			const float2 WorldNormalOct = UnitVectorToOctahedron(GBuffer.WorldNormal);
+			IrisNormal = OctahedronToUnitVector(WorldNormalOct + IrisNormalDelta);
+		#else
+			IrisNormal = OctahedronToUnitVector(GBuffer.CustomData.yz * 2 - 1);
+			IrisDistance = GBuffer.StoredMetallic;
+		#endif
+			BSDF = GetSubstrateEyeBSDF(GBuffer.StoredBaseColor, GBuffer.Roughness, IrisMask, IrisDistance, IrisNormal, IrisPlaneNormal, SSSProfileID, Emissive, SharedLocalBasisIndex).InlinedBSDF;
+		}
+		else
+		{
+			if (GBuffer.ShadingModelID == SHADINGMODELID_CLOTH)
+			{
+				FuzzColor = ExtractSubsurfaceColor(GBuffer);
+				FuzzAmount = GBuffer.CustomData.a;
+				FuzzRoughness = GBuffer.Roughness;
+			}
+			else if (GBuffer.ShadingModelID == SHADINGMODELID_SUBSURFACE_PROFILE)
+			{
+				SSSType = SSS_TYPE_DIFFUSION;	//  SSS_TYPE_DIFFUSION_PROFILE will be selected in InternalGetSubstrateSlabBSDF based on SSSProfileID and SSSMFPScale.
+				SSSMFP = 0;
+				SSSProfileID = ExtractSubsurfaceProfileInt(GBuffer.CustomData.x);
+				SSSMFPScale = GBuffer.CustomData.a;
+				SSSPhaseAnisotropy = 0.93f;
+			}
+			else if (GBuffer.ShadingModelID == SHADINGMODELID_SUBSURFACE)
+			{
+				SSSType = SSS_TYPE_WRAP;
+				{
+					// See SubstrateGetBSDFSubSurfaceColor
+					float3 Extinction = TransmittanceToExtinction(ExtractSubsurfaceColor(GBuffer), SUBSTRATE_SIMPLEVOLUME_THICKNESS_CM * CENTIMETER_TO_METER);
+					SSSMFP = (1.0f / max(SUBSTRATE_EPSILON, Extinction));
+				}
+				SSSPhaseAnisotropy = GBuffer.CustomData.a;		// TODO test, should be one minus opacity
+				SSSMFPScale = 1.0f;
+			}
+			else if (GBuffer.ShadingModelID == SHADINGMODELID_TWOSIDED_FOLIAGE)
+			{
+				SSSType = SSS_TYPE_TWO_SIDED_WRAP;
+				{
+					// See SubstrateGetBSDFSubSurfaceColor
+					float3 Extinction = TransmittanceToExtinction(ExtractSubsurfaceColor(GBuffer), SUBSTRATE_SIMPLEVOLUME_THICKNESS_CM * CENTIMETER_TO_METER);
+					SSSMFP = (1.0f / max(SUBSTRATE_EPSILON, Extinction));
+				}
+				SSSPhaseAnisotropy = GBuffer.CustomData.a;		// TODO test, should be one minus opacity
+				SSSMFPScale = 1.0f;
+			}
+			else if (GBuffer.ShadingModelID == SHADINGMODELID_CLEAR_COAT)
+			{
+			#if SUBSTRATE_FASTPATH==0 && CLEAR_COAT_BOTTOM_NORMAL 
+				ClearCoatUseSecondNormal = 1.0f;
+				const float2 oct1 = ((float2(GBuffer.CustomData.a, GBuffer.CustomData.z) * 4) - (512.0 / 255.0)) + UnitVectorToOctahedron(GBuffer.WorldNormal);
+				ClearCoatBottomNormal = OctahedronToUnitVector(oct1);
+			#endif
+
+				SecondRoughnessWeight = GBuffer.CustomData.x;
+				SecondRoughness = GBuffer.CustomData.y;
+				SecondRoughnessAsSimpleClearCoat = (SecondRoughnessWeight > 0.0f || any(GBuffer.WorldNormal != ClearCoatBottomNormal)) ? 1.0f : 0.0f;
+			}
+
+			BSDF = InternalGetSubstrateSlabBSDF(
+				Footprint,
+				GBuffer.WorldNormal,
+				DiffuseAlbedo, F0, F90,
+				GBuffer.Roughness, GBuffer.Anisotropy,
+				SSSProfileID, SSSMFP, SSSMFPScale, SSSPhaseAnisotropy, SSSType,
+				Emissive,
+				SecondRoughness, SecondRoughnessWeight, SecondRoughnessAsSimpleClearCoat,
+				ClearCoatUseSecondNormal, ClearCoatBottomNormal,
+				FuzzAmount, FuzzColor, FuzzRoughness,
+				GlintValue, GlintUV, GlintUVdx, GlintUVdy,
+				SpecularProfileId,
+				ThicknessCm,
+				bIsThinSurface,
+				bIsAtTheBottomOfTopology,
+				SharedLocalBasisIndex, SharedLocalBasisTypes).InlinedBSDF;
+		}
+
+		BSDF.Coverage = 1;
+		BSDF.Emissive = 0;
+		BSDF.ThicknessCm = (SSSType == SSS_TYPE_WRAP || SSSType == SSS_TYPE_TWO_SIDED_WRAP) ? SUBSTRATE_SIMPLEVOLUME_THICKNESS_CM : SUBSTRATE_LAYER_DEFAULT_THICKNESS_CM; // Identity rescaling in InternalGetSubstrateSlabBSDF 
+		BSDF.TmpMFP = SSSMFP;
+		BSDF.TopLayerDataWeight = 1;
+
+		BSDF.OperatorIndex = 0;
+		BSDF.LuminanceWeightV = 1.0f;
+		BSDF.CoverageAboveAlongN = 0.0f;
+		BSDF.TransmittanceAboveAlongN = 1.0f;
+		BSDF.bIsBottom = true;
+		BSDF.bIsTop = true;
+		BSDF_SETISTOPLAYER(BSDF, 1);
+	}
+
+	FSubstrateGBufferBSDF Out;
+	Out.BSDF = BSDF;
+	Out.BSDFTangentBasis = BSDFTangentBasis;
+	Out.BSDFAO = BSDFAO;
+	return Out;
+}
+
+#endif // SUBSTRATE_ENABLED
