[VSM] Allow hiding parts of the stats to reduce clutter. Moved the cluster stats to a separate panel and make them easier to read. Fixed a bug where the cluster stats tried to retrieve the name of a deallocated light.

diff --git a/Engine/Shaders/Private/VirtualShadowMaps/VirtualShadowMapPrintStats.usf b/Engine/Shaders/Private/VirtualShadowMaps/VirtualShadowMapPrintStats.usf
--- a/Engine/Shaders/Private/VirtualShadowMaps/VirtualShadowMapPrintStats.usf
+++ b/Engine/Shaders/Private/VirtualShadowMaps/VirtualShadowMapPrintStats.usf
@@ -220,10 +220,10 @@ void BuildNaniteGeoHistogram(uint ThreadId)
 
 void PrintStats()
 {
-	const bool bShowAll                = ShowStatsValue == 1;
-	// Work around potential compiler bug...
-	const bool bShowPhysicalPages      = (ShowStatsValue == 1) || ShowStatsValue == 2;
-	const bool bShowNonNaniteInstances = (ShowStatsValue == 1) || ShowStatsValue == 3;
+	const bool bShowBasic              = (ShowStatsValue & EVSMStatSection::Basic) > 0;
+	const bool bShowPhysicalPages      = (ShowStatsValue & EVSMStatSection::PhysicalPages) > 0;
+	const bool bShowNonNaniteInstances = (ShowStatsValue & EVSMStatSection::NonNanite) > 0;
+	const bool bShowNaniteClusters     = (ShowStatsValue & EVSMStatSection::Clusters) > 0;
 
 	const float TopMargin = 0.05f;
 	const float HeadlineX = 0.47f;
@@ -232,7 +232,7 @@ void PrintStats()
 	const int NumLabelCharacters = 26;
 	const float ValueX = HeadlineX + Context.Config.FontSpacing.x * NumLabelCharacters;
 
-	if (bShowPhysicalPages)
+	if (bShowBasic || bShowPhysicalPages)
 	{
 		Context.Pos.x = HeadlineX;
 		Print(Context, TEXT("Physical Pages"), FontWhite);
@@ -246,66 +246,75 @@ void PrintStats()
 		PrintLeftAlign(Context, InStatsBuffer[VSM_STAT_REQUESTED_THIS_FRAME_PAGES]);
 		Newline(Context);
 
-		Context.Pos.x = ItemX;
-		Print(Context, TEXT("Empty"), FontWhite);
-		Context.Pos.x = ValueX;
-		PrintLeftAlign(Context, InStatsBuffer[VSM_STAT_EMPTY_PAGES]);
-		Newline(Context);
-
-		Context.Pos.x = ItemX;
-		Print(Context, TEXT("Allocated New"), FontWhite);
-		Context.Pos.x = ValueX;
-		PrintLeftAlign(Context, InStatsBuffer[VSM_STAT_ALLOCATED_NEW]);
-		Newline(Context);
-
-		Context.Pos.x = ItemX;
-		Print(Context, TEXT("Cleared"), FontWhite);
-		Context.Pos.x = ValueX;
-		PrintLeftAlign(Context, InStatsBuffer[VSM_STAT_NUM_PAGES_TO_CLEAR]);
-		Newline(Context);
-
-		Context.Pos.x = ItemX;
-		Print(Context, TEXT("HZB Built"), FontWhite);
-		Context.Pos.x = ValueX;
-		PrintLeftAlign(Context, InStatsBuffer[VSM_STAT_NUM_HZB_PAGES_BUILT]);
-		Newline(Context);
-
+		if (bShowPhysicalPages)
 		{
-			Context.Pos.x = HeadlineX;
-			Print(Context, TEXT("Static Cached"), FontWhite);
+			Context.Pos.x = ItemX;
+			Print(Context, TEXT("Empty"), FontWhite);
 			Context.Pos.x = ValueX;
-			PrintLeftAlign(Context, InStatsBuffer[VSM_STAT_STATIC_CACHED_PAGES]);
+			PrintLeftAlign(Context, InStatsBuffer[VSM_STAT_EMPTY_PAGES]);
 			Newline(Context);
 
 			Context.Pos.x = ItemX;
-			Print(Context, TEXT("Invalidated"), FontWhite);
+			Print(Context, TEXT("Allocated New"), FontWhite);
 			Context.Pos.x = ValueX;
-			PrintLeftAlign(Context, InStatsBuffer[VSM_STAT_STATIC_INVALIDATED_PAGES]);
+			PrintLeftAlign(Context, InStatsBuffer[VSM_STAT_ALLOCATED_NEW]);
 			Newline(Context);
 
-			Context.Pos.x = HeadlineX;
-			Print(Context, TEXT("Dynamic Cached"), FontWhite);
+			Context.Pos.x = ItemX;
+			Print(Context, TEXT("Cleared"), FontWhite);
 			Context.Pos.x = ValueX;
-			PrintLeftAlign(Context, InStatsBuffer[VSM_STAT_DYNAMIC_CACHED_PAGES]);
+			PrintLeftAlign(Context, InStatsBuffer[VSM_STAT_NUM_PAGES_TO_CLEAR]);
 			Newline(Context);
 
 			Context.Pos.x = ItemX;
-			Print(Context, TEXT("Invalidated"), FontWhite);
+			Print(Context, TEXT("HZB Built"), FontWhite);
 			Context.Pos.x = ValueX;
-			PrintLeftAlign(Context, InStatsBuffer[VSM_STAT_DYNAMIC_INVALIDATED_PAGES]);
+			PrintLeftAlign(Context, InStatsBuffer[VSM_STAT_NUM_HZB_PAGES_BUILT]);
 			Newline(Context);
+		}
 
-			Context.Pos.x = ItemX;
-			Print(Context, TEXT("Considered for WPO"), FontWhite);
+		{
+			Context.Pos.x = HeadlineX;
+			Print(Context, TEXT("Static Cached"), FontWhite);
 			Context.Pos.x = ValueX;
-			PrintLeftAlign(Context, InStatsBuffer[VSM_STAT_WPO_CONSIDERED_PAGES]);
+			PrintLeftAlign(Context, InStatsBuffer[VSM_STAT_STATIC_CACHED_PAGES]);
 			Newline(Context);
 
-			Context.Pos.x = ItemX;
-			Print(Context, TEXT("Merged"), FontWhite);
+			if (bShowPhysicalPages)
+			{
+				Context.Pos.x = ItemX;
+				Print(Context, TEXT("Invalidated"), FontWhite);
+				Context.Pos.x = ValueX;
+				PrintLeftAlign(Context, InStatsBuffer[VSM_STAT_STATIC_INVALIDATED_PAGES]);
+				Newline(Context);
+			}
+
+			Context.Pos.x = HeadlineX;
+			Print(Context, TEXT("Dynamic Cached"), FontWhite);
 			Context.Pos.x = ValueX;
-			PrintLeftAlign(Context, InStatsBuffer[VSM_STAT_NUM_PAGES_TO_MERGE]);
+			PrintLeftAlign(Context, InStatsBuffer[VSM_STAT_DYNAMIC_CACHED_PAGES]);
 			Newline(Context);
+
+			if (bShowPhysicalPages)
+			{
+				Context.Pos.x = ItemX;
+				Print(Context, TEXT("Invalidated"), FontWhite);
+				Context.Pos.x = ValueX;
+				PrintLeftAlign(Context, InStatsBuffer[VSM_STAT_DYNAMIC_INVALIDATED_PAGES]);
+				Newline(Context);
+
+				Context.Pos.x = ItemX;
+				Print(Context, TEXT("Considered for WPO"), FontWhite);
+				Context.Pos.x = ValueX;
+				PrintLeftAlign(Context, InStatsBuffer[VSM_STAT_WPO_CONSIDERED_PAGES]);
+				Newline(Context);
+
+				Context.Pos.x = ItemX;
+				Print(Context, TEXT("Merged"), FontWhite);
+				Context.Pos.x = ValueX;
+				PrintLeftAlign(Context, InStatsBuffer[VSM_STAT_NUM_PAGES_TO_MERGE]);
+				Newline(Context);
+			}
 		}
 
 		Context.Pos.x = HeadlineX;
@@ -378,74 +387,86 @@ void PrintStats()
 
 	uint ClustersHW = NanitePerformanceFeedback[VSM_NPF_HEADER_TOTAL_HW_CLUSTERS];
 	uint ClustersSW = NanitePerformanceFeedback[VSM_NPF_HEADER_TOTAL_SW_CLUSTERS];
-	Context.Pos.x = HeadlineX;
-	Print(Context, TEXT("Nanite Clusters HW"), FontWhite);
-	Context.Pos.x = ValueX;
-	PrintLeftAlign(Context, ClustersHW);
-	Newline(Context);
 
-	Context.Pos.x = HeadlineX;
-	Print(Context, TEXT("Nanite Clusters SW"), FontWhite);
-	Context.Pos.x = ValueX;
-	PrintLeftAlign(Context, ClustersSW);
-	Newline(Context);
-
-	Context.Pos.x = HeadlineX;
-	Print(Context, TEXT("Nanite Clusters Per VSM"), FontWhite);
-	Newline(Context);
-	
-	float2 HistogramScale = float2(4, 200);
-	uint HistogramMax = max(4096U, 2U << firstbithigh(NaniteGeoHistogramMax));
-	float4 HistogramMaxIndicatorColor = float4(GreenToRedTurbo((firstbithigh(HistogramMax) - 12) / 4.0f), 1.0); // green to red, 4096 to 32768
-	if (NaniteGeoHistogramNonZeroBins < 64)
+	if (bShowBasic && !bShowNaniteClusters)
 	{
-		HistogramScale.x = 8;
+		Context.Pos.x = HeadlineX;
+		Print(Context, TEXT("Nanite Clusters"), FontWhite);
+		Context.Pos.x = ValueX;
+		PrintLeftAlign(Context, ClustersHW + ClustersSW);
+		Newline(Context);
 	}
+	else if (bShowNaniteClusters)
+	{
+		Context.Pos.x = HeadlineX;
+		Print(Context, TEXT("Nanite Clusters HW"), FontWhite);
+		Context.Pos.x = ValueX;
+		PrintLeftAlign(Context, ClustersHW);
+		Newline(Context);
 
-	Print(Context, HistogramMax, FontGrey);
-	Context.Pos.x = HeadlineX;
-	Context.Pos.y += HistogramScale.y / Context.Config.Resolution.y;
-	Print(Context, 0U, FontGrey);
-	Context.Pos.x = HeadlineX;
-	Context.Pos.y -= HistogramScale.y / Context.Config.Resolution.y;
+		Context.Pos.x = HeadlineX;
+		Print(Context, TEXT("Nanite Clusters SW"), FontWhite);
+		Context.Pos.x = ValueX;
+		PrintLeftAlign(Context, ClustersSW);
+		Newline(Context);
 
-	Context.Pos *= float2(Context.Config.Resolution);
-	float4 HistogramRect = float4(
-		Context.Pos,
-		Context.Pos + float2(HistogramScale.x * max(50, NaniteGeoHistogramNonZeroBins), HistogramScale.y)
-	);
-	AddFilledQuadSS(Context, HistogramRect.xy, HistogramRect.zw, float4(HistogramMaxIndicatorColor.rgb * 0.25,0.4));
-	Context.Pos.y += HistogramScale.y;
-	for (int Index = 0; Index < NANITE_GEO_HISTOGRAM_NUM_BINS; Index++)
-	{
-		float HWClusters = (NaniteGeoHistogramHWClusters[Index] / float(HistogramMax)) * HistogramScale.y;
-		float SWClusters = (NaniteGeoHistogramSWClusters[Index] / float(HistogramMax)) * HistogramScale.y;
-		
-		bool bIsSinglePage = Index < VSM_GENERATE_STATS_GROUPSIZE;
-		FVirtualShadowMapHandle VSMHandle;
-		if (bIsSinglePage)
-		{
-			VSMHandle = FVirtualShadowMapHandle::MakeFromId(Index);
-		}
-		else
+		Context.Pos.x = HeadlineX;
+		Print(Context, TEXT("Nanite Clusters Per VSM"), FontWhite);
+		Newline(Context);
+	
+		float2 HistogramScale = float2(4, 200);
+		uint HistogramMax = max(4096U, 2U << firstbithigh(NaniteGeoHistogramMax));
+		float4 HistogramMaxIndicatorColor = float4(GreenToRedTurbo((firstbithigh(HistogramMax) - 12) / 4.0f), 1.0); // green to red, 4096 to 32768
+		if (NaniteGeoHistogramNonZeroBins < 64)
 		{
-			VSMHandle = FVirtualShadowMapHandle::MakeFromFullIndex(Index - VSM_GENERATE_STATS_GROUPSIZE);
+			HistogramScale.x = 8;
 		}
-		float ThrottleIntensity = 0.0f;
-		if (bThrottlingEnabled)
+
+		Print(Context, HistogramMax, FontGrey);
+		Context.Pos.x = HeadlineX;
+		Context.Pos.y += HistogramScale.y / Context.Config.Resolution.y;
+		Print(Context, 0U, FontGrey);
+		Context.Pos.x = HeadlineX;
+		Context.Pos.y -= HistogramScale.y / Context.Config.Resolution.y;
+
+		Context.Pos *= float2(Context.Config.Resolution);
+		float4 HistogramRect = float4(
+			Context.Pos,
+			Context.Pos + float2(HistogramScale.x * max(50, NaniteGeoHistogramNonZeroBins), HistogramScale.y)
+		);
+		AddFilledQuadSS(Context, HistogramRect.xy, HistogramRect.zw, float4(HistogramMaxIndicatorColor.rgb * 0.25,0.4));
+		Context.Pos.y += HistogramScale.y;
+		for (int Index = 0; Index < NANITE_GEO_HISTOGRAM_NUM_BINS; Index++)
 		{
-			ThrottleIntensity = asfloat(ThrottleBuffer[VSM_TB_SIZEOF_HEADER + VSM_TB_SIZEOF_ENTRY * VSMHandle.GetDataIndex() + VSM_TB_ENTRY_THROTTLE]);
-		}
+			float HWClusters = (NaniteGeoHistogramHWClusters[Index] / float(HistogramMax)) * HistogramScale.y;
+			float SWClusters = (NaniteGeoHistogramSWClusters[Index] / float(HistogramMax)) * HistogramScale.y;
+		
+			bool bIsSinglePage = Index < VSM_GENERATE_STATS_GROUPSIZE;
+			FVirtualShadowMapHandle VSMHandle;
+			if (bIsSinglePage)
+			{
+				VSMHandle = FVirtualShadowMapHandle::MakeFromId(Index);
+			}
+			else
+			{
+				VSMHandle = FVirtualShadowMapHandle::MakeFromFullIndex(Index - VSM_GENERATE_STATS_GROUPSIZE);
+			}
+			float ThrottleIntensity = 0.0f;
+			if (bThrottlingEnabled)
+			{
+				ThrottleIntensity = asfloat(ThrottleBuffer[VSM_TB_SIZEOF_HEADER + VSM_TB_SIZEOF_ENTRY * VSMHandle.GetDataIndex() + VSM_TB_ENTRY_THROTTLE]);
+			}
 
-		bool bHasClusterData = HWClusters > 0 || SWClusters > 0;
+			bool bHasClusterData = HWClusters > 0 || SWClusters > 0;
 
-		if (HWClusters > 0) { AddFilledQuadSS(Context, Context.Pos, Context.Pos + float2(HistogramScale.x, -HWClusters), float4(0,0,1,1.0)); }
-		if (SWClusters > 0) { AddFilledQuadSS(Context, Context.Pos + float2(0, -HWClusters), Context.Pos + float2(0, -HWClusters) + float2(HistogramScale.x, -SWClusters), float4(0,1,1,1.0)); }
-		if (bHasClusterData && ThrottleIntensity > 0) { AddFilledQuadSS(Context, Context.Pos, Context.Pos + float2(HistogramScale.x, ThrottleIntensity * HistogramScale.y), float4(1,0,0,1.0)); }
-		if (bHasClusterData) { Context.Pos.x += HistogramScale.x; }
+			if (HWClusters > 0) { AddFilledQuadSS(Context, Context.Pos, Context.Pos + float2(HistogramScale.x, -HWClusters), float4(0,0,1,1.0)); }
+			if (SWClusters > 0) { AddFilledQuadSS(Context, Context.Pos + float2(0, -HWClusters), Context.Pos + float2(0, -HWClusters) + float2(HistogramScale.x, -SWClusters), float4(0,1,1,1.0)); }
+			if (bHasClusterData && ThrottleIntensity > 0) { AddFilledQuadSS(Context, Context.Pos, Context.Pos + float2(HistogramScale.x, ThrottleIntensity * HistogramScale.y), float4(1,0,0,1.0)); }
+			if (bHasClusterData) { Context.Pos.x += HistogramScale.x; }
+		}
+		Context.Pos /= float2(Context.Config.Resolution);
+		Newline(Context);
 	}
-	Context.Pos /= float2(Context.Config.Resolution);
-	Newline(Context);
 }
 
 
diff --git a/Engine/Shaders/Shared/VirtualShadowMapDefinitions.h b/Engine/Shaders/Shared/VirtualShadowMapDefinitions.h
--- a/Engine/Shaders/Shared/VirtualShadowMapDefinitions.h
+++ b/Engine/Shaders/Shared/VirtualShadowMapDefinitions.h
@@ -172,11 +172,34 @@ inline uint GetMipLevelLocal(float Footprint, uint MipMode, float ShadowMapResol
 	return MipLevel;
 }
 
+struct EVSMStatSection
+{
+	enum Type
+	{
+		None = 0,
+		Basic = 1U << 0,
+
+		PhysicalPages = 1U << 1,
+		NonNanite = 1U << 2,
+		Clusters = 1U << 3,
+
+		MAX = 1U << 4,
+
+		All = (MAX - 1)
+	};
+};
+
+#define ENUMERATE_VSM_STAT_SECTIONS(op) \
+	op(PhysicalPages) \
+	op(NonNanite) \
+	op(Clusters) \
+
 #ifdef __cplusplus
 } // namespace UE::HLSL
 
 using FVSMVisibleInstanceCmd = UE::HLSL::FVSMVisibleInstanceCmd;
 using FVSMCullingBatchInfo = UE::HLSL::FVSMCullingBatchInfo;
 using FNextVirtualShadowMapData = UE::HLSL::FNextVirtualShadowMapData;
+using EVSMStatSection = UE::HLSL::EVSMStatSection;
 
 #endif
diff --git a/Engine/Source/Editor/CommonMenuExtensions/Private/VirtualShadowMapVisualizationMenuCommands.cpp b/Engine/Source/Editor/CommonMenuExtensions/Private/VirtualShadowMapVisualizationMenuCommands.cpp
--- a/Engine/Source/Editor/CommonMenuExtensions/Private/VirtualShadowMapVisualizationMenuCommands.cpp
+++ b/Engine/Source/Editor/CommonMenuExtensions/Private/VirtualShadowMapVisualizationMenuCommands.cpp
@@ -144,17 +144,16 @@ void FVirtualShadowMapVisualizationMenuCommands::BindCommands(FUICommandList& Co
 		);
 	}
 
-	IConsoleVariable* ShowStatsCVar = IConsoleManager::Get().FindConsoleVariable(TEXT("r.Shadow.Virtual.ShowStats"));
-	check(ShowStatsCVar);
+	IConsoleVariable* ShowStatsVisibleCVar = IConsoleManager::Get().FindConsoleVariable(TEXT("r.Shadow.Virtual.ShowStats.Visible"));
+	check(ShowStatsVisibleCVar);
 	CommandList.MapAction(
 		ShowStatsCommand,
 		FExecuteAction::CreateLambda([=]()
 		{
-			bool bStatsVisible = ShowStatsCVar->GetInt() != 0;
-			ShowStatsCVar->Set(bStatsVisible ? 0 : 1);
+			ShowStatsVisibleCVar->Set(!ShowStatsVisibleCVar->GetBool());
 		}),
 		FCanExecuteAction::CreateLambda([]() { return true; }),
-		FIsActionChecked::CreateLambda([=]() { return ShowStatsCVar->GetInt() != 0; })
+		FIsActionChecked::CreateLambda([=]() { return ShowStatsVisibleCVar->GetBool(); })
 	);
 	CommandList.MapAction(
 		VisualizeNextLightCommand,
diff --git a/Engine/Source/Runtime/Renderer/Private/PostProcess/PostProcessing.cpp b/Engine/Source/Runtime/Renderer/Private/PostProcess/PostProcessing.cpp
--- a/Engine/Source/Runtime/Renderer/Private/PostProcess/PostProcessing.cpp
+++ b/Engine/Source/Runtime/Renderer/Private/PostProcess/PostProcessing.cpp
@@ -576,8 +576,8 @@ void AddPostProcessingPasses(
 #endif
 
 #if WITH_EDITOR || !UE_BUILD_SHIPPING
-	PassSequence.SetEnabled(EPass::VisualizeVirtualShadowMaps_PreEditorPrimitives, EngineShowFlags.VisualizeVirtualShadowMap && VirtualShadowMapArray != nullptr);
-	PassSequence.SetEnabled(EPass::VisualizeVirtualShadowMaps_PostEditorPrimitives, EngineShowFlags.VisualizeVirtualShadowMap && VirtualShadowMapArray != nullptr);
+	PassSequence.SetEnabled(EPass::VisualizeVirtualShadowMaps_PreEditorPrimitives, VirtualShadowMapArray != nullptr);
+	PassSequence.SetEnabled(EPass::VisualizeVirtualShadowMaps_PostEditorPrimitives, VirtualShadowMapArray != nullptr);
 #else
 	PassSequence.SetEnabled(EPass::VisualizeVirtualShadowMaps_PreEditorPrimitives, false);
 	PassSequence.SetEnabled(EPass::VisualizeVirtualShadowMaps_PostEditorPrimitives, false);
diff --git a/Engine/Source/Runtime/Renderer/Private/VirtualShadowMaps/VirtualShadowMapArray.cpp b/Engine/Source/Runtime/Renderer/Private/VirtualShadowMaps/VirtualShadowMapArray.cpp
--- a/Engine/Source/Runtime/Renderer/Private/VirtualShadowMaps/VirtualShadowMapArray.cpp
+++ b/Engine/Source/Runtime/Renderer/Private/VirtualShadowMaps/VirtualShadowMapArray.cpp
@@ -109,13 +109,128 @@ TAutoConsoleVariable<int32> CVarBuildDynamicHZB(
 	ECVF_RenderThreadSafe
 );
 
-TAutoConsoleVariable<int32> CVarVSMShowStats(
-	TEXT("r.Shadow.Virtual.ShowStats"),
-	0,
-	TEXT("Show VSM statistics."),
+TAutoConsoleVariable<bool> CVarShowStatsVisible(
+	TEXT("r.Shadow.Virtual.ShowStats.Visible"),
+	false,
+	TEXT("Render VSM Stats to screen"),
 	ECVF_RenderThreadSafe
 );
 
+static EVSMStatSection::Type GVSMStatSectionsEnabled = EVSMStatSection::All;
+EVSMStatSection::Type FVirtualShadowMap::GetEnabledStatSections()
+{
+	if (!CVarShowStatsVisible.GetValueOnAnyThread())
+	{
+		return EVSMStatSection::None;
+	}
+
+	return GVSMStatSectionsEnabled;
+}
+
+static FAutoConsoleCommand CCmdVSMShowStats = FAutoConsoleCommand(
+	TEXT("r.Shadow.Virtual.ShowStats"),
+	TEXT("Show/hide VSM statistics.\n")
+	TEXT("Usage: ShowStats <0|1|basic|all>\n")
+	TEXT("Usage: ShowStats <section1,section2,...>\n")
+	TEXT("Sections:\n")
+	#define PRINT_LABEL_LINE(x) TEXT("       " #x "\n")
+	ENUMERATE_VSM_STAT_SECTIONS(PRINT_LABEL_LINE),
+	#undef PRINT_LABEL_LINE
+	FConsoleCommandWithArgsAndOutputDeviceDelegate::CreateStatic([](const TArray<FString>& Args, FOutputDevice& Output)
+	{
+		if (Args.Num() < 1)
+		{
+			if (!CVarShowStatsVisible.GetValueOnAnyThread())
+			{
+				Output.Log(TEXT("0"));
+			}
+			else if (GVSMStatSectionsEnabled == EVSMStatSection::All)
+			{
+				Output.Log(TEXT("1"));
+			}
+			else
+			{
+				FString FlagString;
+				#define CONCAT_FLAG_LABEL(x) if ((GVSMStatSectionsEnabled & EVSMStatSection::x ) > 0) { FlagString += TEXT( #x ); FlagString += TEXT(","); }
+				ENUMERATE_VSM_STAT_SECTIONS(CONCAT_FLAG_LABEL)
+				#undef CONCAT_FLAG_LABEL
+				Output.Log(FlagString);
+			}
+			return;
+		}
+
+		FString ValueString = Args[0];
+
+		if (ValueString == TEXT("0"))
+		{
+			CVarShowStatsVisible->Set(false);
+		}
+		else if (ValueString == TEXT("1"))
+		{
+			CVarShowStatsVisible->Set(true);
+		}
+		else if (ValueString == TEXT("basic"))
+		{
+			CVarShowStatsVisible->Set(true);
+			GVSMStatSectionsEnabled = EVSMStatSection::Basic;
+			
+		}
+		else if (ValueString == TEXT("all"))
+		{
+			CVarShowStatsVisible->Set(true);
+			GVSMStatSectionsEnabled = EVSMStatSection::All;
+		}
+		else
+		{
+			CVarShowStatsVisible->Set(true);
+
+			// Iterate over flags in comma-separated string
+			auto Tokenize = [](const FStringView& String, TCHAR Separator, TFunctionRef<void(const FStringView&)> Func)
+			{
+				FStringView RemainingString { String };
+				for (int TokenEnd; RemainingString.FindChar(Separator, TokenEnd);)
+				{
+					FStringView Token = RemainingString.SubStr(0, TokenEnd);
+					Func(Token);
+					RemainingString.RightChopInline(TokenEnd + 1);
+				}
+				if (RemainingString.Len() > 0)
+				{
+					Func(RemainingString);
+				}
+			};
+
+			EVSMStatSection::Type SectionsEnabled = EVSMStatSection::Basic;
+
+			// Match string to section flag and turn it on
+			auto ParseFlag = [&](EVSMStatSection::Type Flag, FString FlagLabel, FStringView Token)
+			{
+				if (Token.Equals(FlagLabel, ESearchCase::IgnoreCase))
+				{
+					SectionsEnabled = EVSMStatSection::Type(SectionsEnabled | Flag);
+					return true;
+				}
+				return false;
+			};
+
+			Tokenize(ValueString, ',', [&](const FStringView& Token)
+			{
+				#define PARSE_FLAG(x) ParseFlag(EVSMStatSection::x, TEXT(#x), Token) ||
+				bool bFlagValid = ENUMERATE_VSM_STAT_SECTIONS(PARSE_FLAG) 0;
+				#undef PARSE_FLAG
+
+				if(!bFlagValid)
+				{
+					FString TokenStr(Token);
+					Output.Logf(TEXT("Invalid argument '%s'"), *TokenStr);
+				}
+			});
+
+			GVSMStatSectionsEnabled = SectionsEnabled;
+		}
+	})
+);
+
 static TAutoConsoleVariable<float> CVarPageDilationBorderSizeDirectional(
 	TEXT("r.Shadow.Virtual.PageDilationBorderSizeDirectional"),
 	0.05f,
@@ -1894,7 +2009,7 @@ bool FVirtualShadowMapArray::ShouldGenerateStats() const
 	constexpr bool bRunPageAreaDiagnostics = false;
 	constexpr bool bInsightsVSMChannelEnabled = false;
 #endif
-	return CVarVSMShowStats.GetValueOnRenderThread() || CacheManager->IsAccumulatingStats() || bRunPageAreaDiagnostics || IsCsvLogEnabled() || bInsightsVSMChannelEnabled;
+	return FVirtualShadowMap::GetEnabledStatSections() != EVSMStatSection::None || CacheManager->IsAccumulatingStats() || bRunPageAreaDiagnostics || IsCsvLogEnabled() || bInsightsVSMChannelEnabled;
 }
 
 bool FVirtualShadowMapArray::IsCsvLogEnabled() const
@@ -2968,7 +3083,7 @@ void FVirtualShadowMapArray::LogStats(FRDGBuilder& GraphBuilder, const FViewInfo
 		ShaderPrint::SetParameters(GraphBuilder, View.ShaderPrintData, PassParameters->ShaderPrintStruct);
 		ShaderPrint::RequestSpaceForTriangles(8192);
 
-		int ShowStatsValue = CVarVSMShowStats.GetValueOnRenderThread();
+		int ShowStatsValue = (int)FVirtualShadowMap::GetEnabledStatSections();
 		PassParameters->ShowStatsValue = ShowStatsValue;
 
 #if !UE_BUILD_SHIPPING
@@ -4482,13 +4597,23 @@ FScreenPassTexture FVirtualShadowMapArray::AddVisualizePass(FRDGBuilder& GraphBu
 
 #if !UE_BUILD_SHIPPING
 
-	if (!IsAllocated() || DebugVisualizationOutput.IsEmpty())
+	if (!IsAllocated())
 	{
 		return MoveTemp(FinalizeOutput(GraphBuilder, View, Output, OverrideOutput));
 	}
 
+	FScreenPassTextureViewport OutputViewport(Output);
+	FScreenPassRenderTarget OutputTarget(Output.Texture, OutputViewport.Rect, ERenderTargetLoadAction::ELoad);
+
+	if (CacheManager)
+	{
+		Output = CacheManager->AddVisualizePass(GraphBuilder, View, ViewIndex, Pass, SceneColor, OutputTarget);
+	}
+
+	bool VisualizeVSM = View.Family->EngineShowFlags.VisualizeVirtualShadowMap;
+
 	const FVirtualShadowMapVisualizationData& VisualizationData = GetVirtualShadowMapVisualizationData();
-	if (!VisualizationData.IsActive())
+	if (!VisualizeVSM || !VisualizationData.IsActive() || DebugVisualizationOutput.IsEmpty())
 	{
 		return MoveTemp(FinalizeOutput(GraphBuilder, View, Output, OverrideOutput));
 	}
@@ -4496,8 +4621,6 @@ FScreenPassTexture FVirtualShadowMapArray::AddVisualizePass(FRDGBuilder& GraphBu
 	RDG_EVENT_SCOPE(GraphBuilder, "VirtualShadowMapsVisualization");
 
 	FScreenPassTextureViewport InputViewport(DebugVisualizationOutput[ViewIndex]->Desc.Extent);
-	FScreenPassTextureViewport OutputViewport(Output);
-	FScreenPassRenderTarget OutputTarget(Output.Texture, OutputViewport.Rect, ERenderTargetLoadAction::ELoad);
 
 	int ActiveModeId = VisualizationData.GetActiveModeID();
 	int VisualizeVirtualShadowMapId = VisualizeLight[ViewIndex].GetVirtualShadowMapId();
diff --git a/Engine/Source/Runtime/Renderer/Private/VirtualShadowMaps/VirtualShadowMapArray.h b/Engine/Source/Runtime/Renderer/Private/VirtualShadowMaps/VirtualShadowMapArray.h
--- a/Engine/Source/Runtime/Renderer/Private/VirtualShadowMaps/VirtualShadowMapArray.h
+++ b/Engine/Source/Runtime/Renderer/Private/VirtualShadowMaps/VirtualShadowMapArray.h
@@ -84,6 +84,8 @@ public:
 	
 	static_assert(MaxMipLevels <= 8, ">8 mips requires more PageFlags bits. See VSM_PAGE_FLAGS_BITS_PER_HMIP in PageAccessCommon.ush");
 
+	static EVSMStatSection::Type GetEnabledStatSections();
+
 	// TODO: Currently only used for these constants... probably rename the cache structure to virtual shadow map now instead
 private:
 	FVirtualShadowMap() {}
diff --git a/Engine/Source/Runtime/Renderer/Private/VirtualShadowMaps/VirtualShadowMapCacheManager.cpp b/Engine/Source/Runtime/Renderer/Private/VirtualShadowMaps/VirtualShadowMapCacheManager.cpp
--- a/Engine/Source/Runtime/Renderer/Private/VirtualShadowMaps/VirtualShadowMapCacheManager.cpp
+++ b/Engine/Source/Runtime/Renderer/Private/VirtualShadowMaps/VirtualShadowMapCacheManager.cpp
@@ -17,6 +17,12 @@
 #include "Shadows/ShadowScene.h"
 #include "ProfilingDebugging/CountersTrace.h"
 
+#if !UE_BUILD_SHIPPING
+#include "Engine/Font.h"
+#include "Engine/Canvas.h"
+#include "CanvasItem.h"
+#endif
+
 #define LOCTEXT_NAMESPACE "VirtualShadowMapCacheManager"
 CSV_DECLARE_CATEGORY_EXTERN(VSM);
 
@@ -157,8 +163,6 @@ FAutoConsoleVariableRef CVarVSMAllowScreenOverflowMessages(
 	ECVF_RenderThreadSafe
 );
 
-extern TAutoConsoleVariable<int32> CVarVSMShowStats;
-
 static bool bDumpLightNaniteStats = false;
 static FAutoConsoleCommandWithArgsAndOutputDevice DumpLightNaniteStatsConsoleCommand(
 	TEXT("r.Shadow.Virtual.DumpLightNaniteStats"),
@@ -885,13 +889,13 @@ void FVirtualShadowMapArrayCacheManager::InitExtension(FScene& InScene)
 		{
 			uint32 Overlap = PageAreaDiags[Index];
 			uint32 PersistentPrimitiveId = PageAreaDiags[Index + 1];
-			int32 PrimtiveIndex = Scene.GetPrimitiveIndex(FPersistentPrimitiveIndex{ int32(PersistentPrimitiveId) });
-			if (Overlap > 0 && PrimtiveIndex != INDEX_NONE)
+			int32 PrimitiveIndex = Scene.GetPrimitiveIndex(FPersistentPrimitiveIndex{ int32(PersistentPrimitiveId) });
+			if (Overlap > 0 && PrimitiveIndex != INDEX_NONE)
 			{
 				if (RealTimeSeconds - LastLoggedPageOverlapAppTime[PersistentPrimitiveId] > 5.0f)
 				{
 					LastLoggedPageOverlapAppTime[PersistentPrimitiveId] = RealTimeSeconds;
-					UE_LOG(LogRenderer, Warning, TEXT("Non-Nanite VSM page overlap performance Warning, %d, %s, %s"), Overlap, *Scene.Primitives[PrimtiveIndex]->GetOwnerActorNameOrLabelForDebuggingOnly(), *Scene.Primitives[PrimtiveIndex]->GetFullnameForDebuggingOnly());
+					UE_LOG(LogRenderer, Warning, TEXT("Non-Nanite VSM page overlap performance Warning, %d, %s, %s"), Overlap, *Scene.Primitives[PrimitiveIndex]->GetOwnerActorNameOrLabelForDebuggingOnly(), *Scene.Primitives[PrimitiveIndex]->GetFullnameForDebuggingOnly());
 				}
 				LargePageAreaItems.Add(PersistentPrimitiveId, FLargePageAreaItem{ Overlap, RealTimeSeconds });
 			}
@@ -899,18 +903,11 @@ void FVirtualShadowMapArrayCacheManager::InitExtension(FScene& InScene)
 
 		auto ParseNPFDiagnostic = [&](uint32 VSMId, uint32 LightId, uint32 ClusterCount)
 		{
+			bool bPrintEntry = false;
+
 			if (ClusterCount > 0)
 			{
-				FString LightNameWithLevel;
-				if (LightId <= uint64(Scene.Lights.GetMaxIndex()))
-				{
-					const FLightSceneInfo* LightSceneInfo = Scene.Lights[LightId].LightSceneInfo;
-					if (LightSceneInfo && LightSceneInfo->Proxy)
-					{
-						FSceneRenderer::GetLightNameForDrawEvent(LightSceneInfo->Proxy, LightNameWithLevel);
-					}
-				}
-				NPFDiagnosticMessages.Add(FString::Printf(TEXT("%8d | %48s | %12d"), VSMId, *LightNameWithLevel, ClusterCount));
+				bPrintEntry = true;
 				NPFDiagnosticTimer.FindOrAdd(VSMId) = 180;
 			}
 			else
@@ -918,13 +915,27 @@ void FVirtualShadowMapArrayCacheManager::InitExtension(FScene& InScene)
 				uint8* FramesRemaining = NPFDiagnosticTimer.Find(VSMId);
 				if (FramesRemaining != nullptr)
 				{
-					NPFDiagnosticMessages.Add(FString(TEXT("")));
+					bPrintEntry = true;
 					if ((*FramesRemaining)-- == 0)
 					{
 						NPFDiagnosticTimer.Remove(VSMId);
 					}
 				}
 			}
+
+			if (bPrintEntry)
+			{
+				FString LightNameWithLevel;
+				if (Scene.Lights.IsValidIndex(LightId))
+				{
+					const FLightSceneInfo* LightSceneInfo = Scene.Lights[LightId].LightSceneInfo;
+					if (LightSceneInfo && LightSceneInfo->Proxy)
+					{
+						FSceneRenderer::GetLightNameForDrawEvent(LightSceneInfo->Proxy, LightNameWithLevel);
+					}
+				}
+				NPFDiagnosticMessages.Add(FString::Printf(TEXT("%8d | %48s | %12d"), VSMId, *LightNameWithLevel, ClusterCount));
+			}
 		};
 
 		NPFDiagnosticMessages.SetNum(0, EAllowShrinking::No);
@@ -997,15 +1008,6 @@ void FVirtualShadowMapArrayCacheManager::InitExtension(FScene& InScene)
 		{
 			OutMessages.Add(FCoreDelegates::EOnScreenMessageSeverity::Warning, FText::FromString(FString::Printf(TEXT("Virtual Shadow Map Stats Accumulation Active. Set r.Shadow.Virtual.AccumulateStats to 0 to stop."))));
 		}
-
-		if (CVarVSMShowStats.GetValueOnRenderThread() > 1)
-		{
-			for (const FString& Message : NPFDiagnosticMessages)
-			{
-				OutMessages.Add(FCoreDelegates::EOnScreenMessageSeverity::Info, FText::FromString(Message));
-			}
-			NPFDiagnosticMessages.SetNum(0, EAllowShrinking::No);
-		}
 	});
 #endif
 }
@@ -1017,6 +1019,79 @@ FVirtualShadowMapArrayCacheManager::~FVirtualShadowMapArrayCacheManager()
 #endif
 }
 
+FScreenPassTexture FVirtualShadowMapArrayCacheManager::AddVisualizePass(FRDGBuilder& GraphBuilder, const FViewInfo& View, int32 ViewIndex, EVSMVisualizationPostPass Pass, FScreenPassTexture& SceneColor, FScreenPassRenderTarget& Output)
+{
+#if !UE_BUILD_SHIPPING
+	if (Pass == EVSMVisualizationPostPass::PostEditorPrimitives && (FVirtualShadowMap::GetEnabledStatSections() & EVSMStatSection::Clusters))
+	{
+		if (!MonospaceFont)
+		{
+			FTopLevelAssetPath MonospaceFontPath = FTopLevelAssetPath("/Engine/EngineFonts/DroidSansMono.DroidSansMono");
+			MonospaceFont = LoadObject<UFont>(nullptr, *MonospaceFontPath.ToString());
+		}
+		
+		AddDrawCanvasPass(GraphBuilder, RDG_EVENT_NAME("NPFDiagnostics"), View, Output,
+			[&](FCanvas& Canvas)
+			{
+				const float DPIScale = Canvas.GetDPIScale();
+				Canvas.SetBaseTransform(FMatrix(FScaleMatrix(DPIScale) * Canvas.CalcBaseTransform2D(Canvas.GetViewRect().Width(), Canvas.GetViewRect().Height())));
+
+				auto DrawColorTile = [&](float X, float Y, float Width, float Height, const FLinearColor& Color)
+				{
+					Canvas.DrawTile(X / DPIScale, Y / DPIScale, Width / DPIScale, Height / DPIScale, 0, 0, 0, 0, Color);
+				};
+
+				auto MeasureString = [&](FStringView Text, const UFont* Font = GetStatsFont())
+				{
+					FCanvasTextStringViewItem TextItem(FVector2D::ZeroVector, Text, Font, FLinearColor::Transparent);
+					
+					TextItem.bOutlined = true;
+					TextItem.OutlineColor = FLinearColor::Transparent;
+					Canvas.DrawItem(TextItem);
+					return TextItem.DrawnSize;
+				};
+
+				auto DrawString = [&](float X, float Y, FStringView Text, const FLinearColor& Color = FLinearColor::White, const UFont* Font = GetStatsFont())
+				{
+					FCanvasTextStringViewItem TextItem(FVector2D(X / DPIScale, Y / DPIScale), Text, Font, Color);
+					
+					TextItem.bOutlined = true;
+					TextItem.OutlineColor = FLinearColor::Black;
+					Canvas.DrawItem(TextItem);
+					return TextItem.DrawnSize;
+				};
+
+				FVector2D MessageBoxSize;
+				{
+					FIntPoint MessageBoxCharacters {0, NPFDiagnosticMessages.Num()};
+					for (const FString& Message : NPFDiagnosticMessages)
+					{
+						MessageBoxCharacters.X = FMath::Max(MessageBoxCharacters.X, Message.Len());
+					}
+
+					FVector2D CharSize = MeasureString(TEXT("A"), MonospaceFont);
+					MessageBoxSize = CharSize * MessageBoxCharacters + FVector2D(10, 10);
+				}
+				
+				int BaselineX = 10;
+				int BaselineY = Canvas.GetViewRect().Height() - (Canvas.GetViewRect().Height() / 3);
+				BaselineY = FMath::Min(BaselineY, Canvas.GetViewRect().Height() - (MessageBoxSize.Y + 10));
+				FLinearColor MessageBoxColor(0.1f, 0.1f, 0.1f, 0.4f);
+
+				DrawColorTile(BaselineX, BaselineY, MessageBoxSize.X, MessageBoxSize.Y, MessageBoxColor);
+				int Y = 0;
+				for (const FString& Message : NPFDiagnosticMessages)
+				{
+					Y += DrawString(BaselineX, BaselineY + Y, Message, FLinearColor::White, MonospaceFont).Y;
+				}
+			});
+		NPFDiagnosticMessages.SetNum(0, EAllowShrinking::No);
+	}
+#endif
+
+	return Output;
+}
+
 #if !UE_BUILD_SHIPPING
 FText FVirtualShadowMapArrayCacheManager::GetOverflowMessage(uint32 OverflowTypeIndex) const
 {
diff --git a/Engine/Source/Runtime/Renderer/Private/VirtualShadowMaps/VirtualShadowMapCacheManager.h b/Engine/Source/Runtime/Renderer/Private/VirtualShadowMaps/VirtualShadowMapCacheManager.h
--- a/Engine/Source/Runtime/Renderer/Private/VirtualShadowMaps/VirtualShadowMapCacheManager.h
+++ b/Engine/Source/Runtime/Renderer/Private/VirtualShadowMaps/VirtualShadowMapCacheManager.h
@@ -344,6 +344,7 @@ public:
 	}
 
 	bool IsAccumulatingStats();
+	FScreenPassTexture AddVisualizePass(FRDGBuilder& GraphBuilder, const FViewInfo& View, int32 ViewIndex, EVSMVisualizationPostPass Pass, FScreenPassTexture& SceneColor, FScreenPassRenderTarget& Output);
 
 	using FInstanceGPULoadBalancer = TInstanceCullingLoadBalancer<SceneRenderingAllocator>;
 
@@ -540,6 +541,8 @@ private:
 
 	TArray<FString> NPFDiagnosticMessages {};
 	TMap<uint32, uint8> NPFDiagnosticTimer {};
+
+	TObjectPtr<UFont> MonospaceFont = nullptr;
 #endif // UE_BUILD_SHIPPING
 
 	FShadowInvalidatingInstancesImplementation ShadowInvalidatingInstancesImplementation;
