Fix error compiling FORWARD_SHADING permutations of BasePassPixelShader.usf.

diff --git a/Engine/Shaders/Private/BasePassPixelShader.usf b/Engine/Shaders/Private/BasePassPixelShader.usf
--- a/Engine/Shaders/Private/BasePassPixelShader.usf
+++ b/Engine/Shaders/Private/BasePassPixelShader.usf
@@ -240,6 +240,8 @@ bool ManualDepthTestEqual(float4 SvPosition, float PSDeviceZWithOffset)
 	return DepthTest_Equal;
 }
 
+#endif // !FORWARD_SHADING
+
 /** Calculates lighting for translucency. */
 void GetTranslucencyVolumeLighting(
 	FMaterialPixelParameters MaterialParameters, 
@@ -393,8 +395,6 @@ void GetTranslucencyVolumeLighting(
 	#endif
 }
 
-#endif
-
 /** Computes sky diffuse lighting, including precomputed shadowing. */
 void GetSkyLighting(FMaterialPixelParameters MaterialParameters, VTPageTableResult LightmapVTPageTableResult, bool bEvaluateBackface, float3 WorldNormal, LightmapUVType LightmapUV, uint LightmapDataIndex, float3 SkyOcclusionUV3D, out float3 OutDiffuseLighting, out float3 OutSubsurfaceLighting)
 {
