Substrate - fixed shader error that was failing FXC compilation. (non empty case should return or break)

#rnx
#tests pc editor, console cook

diff --git a/Engine/Shaders/Private/Substrate/Substrate.ush b/Engine/Shaders/Private/Substrate/Substrate.ush
--- a/Engine/Shaders/Private/Substrate/Substrate.ush
+++ b/Engine/Shaders/Private/Substrate/Substrate.ush
@@ -3672,6 +3672,7 @@ uint SubstrateGetBSDFSubSurfaceProfileID(in FSubstrateBSDF BSDF)
 		{
 			return SubstrateSubsurfaceProfileIdTo8bits(SLAB_SSSPROFILEID(BSDF));
 		}
+		break;
 	}
 	}
 	return 0;
