[SingleLayerWater] Enabled SSR on mobile SM5 water

diff --git a/Engine/Shaders/Private/MobileLightingCommon.ush b/Engine/Shaders/Private/MobileLightingCommon.ush
--- a/Engine/Shaders/Private/MobileLightingCommon.ush
+++ b/Engine/Shaders/Private/MobileLightingCommon.ush
@@ -16,7 +16,7 @@
 #if MOBILE_SSR_ENABLED && !(MOBILE_QL_FORCE_FULLY_ROUGH || MATERIAL_FULLY_ROUGH)
 	#if MOBILE_DEFERRED_SHADING
 		#if IS_MOBILE_BASE_PASS
-			#define MOBILE_USE_SSR (MATERIAL_SSR)
+			#define MOBILE_USE_SSR (MATERIAL_SSR || (MOBILE_SUPPORTS_SM5_MATERIAL_NODES && MATERIAL_SHADINGMODEL_SINGLELAYERWATER))
 		#else
 			#define MOBILE_USE_SSR 1
 		#endif
