Substrate - Fixed advanced visualisation shadser compilation.

#tests PC editor dx12
#rnx

diff --git a/Engine/Shaders/Private/Substrate/SubstrateVisualize.usf b/Engine/Shaders/Private/Substrate/SubstrateVisualize.usf
--- a/Engine/Shaders/Private/Substrate/SubstrateVisualize.usf
+++ b/Engine/Shaders/Private/Substrate/SubstrateVisualize.usf
@@ -707,11 +707,13 @@ void EvaluateLighting(
 		const bool bEnableSpecular = ReflectionStruct.SkyLightParameters.y > 0.0f;
 		FSubstrateEnvLightResult SubstrateEnvLight = SubstrateEvaluateForEnvLight(SubstrateBSDFContext, bEnableSpecular, Settings);
 
+		FSubstrateIrradianceAndOcclusion SubstrateIrradianceAndOcclusion = SubstrateGetIrradianceAndAO(SubstratePixelHeader);
+
 		float3 DiffuseLighting = 0.0;
 		float3 SpecularLighting = 0.0;
 		SubstrateEnvLightingCommon(
 			SubstrateEnvLight,
-			SubstratePixelHeader,
+			SubstrateIrradianceAndOcclusion,
 			SubstrateBSDFContext,
 			BSDF,
 			N,	//in float3 BentNormal,
