Path Tracer: Improve duplicate AHS detection by tracking one additional hit

Very thin meshes can end up having more than one triangle replicated across BLAS leaves, leading to duplicate AHS invocations being undetected by the current most-recent-hit strategy.

This patch tracks an additional hit value to help mitigate this issue in more cases.

diff --git a/Engine/Shaders/Private/PathTracing/PathTracingCommon.ush b/Engine/Shaders/Private/PathTracing/PathTracingCommon.ush
--- a/Engine/Shaders/Private/PathTracing/PathTracingCommon.ush
+++ b/Engine/Shaders/Private/PathTracing/PathTracingCommon.ush
@@ -372,7 +372,10 @@ struct FPackedPathTracingPayload : FMinimalPayload
 	void PackRoughnessAniso(float R, float A)   { PackedData[9] = PackFloat2ToUInt(float2(R, A)); }
 
 #endif
-	
+
+#if IS_PAYLOAD_ENABLED(RT_PAYLOAD_TYPE_GPULIGHTMASS)
+	// NOTE: GPU Lightmass doesnt't have a shading model that needs volumetric absorption along shadow rays
+#else
 	// tau: optical depth along shadow ray, this is the integral of the extinction coefficient
 	float3 GetTau()
 	{
@@ -389,6 +392,28 @@ struct FPackedPathTracingPayload : FMinimalPayload
 		PackedData[6] = asuint(Tau.y);
 		PackedData[7] = asuint(Tau.z);
 	}
+#endif
+
+	// Check if we have already run the same hit through AHS
+	bool CheckDuplicateAHS(float T)
+	{
+#if IS_PAYLOAD_ENABLED(RT_PAYLOAD_TYPE_GPULIGHTMASS)
+		return T == HitT || T == asfloat(PackedData[5]);
+#else
+		// either equal to the last hit or second-to-last hit
+		return T == HitT || T == asfloat(PackedData[8]);
+#endif
+	}
+
+	void UpdateDuplicateAHS(float T)
+	{
+#if IS_PAYLOAD_ENABLED(RT_PAYLOAD_TYPE_GPULIGHTMASS)
+		PackedData[5] = asuint(HitT); // remember second to last
+#else
+		PackedData[8] = asuint(HitT); // remember second to last
+#endif
+		HitT = T; // update most recent hit
+	}
 
 	float4 GetDBufferA()
 	{
diff --git a/Engine/Shaders/Private/PathTracing/PathTracingMaterialHitShader.usf b/Engine/Shaders/Private/PathTracing/PathTracingMaterialHitShader.usf
--- a/Engine/Shaders/Private/PathTracing/PathTracingMaterialHitShader.usf
+++ b/Engine/Shaders/Private/PathTracing/PathTracingMaterialHitShader.usf
@@ -1301,7 +1301,7 @@ RAY_TRACING_ENTRY_ANY_HIT(PathTracingMaterialAHS,
 		// not a shadow ray -- don't need to run the AHS logic
 		return;
 	}
-	if (PackedPayload.HitT == RayTCurrent())
+	if (PackedPayload.CheckDuplicateAHS(RayTCurrent()))
 	{
 		// We just processed this exact hit, don't double-count it
 		IgnoreHit();
@@ -1351,7 +1351,7 @@ RAY_TRACING_ENTRY_ANY_HIT(PathTracingMaterialAHS,
 
 	// the following blend modes need to process shadows after having executed the material
 #if MATERIALBLENDING_MODULATE || (MATERIALBLENDING_MASKED && MATERIAL_DITHER_OPACITY_MASK) || MATERIALBLENDING_ALPHACOMPOSITE || MATERIALBLENDING_TRANSLUCENT
-	PackedPayload.HitT = RayTCurrent(); // remember the distance at which we ran the shader, in case we get a double-hit
+	PackedPayload.UpdateDuplicateAHS(RayTCurrent());
 
 	{ // TODO: Cleanup this indentation level
 
