Path Tracer: Implement transparency support in visualization modes

Refactor the logic for tracing rays to allow for stochastic transparency. Like with the reference path tracer, the denoising signal is accumulated across all hits.

Add a transparency count visualization mode to show how many additional hits are processed as a blue-red heatmap.

Tweak the specular color denoiser signal to better model the actual albedo (taking into account viewing angle and roughness).

Add cvar to control if translucent hits should be processed or not (for benchmarking purposes).

diff --git a/Engine/Shaders/Private/PathTracing/PathTracingDebug.usf b/Engine/Shaders/Private/PathTracing/PathTracingDebug.usf
--- a/Engine/Shaders/Private/PathTracing/PathTracingDebug.usf
+++ b/Engine/Shaders/Private/PathTracing/PathTracingDebug.usf
@@ -7,13 +7,15 @@
 
 #define USE_BSDF_SAMPLING 1 // 0: light sampling only (no MIS) 1: light sampling and BSDF sampling (with MIS)
 
-#define INCLUDE_TRANSLUCENT   0 // include objects tagged as TRANSLUCENT in the TLAS (transparency is not implemented yet, so better to leave this off for now)
-
 #define USE_PATH_TRACING_LIGHT_GRID	1
 #define USE_RAY_TRACING_DECAL_GRID	1
 
 #ifndef PATH_TRACER_PRIMARY_RAYS // Should be set by permutation
-#define PATH_TRACER_PRIMARY_RAYS 1
+#define PATH_TRACER_PRIMARY_RAYS 0
+#endif
+
+#ifndef PATH_TRACER_INCLUDE_TRANSLUCENT // Should be set by permutation
+#define PATH_TRACER_INCLUDE_TRANSLUCENT   0 // include objects tagged as TRANSLUCENT when tracing rays?
 #endif
 
 
@@ -206,49 +208,27 @@ int GetDominantSide(float3 N)
 
 #if PATH_TRACER_PRIMARY_RAYS
 
-FPathTracingResult PathTracingDebugCore(FRayDesc Ray, uint2 LaunchIndex, int SampleIndex)
+FPackedPathTracingPayload TraceTransparentRay(FRayDesc Ray, int ScatterType, int Bounce, float PathRoughness, inout float3 PathThroughput, inout FRandomSequence RandSequence, inout FPathTracingResult Result, inout float3 LocalRadiance)
 {
-	FPathTracingResult Result = (FPathTracingResult) 0;
-	Result.Radiance.a = 1;
-	float3 PathThroughput = 1;
-	float PathRoughness = 0;
-	uint ScatterType = PATHTRACER_SCATTER_CAMERA;
-	FRandomSequence RandSequence = RandomSequenceCreate(uint3(LaunchIndex, View.StateFrameIndex), SampleIndex, SamplesPerPixel);
+	const uint RayFlags = Bounce == 0 ? RAY_FLAG_CULL_BACK_FACING_TRIANGLES : 0;
+	const uint MissShaderIndex = 0;
+	uint InstanceInclusionMask = (ScatterType == PATHTRACER_SCATTER_CAMERA || ScatterType == PATHTRACER_SCATTER_REFRACT) ? PATHTRACER_MASK_CAMERA : PATHTRACER_MASK_INDIRECT;
 
-	// State about the markov chain selection that needs to be kept from one bounce to the next
-	float ScoresSum = 0;
-	float3 PreviousN = 0;
-	FMarkovChainState ChosenState = (FMarkovChainState)0;
-	// MCPG parameters
-#define N_MC  5 // From paper, seems to be the smallest value which produces good results, some subtle improvements from higher values at the cost of performance
-	const float AdaptiveGridProb = 0.7; // Mentioned in text - TODO: could rely on adaptive grid only with a tweak to level distribution?
-	const float StaticGridWidth = 25.6; // matches level 24 of adaptive grid
-	const float LevelScale = 3.0; 		// 3 levels per octave
-	const float MinWidth = 0.1; 		// 1mm - smallest possible voxel size
-	const float TanFactor = 0.0087;		// ~2*tan(0.25 degrees)
-	const int 	NMax = 1024;			// Blend factors [See Dittebrandt et al.]
-	const float AlphaMin = 1.0 / 100.0; // Allow blending over 100 frames, higher than original publication because the chains get reset during BSDF sampling
-	const float MaterialProb = 0.15f; 	// Chance of BSDF sampling (TODO: adapt this as a function of roughness)
+	const float2 RandSample = RandSequence.Get2D();
 
-	LOOP
-	for (int Bounce = 0; Bounce <= MaxBounces; Bounce++)
-	{
-		// Trace Ray
-		const uint RayFlags = Bounce == 0 ? RAY_FLAG_CULL_BACK_FACING_TRIANGLES : 0;
-#if INCLUDE_TRANSLUCENT
-		const uint InstanceInclusionMask = Bounce == 0 ? PATHTRACER_MASK_CAMERA | PATHTRACER_MASK_CAMERA_TRANSLUCENT : PATHTRACER_MASK_INDIRECT | PATHTRACER_MASK_INDIRECT_TRANSLUCENT;
-#else
-		const uint InstanceInclusionMask = Bounce == 0 ? PATHTRACER_MASK_CAMERA : PATHTRACER_MASK_INDIRECT;
-#endif
-		const uint MissShaderIndex = 0;
+	FRISContext HitSample = InitRISContext(RandSample.x);
 
-		FPackedPathTracingPayload PackedPayload = InitPathTracingPayload(ScatterType, PathRoughness);
-		PackedPayload.SetStochasticSlabRand(RandSequence.Get1D());
-#if !INCLUDE_TRANSLUCENT
-		PackedPayload.SetFlag(PATH_TRACING_PAYLOAD_INPUT_FLAG_IGNORE_TRANSLUCENT);
-#endif
+	float3 ResultThroughput = 0;
+	FPackedPathTracingPayload ResultPayload = (FPackedPathTracingPayload) 0;
+	FPackedPathTracingPayload OpaquePackedPayload = InitPathTracingPayload(ScatterType, PathRoughness);
+	OpaquePackedPayload.SetDBufferA(float4(0, 0, 0, 1));
+	OpaquePackedPayload.SetDBufferB(float4(0, 0, 0, 1));
+	OpaquePackedPayload.SetDBufferC(float4(0, 0, 0, 1));
+	OpaquePackedPayload.SetStochasticSlabRand(RandSample.y);
+	OpaquePackedPayload.SetFlag(PATH_TRACING_PAYLOAD_INPUT_FLAG_IGNORE_TRANSLUCENT);
 
 #if PATH_TRACER_USE_SER 
+	{
 		NvHitObject Hit;
 		NvTraceRayHitObject(
 			TLAS,
@@ -258,14 +238,72 @@ FPathTracingResult PathTracingDebugCore(FRayDesc Ray, uint2 LaunchIndex, int Sam
 			RAY_TRACING_NUM_SHADER_SLOTS,
 			MissShaderIndex,
 			Ray.GetNativeDesc(),
-			PackedPayload,
+			OpaquePackedPayload,
 			Hit);
 		if (Bounce > 0)
 		{
 			NvReorderThread(Hit);
 		}
-		NvInvokeHitObject(TLAS, Hit, PackedPayload);
+		NvInvokeHitObject(TLAS, Hit, OpaquePackedPayload);
+	}
+#else
+	TraceRay(
+		TLAS,
+		RayFlags,
+		InstanceInclusionMask,
+		RAY_TRACING_SHADER_SLOT_MATERIAL,
+		RAY_TRACING_NUM_SHADER_SLOTS,
+		MissShaderIndex,
+		Ray.GetNativeDesc(),
+		OpaquePackedPayload);
+#endif
+	if (OpaquePackedPayload.IsHit())
+	{
+		// TODO: If hit, also process decals, etc ...
+
+		// Now that we know where the opaque hit is, trace the remaining portion of the ray prior to this
+		Ray.TMax = asfloat(asuint(OpaquePackedPayload.HitT) - 1);
+	}
+
+	InstanceInclusionMask = (ScatterType == PATHTRACER_SCATTER_CAMERA || ScatterType == PATHTRACER_SCATTER_REFRACT) ? PATHTRACER_MASK_CAMERA_TRANSLUCENT : PATHTRACER_MASK_INDIRECT_TRANSLUCENT;
+
+	for (;;)
+	{
+#if PATH_TRACER_INCLUDE_TRANSLUCENT
+		FPackedPathTracingPayload PackedPayload = InitPathTracingPayload(ScatterType, PathRoughness);
+		PackedPayload.SetDBufferA(float4(0, 0, 0, 1));
+		PackedPayload.SetDBufferB(float4(0, 0, 0, 1));
+		PackedPayload.SetDBufferC(float4(0, 0, 0, 1));
+		PackedPayload.SetStochasticSlabRand(RandSample.y);
+		if (OpaquePackedPayload.IsHit())
+		{
+			const float3 ViewZ = View.ViewToTranslatedWorld[2].xyz;
+			const float3 TranslatedWorldPos = Ray.Origin + OpaquePackedPayload.HitT * Ray.Direction;
+			PackedPayload.SetFlag(PATH_TRACING_PAYLOAD_INPUT_FLAG_HAS_SCENE_DEPTH);
+			PackedPayload.SetSceneDepth(dot(TranslatedWorldPos, ViewZ));
+		}
+
+#if PATH_TRACER_USE_SER 
+		{
+			NvHitObject Hit;
+			NvTraceRayHitObject(
+				TLAS,
+				RayFlags,
+				InstanceInclusionMask,
+				RAY_TRACING_SHADER_SLOT_MATERIAL,
+				RAY_TRACING_NUM_SHADER_SLOTS,
+				MissShaderIndex,
+				Ray.GetNativeDesc(),
+				PackedPayload,
+				Hit);
+			if (Bounce > 0)
+			{
+				NvReorderThread(Hit);
+			}
+			NvInvokeHitObject(TLAS, Hit, PackedPayload);
+		}
 #else
+		// Trace the ray
 		TraceRay(
 			TLAS,
 			RayFlags,
@@ -276,74 +314,173 @@ FPathTracingResult PathTracingDebugCore(FRayDesc Ray, uint2 LaunchIndex, int Sam
 			Ray.GetNativeDesc(),
 			PackedPayload);
 #endif
+		if (PackedPayload.IsMiss())
+		{
+			// If we miss the translucent geometry, the next hit is the opaque one
+			PackedPayload = OpaquePackedPayload;
+		}
+#else // PATH_TRACER_INCLUDE_TRANSLUCENT
+
+#define PackedPayload OpaquePackedPayload
 
+#endif // PATH_TRACER_INCLUDE_TRANSLUCENT
 		// If miss - exit
 		if (PackedPayload.IsMiss())
 		{
-			break; // exit bounce loop
+			// still no hit -- the ray must have escaped to the background, we are done
+			// TODO: could add background Alpha here
+			break;
 		}
 
-		const float3 TranslatedWorldPos = Ray.Origin + PackedPayload.HitT * Ray.Direction;
-		const float3 WorldNormal = PackedPayload.UnpackWorldNormal();
-		const uint PrimitiveLightingChannelMask = PackedPayload.GetPrimitiveLightingChannelMask();
-
-		const bool bCameraRay = Bounce == 0;
-		const bool bLastBounce = Bounce == MaxBounces;
+		if (Bounce == 0 && !PackedPayload.IsHoldout())
+		{
+			Result.Radiance.a += Luminance(PathThroughput * (1.0 - PackedPayload.UnpackTransparencyColor()));
+		}
 
-		FLightLoopCount LightLoopCount = (FLightLoopCount) 0; 
-		if (!bLastBounce)
+		if (Bounce == 0 || EnableEmissive)
 		{
-			LightLoopCount = LightGridLookup(TranslatedWorldPos);
+			LocalRadiance += PathThroughput * PackedPayload.UnpackRadiance();
 		}
 
-		// Capture depth on camera ray
-		if (bCameraRay)
+		if (Bounce == 0 && PackedPayload.ShouldOutputDepth() && Result.Z == 0.0)
 		{
 			const float3 ViewZ = View.ViewToTranslatedWorld[2].xyz;
-			const float LinearZ = dot(TranslatedWorldPos, ViewZ);
-			Result.Z = ConvertToDeviceZ(LinearZ);
-			Result.LinearDepth = LinearZ;
-			Result.Radiance.a = 0; // opaque alpha (we aren't handling translucency right now)
+			const float3 TranslatedWorldPos = Ray.Origin + PackedPayload.HitT * Ray.Direction;
+			Result.LinearDepth = dot(TranslatedWorldPos, ViewZ);
+			Result.Z = ConvertToDeviceZ(Result.LinearDepth);
+			Result.WorldNormal = PackedPayload.UnpackWorldNormal();
 		}
 
-		// Add radiance (emissive surfaces)
-		float3 LocalRadiance = 0; 
-		if (bCameraRay || EnableEmissive)
+		if (Bounce < MaxBounces)
 		{
-			LocalRadiance += PathThroughput * PackedPayload.UnpackRadiance();
+			FPathTracingPayload HitPayload = UnpackPathTracingPayload(PackedPayload, Ray);
+			AdjustPayloadAfterUnpack(HitPayload, true);
+			float3 Contrib = PathThroughput * EstimateMaterialAlbedo(HitPayload);
+
+			// only allow albedo tracking through a very narrow range of roughness values
+			const float DenoiserMinRoughness = 0.01;
+			const float DenoiserMaxRoughness = 0.02;
+			float DenoiserContrib = 1 - saturate((PathRoughness - DenoiserMinRoughness) / (DenoiserMaxRoughness - DenoiserMinRoughness));
+			if (DenoiserContrib > 0)
+			{
+				float3 DenoiserThroughput = PathThroughput * DenoiserContrib * HitPayload.BSDFOpacity;
+
+				// TEMP: If the material has a valid SSS lobe, just merge it into diffuse for now
+				FSSSRandomWalkInfo SSS = GetMaterialSSSInfo(HitPayload, Ray.Direction);
+				if (any(SSS.Weight > 0))
+				{
+					HitPayload.DiffuseColor += SSS.Color;
+				}
+
+				// Accumulate G-Buffer like data for the denoiser
+#if PATHTRACING_SUBSTRATE_PAYLOAD
+				float3 RoughnessData = HitPayload.RoughnessData;
+				float Roughness = lerp(RoughnessData.x, RoughnessData.y, RoughnessData.z); // mix both lobes to get an average value
+				DenoiserThroughput *= HitPayload.WeightV * HitPayload.BSDFOpacity * lerp(1.0, HitPayload.TransmittanceN, HitPayload.CoverageAboveAlongN);
+#else
+				float Roughness = HitPayload.Roughness;
+#endif
+				float NormThroughput = Luminance(DenoiserThroughput);
+				float3 DiffuseColor = HitPayload.DiffuseColor;
+				float3 SpecularColor = HitPayload.SpecularColor;
+
+				float MinRoughness = max(DenoiserMinRoughness, PathRoughness);
+				float DenoiserRoughnessWeight = saturate((Roughness - MinRoughness) / (DenoiserMaxRoughness - MinRoughness));
+#if 1
+				// Remapping of NoV seems to be required by DLSS
+				// TODO: Could use the energy conservation routines and tweak diffuse as well
+				float NoV = saturate(dot(-Ray.Direction, HitPayload.WorldNormal));
+				float x = NoV * NoV;
+				NoV = 0.75 * x / (x + (1 - NoV) * (1 - NoV));
+				float3 F = EnvBRDF(HitPayload.SpecularColor, Roughness, NoV);
+#else
+				float3 F = HitPayload.SpecularColor;
+#endif
+				Result.DiffuseColor += DenoiserThroughput * HitPayload.DiffuseColor;
+				Result.SpecularColor += DenoiserThroughput * F * DenoiserRoughnessWeight;
+				Result.Roughness += NormThroughput * Roughness;
+				//Result.WorldNormal += NormThroughput * HitPayload.WorldNormal;
+			}
+
+			float SelectionWeight = max3(Contrib.x, Contrib.y, Contrib.z);
+			if (HitSample.Accept(SelectionWeight))
+			{
+				// stash this hit as a candidate return value
+				ResultPayload = PackedPayload;
+				ResultThroughput = PathThroughput / SelectionWeight;
+			}
+		}
+
+		// account for local transparency change
+		PathThroughput *= PackedPayload.UnpackTransparencyColor();
+
+		// prepare next step around the loop (if any throughput is left)
+		// retrace the exact same ray with TMin one ulp past the hit we just found
+		Ray.TMin = ComputeNewTMin(Ray.Origin, Ray.Direction, PackedPayload.HitT);
+		if (!(Ray.TMin < Ray.TMax) || all(PathThroughput == 0))
+		{
+			break;
+		}
+	}
+	PathThroughput = ResultThroughput * HitSample.GetNormalization();
+	return ResultPayload;
+}
+
+
+
+FPathTracingResult PathTracingDebugCore(FRayDesc Ray, uint2 LaunchIndex, int SampleIndex)
+{
+	FPathTracingResult Result = (FPathTracingResult) 0;
+	float3 PathThroughput = 1;
+	float PathRoughness = 0;
+	uint ScatterType = PATHTRACER_SCATTER_CAMERA;
+	FRandomSequence RandSequence = RandomSequenceCreate(uint3(LaunchIndex, View.StateFrameIndex), SampleIndex, SamplesPerPixel);
+
+	// State about the markov chain selection that needs to be kept from one bounce to the next
+	float ScoresSum = 0;
+	float3 PreviousN = 0;
+	FMarkovChainState ChosenState = (FMarkovChainState)0;
+	// MCPG parameters
+#define N_MC  5 // From paper, seems to be the smallest value which produces good results, some subtle improvements from higher values at the cost of performance
+	const float AdaptiveGridProb = 0.7; // Mentioned in text - TODO: could rely on adaptive grid only with a tweak to level distribution?
+	const float StaticGridWidth = 25.6; // matches level 24 of adaptive grid
+	const float LevelScale = 3.0; 		// 3 levels per octave
+	const float MinWidth = 0.1; 		// 1mm - smallest possible voxel size
+	const float TanFactor = 0.0087;		// ~2*tan(0.25 degrees)
+	const int 	NMax = 1024;			// Blend factors [See Dittebrandt et al.]
+	const float AlphaMin = 1.0 / 100.0; // Allow blending over 100 frames, higher than original publication because the chains get reset during BSDF sampling
+	const float MaterialProb = 0.15f; 	// Chance of BSDF sampling (TODO: adapt this as a function of roughness)
+
+	LOOP
+	for (int Bounce = 0; Bounce <= MaxBounces; Bounce++)
+	{
+		float3 LocalRadiance = 0;
+		FPackedPathTracingPayload PackedPayload = TraceTransparentRay(Ray, ScatterType, Bounce, PathRoughness, PathThroughput, RandSequence, Result, LocalRadiance);
+
+		if (PackedPayload.IsMiss())
+		{
+			// We didn't hit anything, done
+			break;
 		}
 
+		const float3 TranslatedWorldPos = Ray.Origin + PackedPayload.HitT * Ray.Direction;
+		const float3 WorldNormal = PackedPayload.UnpackWorldNormal();
+		const uint PrimitiveLightingChannelMask = PackedPayload.GetPrimitiveLightingChannelMask();
+
+		const bool bCameraRay = Bounce == 0;
+		const bool bLastBounce = Bounce == MaxBounces;
+
+		FLightLoopCount LightLoopCount = LightGridLookup(TranslatedWorldPos);
+
 		// Compute direct lighting
 		FPathTracingPayload Payload = UnpackPathTracingPayload(PackedPayload, Ray);
 		AdjustPayloadAfterUnpack(Payload, true);
+		// TEMP: If the material has a valid SSS lobe, just merge it into diffuse for now
 		FSSSRandomWalkInfo SSS = GetMaterialSSSInfo(Payload, Ray.Direction);
 		if (any(SSS.Weight > 0))
 		{
-			// If the material has a valid SSS lobe, just merge it into diffuse for now
 			Payload.DiffuseColor += SSS.Color;
 		}
-		float DenoiserContrib = DenoiserAOVWeight(PathRoughness);
-		if (DenoiserContrib > 0)
-		{
-			float3 DenoiserThroughput = PathThroughput * DenoiserContrib;
-
-			// G-Buffer like data for denoisers
-#if PATHTRACING_SUBSTRATE_PAYLOAD
-			float4 RoughnessData = PackedPayload.UnpackRoughnessAniso();
-			float Roughness = lerp(RoughnessData.x, RoughnessData.y, RoughnessData.z); // mix both lobes to get an average value
-			DenoiserThroughput *= Payload.WeightV * Payload.BSDFOpacity * lerp(1.0, Payload.TransmittanceN, Payload.CoverageAboveAlongN);
-#else
-			float Roughness = Payload.Roughness;
-#endif
-			float NormThroughput = Luminance(DenoiserThroughput);
-			float3 DiffuseColor = Payload.DiffuseColor;
-			float3 SpecularColor = Payload.SpecularColor;
-
-			Result.DiffuseColor += DenoiserThroughput * Payload.DiffuseColor;
-			Result.SpecularColor += DenoiserThroughput * Payload.SpecularColor * DenoiserRoughnessWeight(Roughness, PathRoughness);
-			Result.Roughness += NormThroughput * Roughness;
-			Result.WorldNormal += NormThroughput* PackedPayload.UnpackWorldNormal();
-		}
 
 		// Generate main random number for direct+indirect sampling
 		float4 RandSample = RandSequence.Get4D();
@@ -397,7 +534,8 @@ FPathTracingResult PathTracingDebugCore(FRayDesc Ray, uint2 LaunchIndex, int Sam
 
 #if 1
 				// Evaluate real material
-				FMaterialEval MaterialEval = EvalMaterial(-Ray.Direction, LightSample.Direction, Payload, 1.0);
+				float2 LightDiffuseSpecularScale = float2(GetLightDiffuseScale(LightId), GetLightSpecularScale(LightId));
+				FMaterialEval MaterialEval = EvalMaterial(-Ray.Direction, LightSample.Direction, Payload, LightDiffuseSpecularScale);
 				float3 LightContrib = PathThroughput * LightSample.RadianceOverPdf * MaterialEval.Weight * MaterialEval.Pdf;
 #if USE_BSDF_SAMPLING
 				LightContrib *= MISWeightPower(LightSample.Pdf, MaterialEval.Pdf);
@@ -455,6 +593,7 @@ FPathTracingResult PathTracingDebugCore(FRayDesc Ray, uint2 LaunchIndex, int Sam
 				{
 					ShadowRay.TMax = LightResult.HitT;
 					float3 LightContrib = PathThroughput * MaterialSample.Weight * LightResult.Radiance;
+					LightContrib *= MaterialSample.ScatterType == PATHTRACER_SCATTER_DIFFUSE ? GetLightDiffuseScale(LightId) : GetLightSpecularScale(LightId);
 					LightContrib *= MISWeightPower(MaterialSample.Pdf, LightResult.Pdf * LightPickPdf);
 #if LIGHT_LOOP_METHOD == 2
 					// Defer tracing shadow until later
@@ -624,7 +763,7 @@ FPathTracingResult PathTracingDebugCore(FRayDesc Ray, uint2 LaunchIndex, int Sam
 			Throughput = MaterialSample.Weight;
 			SampledDirection = MaterialSample.Direction;
 			ScatterType = MaterialSample.ScatterType;
-			PathRoughness = MaterialSample.Roughness;
+			PathRoughness = max(MaterialSample.Roughness, PathRoughness);
 			ChosenState = (FMarkovChainState)0;
 		}
 		if (ScoresSum > 0)
@@ -661,21 +800,23 @@ FPathTracingResult PathTracingDebugCore(FRayDesc Ray, uint2 LaunchIndex, int Sam
 
 FPathTracingResult PathTracingDebugCore(FRayDesc Ray, uint2 LaunchIndex)
 {
-
-	const uint RayFlags = RAY_FLAG_CULL_BACK_FACING_TRIANGLES;
-#if INCLUDE_TRANSLUCENT
-	const uint InstanceInclusionMask = PATHTRACER_MASK_CAMERA | PATHTRACER_MASK_CAMERA_TRANSLUCENT;
-#else
-	const uint InstanceInclusionMask = PATHTRACER_MASK_CAMERA;
-#endif
+	FPathTracingResult Result = (FPathTracingResult) 0;
+	const uint RayFlags = DebugMode == PATH_TRACER_DEBUG_VIZ_HITKIND ? 0 : RAY_FLAG_CULL_BACK_FACING_TRIANGLES;
 	const uint MissShaderIndex = 0;
+	uint InstanceInclusionMask = PATHTRACER_MASK_CAMERA;
+	int TransparentHitCount = 0;
 
-	FPackedPathTracingPayload PackedPayload = InitPathTracingPayload(PATHTRACER_SCATTER_CAMERA, 0.0);
-	PackedPayload.SetStochasticSlabRand(RandomSequenceCreate(uint3(LaunchIndex, View.StateFrameIndex), 0, 1).Get1D());
+	const float RandSample = RandomSequenceCreate(uint3(LaunchIndex, View.StateFrameIndex), 0, 1).Get1D();
+
+
+	float3 PathThroughput = 1;
+	FPackedPathTracingPayload OpaquePackedPayload = InitPathTracingPayload(PATHTRACER_SCATTER_CAMERA, 0.0);
+	OpaquePackedPayload.SetDBufferA(float4(0, 0, 0, 1));
+	OpaquePackedPayload.SetDBufferB(float4(0, 0, 0, 1));
+	OpaquePackedPayload.SetDBufferC(float4(0, 0, 0, 1));
+	OpaquePackedPayload.SetStochasticSlabRand(RandSample);
+	OpaquePackedPayload.SetFlag(PATH_TRACING_PAYLOAD_INPUT_FLAG_IGNORE_TRANSLUCENT);
 
-#if !INCLUDE_TRANSLUCENT
-	PackedPayload.SetFlag(PATH_TRACING_PAYLOAD_INPUT_FLAG_IGNORE_TRANSLUCENT);
-#endif
 	TraceRay(
 		TLAS,
 		RayFlags,
@@ -684,24 +825,22 @@ FPathTracingResult PathTracingDebugCore(FRayDesc Ray, uint2 LaunchIndex)
 		RAY_TRACING_NUM_SHADER_SLOTS,
 		MissShaderIndex,
 		Ray.GetNativeDesc(),
-		PackedPayload);
+		OpaquePackedPayload);
 
-	float4 Result = float4(0, 0, 0, 1);
-	float Z = 0;
-	if (PackedPayload.IsHit())
+	if (OpaquePackedPayload.IsHit())
 	{
-		const float3 ViewZ = View.ViewToTranslatedWorld[2].xyz;
-		Z = ConvertToDeviceZ(dot(Ray.Origin + PackedPayload.HitT * Ray.Direction, ViewZ));
-	}
+		// TODO: If hit, also process decals, etc ...
 
+		// Now that we know where the opaque hit is, trace the remaining portion of the ray prior to this
+		Ray.TMax = asfloat(asuint(OpaquePackedPayload.HitT) - 1);
+	}
 
 	if (DebugMode == PATH_TRACER_DEBUG_VIZ_VOLUME_LIGHT_COUNT)
 	{
-		float TMax = PackedPayload.IsHit() ? PackedPayload.HitT : Ray.TMax;
 		int NumLights = 0;
 		for (int LightId = SceneInfiniteLightCount; LightId < SceneLightCount; LightId++)
 		{
-			FVolumeLightSampleSetup LightSetup = PrepareLightVolumeSample(LightId, Ray.Origin, Ray.Direction, Ray.TMin, TMax);
+			FVolumeLightSampleSetup LightSetup = PrepareLightVolumeSample(LightId, Ray.Origin, Ray.Direction, Ray.TMin, Ray.TMax);
 			if (LightSetup.IsValid())
 			{
 				float LightProb = LightSetup.LightImportance * GetVolumetricScatteringIntensity(LightId);
@@ -718,88 +857,182 @@ FPathTracingResult PathTracingDebugCore(FRayDesc Ray, uint2 LaunchIndex)
 		}
 
 		float NumLightsMax = min(SceneLightCount - SceneInfiniteLightCount, RAY_TRACING_LIGHT_COUNT_MAXIMUM);
-		Result.rgb = NumLights > 0 ? BlueGreenRedColorMap(saturate(float(NumLights) / NumLightsMax)) : 0.18;
-		if (PackedPayload.IsHit())
+		Result.Radiance.rgb = NumLights > 0 ? BlueGreenRedColorMap(saturate(float(NumLights) / NumLightsMax)) : 0.18;
+		if (OpaquePackedPayload.IsHit())
 		{
-			Result.rgb *= abs(dot(PackedPayload.UnpackWorldNormal(), Ray.Direction));
+			// add some fake shading to help with definition
+			Result.Radiance.rgb *= abs(dot(OpaquePackedPayload.UnpackWorldNormal(), Ray.Direction));
 		}
-		return MakeResult(Result.rgb, Result.a, Z);
+		return Result;
 	}
 
-	if (PackedPayload.IsHit())
+
+	InstanceInclusionMask = PATHTRACER_MASK_CAMERA_TRANSLUCENT;
+
+	for (;;)
 	{
-		const float3 TranslatedWorldPos = Ray.Origin + PackedPayload.HitT * Ray.Direction;
-		const float3 WorldNormal = PackedPayload.UnpackWorldNormal();
-		const uint PrimitiveLightingChannelMask = PackedPayload.GetPrimitiveLightingChannelMask();
+#if PATH_TRACER_INCLUDE_TRANSLUCENT
+		FPackedPathTracingPayload PackedPayload = InitPathTracingPayload(PATHTRACER_SCATTER_CAMERA, 0.0);
+		PackedPayload.SetDBufferA(float4(0, 0, 0, 1));
+		PackedPayload.SetDBufferB(float4(0, 0, 0, 1));
+		PackedPayload.SetDBufferC(float4(0, 0, 0, 1));
+		PackedPayload.SetStochasticSlabRand(RandSample);
+		if (OpaquePackedPayload.IsHit())
+		{
+			const float3 ViewZ = View.ViewToTranslatedWorld[2].xyz;
+			const float3 TranslatedWorldPos = Ray.Origin + OpaquePackedPayload.HitT * Ray.Direction;
+			PackedPayload.SetFlag(PATH_TRACING_PAYLOAD_INPUT_FLAG_HAS_SCENE_DEPTH);
+			PackedPayload.SetSceneDepth(dot(TranslatedWorldPos, ViewZ));
+		}
 
-		FLightLoopCount LightLoopCount = LightGridLookup(TranslatedWorldPos);
+		// Trace the ray
+		TraceRay(
+			TLAS,
+			RayFlags,
+			InstanceInclusionMask,
+			RAY_TRACING_SHADER_SLOT_MATERIAL,
+			RAY_TRACING_NUM_SHADER_SLOTS,
+			MissShaderIndex,
+			Ray.GetNativeDesc(),
+			PackedPayload);
 
-		switch (DebugMode)
+		if (PackedPayload.IsMiss())
 		{
-			case PATH_TRACER_DEBUG_VIZ_RADIANCE: 					return MakeResult(PackedPayload.UnpackRadiance() * View.PreExposure, 0.0, Z);
-			case PATH_TRACER_DEBUG_VIZ_WORLD_NORMAL: 				return MakeResult(WorldNormal * 0.5 + 0.5, 0.0, Z);
-			case PATH_TRACER_DEBUG_VIZ_WORLD_SMOOTH_NORMAL:			return MakeResult(PackedPayload.UnpackWorldSmoothNormal() * 0.5 + 0.5, 0.0, Z);
-			case PATH_TRACER_DEBUG_VIZ_WORLD_GEO_NORMAL:			return MakeResult(PackedPayload.UnpackWorldGeoNormal() * 0.5 + 0.5, 0.0, Z);
-			default:
+			// If we miss the translucent geometry, the next hit is the opaque one
+			PackedPayload = OpaquePackedPayload;
+		}
+		else
+		{
+			TransparentHitCount++;
+		}
+#else // PATH_TRACER_INCLUDE_TRANSLUCENT
+		#define PackedPayload OpaquePackedPayload // just rename
+#endif // PATH_TRACER_INCLUDE_TRANSLUCENT
+		// If miss - exit
+		if (PackedPayload.IsMiss())
+		{
+			// still no hit -- the ray must have escaped to the background, we are done
+			// TODO: could add background Alpha here
+			break;
+		}
+
+		if (!PackedPayload.IsHoldout())
+		{
+			Result.Radiance.a += Luminance(PathThroughput * (1.0 - PackedPayload.UnpackTransparencyColor()));
+		}
+
+		{
+			const float3 TranslatedWorldPos = Ray.Origin + PackedPayload.HitT * Ray.Direction;
+			const float3 WorldNormal = PackedPayload.UnpackWorldNormal();
+			const uint PrimitiveLightingChannelMask = PackedPayload.GetPrimitiveLightingChannelMask();
+
+			FLightLoopCount LightLoopCount = LightGridLookup(TranslatedWorldPos);
+
+			float FakeShading = abs(dot(WorldNormal, Ray.Direction)); // cheap shading, just to give some definition to otherwise "solid color" modes
+			float3 Color = 0, ExtraAlpha = 0;
+			switch (DebugMode)
+			{
+				case PATH_TRACER_DEBUG_VIZ_RADIANCE: 					Color = PackedPayload.UnpackRadiance() * View.PreExposure; break;
+				case PATH_TRACER_DEBUG_VIZ_WORLD_NORMAL: 				Color = WorldNormal * 0.5 + 0.5; break;
+				case PATH_TRACER_DEBUG_VIZ_WORLD_SMOOTH_NORMAL:			Color = PackedPayload.UnpackWorldSmoothNormal() * 0.5 + 0.5;
+				case PATH_TRACER_DEBUG_VIZ_WORLD_GEO_NORMAL:			Color = PackedPayload.UnpackWorldGeoNormal() * 0.5 + 0.5;
 #if SUBSTRATE_ENABLED
-			// TODO: implement more visualizations for this case
-			case PATH_TRACER_DEBUG_VIZ_BASE_COLOR:					return MakeResult(lerp(PackedPayload.UnpackDiffuseColor(), PackedPayload.UnpackSpecularColor(), F0RGBToMetallic(PackedPayload.UnpackSpecularColor())), 0.0, Z);
-			case PATH_TRACER_DEBUG_VIZ_DIFFUSE_COLOR:				return MakeResult(PackedPayload.UnpackDiffuseColor(), 0.0, Z);
-			case PATH_TRACER_DEBUG_VIZ_SPECULAR_COLOR:				return MakeResult(PackedPayload.UnpackSpecularColor(), 0.0, Z);
-			case PATH_TRACER_DEBUG_VIZ_METALLIC:					return MakeResult(F0RGBToMetallic(PackedPayload.UnpackSpecularColor()), 0.0, Z);
-			case PATH_TRACER_DEBUG_VIZ_SPECULAR:					return MakeResult(F0RGBToDielectricSpecular(PackedPayload.UnpackSpecularColor()), 0.0, Z);
-			case PATH_TRACER_DEBUG_VIZ_ROUGHNESS:					return MakeResult(PackedPayload.UnpackRoughnessAniso().xyz, 0.0, Z);
-			case PATH_TRACER_DEBUG_VIZ_ANISOTROPY:					return MakeResult(AnisoDebugColor(PackedPayload.UnpackRoughnessAniso().w), 0.0, Z);
+				// TODO: implement more visualizations for this case
+				case PATH_TRACER_DEBUG_VIZ_BASE_COLOR:					Color = lerp(PackedPayload.UnpackDiffuseColor(), PackedPayload.UnpackSpecularColor(), F0RGBToMetallic(PackedPayload.UnpackSpecularColor())); break;
+				case PATH_TRACER_DEBUG_VIZ_DIFFUSE_COLOR:				Color = PackedPayload.UnpackDiffuseColor(); break;
+				case PATH_TRACER_DEBUG_VIZ_SPECULAR_COLOR:				Color = PackedPayload.UnpackSpecularColor(); break;
+				case PATH_TRACER_DEBUG_VIZ_METALLIC:					Color = F0RGBToMetallic(PackedPayload.UnpackSpecularColor()); break;
+				case PATH_TRACER_DEBUG_VIZ_SPECULAR:					Color = F0RGBToDielectricSpecular(PackedPayload.UnpackSpecularColor()); break;
+				case PATH_TRACER_DEBUG_VIZ_ROUGHNESS:					Color = PackedPayload.UnpackRoughnessAniso().xyz; break;
+				case PATH_TRACER_DEBUG_VIZ_ANISOTROPY:					Color = AnisoDebugColor(PackedPayload.UnpackRoughnessAniso().w; break;
 #else
-			case PATH_TRACER_DEBUG_VIZ_BASE_COLOR:					return MakeResult(PackedPayload.UnpackBaseColor(), 0.0, Z);
-			case PATH_TRACER_DEBUG_VIZ_DIFFUSE_COLOR:				return MakeResult(PackedPayload.UnpackBaseColor() * (1 - PackedPayload.UnpackMetallic()), 0.0, Z);
-			case PATH_TRACER_DEBUG_VIZ_SPECULAR_COLOR:				return MakeResult(lerp(0.08 * PackedPayload.UnpackSpecular(), PackedPayload.UnpackBaseColor(), PackedPayload.UnpackMetallic()), 0.0, Z);
-			case PATH_TRACER_DEBUG_VIZ_METALLIC:					return MakeResult(PackedPayload.UnpackMetallic(), 0.0, Z);
-			case PATH_TRACER_DEBUG_VIZ_SPECULAR:					return MakeResult(PackedPayload.UnpackSpecular(), 0.0, Z);
-			case PATH_TRACER_DEBUG_VIZ_ROUGHNESS:					return MakeResult(PackedPayload.UnpackRoughness(), 0.0, Z);
-			case PATH_TRACER_DEBUG_VIZ_ANISOTROPY:					return MakeResult(AnisoDebugColor(PackedPayload.UnpackAnisotropy()), 0.0, Z);
-			case PATH_TRACER_DEBUG_VIZ_CUSTOM_DATA0:				return MakeResult(PackedPayload.UnpackCustomData0().rgb * View.PreExposure, PackedPayload.UnpackCustomData0().a, Z);
-			case PATH_TRACER_DEBUG_VIZ_CUSTOM_DATA1:				return MakeResult(PackedPayload.UnpackCustomData1().rgb * View.PreExposure, PackedPayload.UnpackCustomData1().a, Z);
+				case PATH_TRACER_DEBUG_VIZ_BASE_COLOR:					Color = PackedPayload.UnpackBaseColor(); break;
+				case PATH_TRACER_DEBUG_VIZ_DIFFUSE_COLOR:				Color = PackedPayload.UnpackBaseColor() * (1 - PackedPayload.UnpackMetallic()); break;
+				case PATH_TRACER_DEBUG_VIZ_SPECULAR_COLOR:				Color = lerp(0.08 * PackedPayload.UnpackSpecular(), PackedPayload.UnpackBaseColor(), PackedPayload.UnpackMetallic()); break;
+				case PATH_TRACER_DEBUG_VIZ_METALLIC:					Color = PackedPayload.UnpackMetallic(); break;
+				case PATH_TRACER_DEBUG_VIZ_SPECULAR:					Color = PackedPayload.UnpackSpecular(); break;
+				case PATH_TRACER_DEBUG_VIZ_ROUGHNESS:					Color = PackedPayload.UnpackRoughness(); break;
+				case PATH_TRACER_DEBUG_VIZ_ANISOTROPY:					Color = AnisoDebugColor(PackedPayload.UnpackAnisotropy()); break;
+				case PATH_TRACER_DEBUG_VIZ_CUSTOM_DATA0:				Color = PackedPayload.UnpackCustomData0().rgb * View.PreExposure; ExtraAlpha = PackedPayload.UnpackCustomData0().a; break;
+				case PATH_TRACER_DEBUG_VIZ_CUSTOM_DATA1:				Color = PackedPayload.UnpackCustomData1().rgb * View.PreExposure; ExtraAlpha = PackedPayload.UnpackCustomData1().a; break;
 #endif
-			case PATH_TRACER_DEBUG_VIZ_OPACITY:						return MakeResult(PackedPayload.UnpackTransparencyColor(), 0.0, Z);
-			case PATH_TRACER_DEBUG_VIZ_IOR:							return MakeResult(PackedPayload.UnpackIor(), 0.0, Z);
-			case PATH_TRACER_DEBUG_VIZ_SHADING_MODEL:				return MakeResult(GetShadingModelColor(PackedPayload.GetShadingModelID()), 0.0, Z);
-			case PATH_TRACER_DEBUG_VIZ_LIGHTING_CHANNEL_MASK:		return MakeResult(GetLightingChannelMaskDebugColor(PackedPayload.GetPrimitiveLightingChannelMask()), 0.0, Z);
-			case PATH_TRACER_DEBUG_VIZ_WORLD_POSITION:
-			{
-				float3 AbsoluteWorldPosition = DFDemote(DFFastSubtract(TranslatedWorldPos, ResolvedView.PreViewTranslation));
-				const float CellWidth = 10.0f; // 10cm
-				int3 Coord = int3(round(AbsoluteWorldPosition / CellWidth));
-				uint HashIndex = GridHash(int4(Coord, 0)).x;
-				float3 Color = HUEtoRGB((HashIndex & 0xFFFFFF) * 5.96046447754e-08);
-				Color *= abs(dot(WorldNormal, Ray.Direction));
-				return MakeResult(Color, 0.0, Z);
-			}
-			case PATH_TRACER_DEBUG_VIZ_PRIMARY_RAYS: break; // fallthrough to implementation below
-			case PATH_TRACER_DEBUG_VIZ_WORLD_TANGENT: 				return MakeResult(PackedPayload.UnpackWorldTangent() * 0.5 + 0.5, 0.0, Z);
-			case PATH_TRACER_DEBUG_VIZ_LIGHT_GRID_COUNT:
-			case PATH_TRACER_DEBUG_VIZ_LIGHT_GRID_AXIS:
-			{
+				case PATH_TRACER_DEBUG_VIZ_OPACITY:						Color = PackedPayload.UnpackTransparencyColor(); break;
+				case PATH_TRACER_DEBUG_VIZ_IOR:							Color = PackedPayload.UnpackIor(); break;
+				case PATH_TRACER_DEBUG_VIZ_SHADING_MODEL:				Color = GetShadingModelColor(PackedPayload.GetShadingModelID()); break;
+				case PATH_TRACER_DEBUG_VIZ_LIGHTING_CHANNEL_MASK:		Color = GetLightingChannelMaskDebugColor(PackedPayload.GetPrimitiveLightingChannelMask()); break;
+				case PATH_TRACER_DEBUG_VIZ_WORLD_POSITION:
+				{
+					float3 AbsoluteWorldPosition = DFDemote(DFFastSubtract(TranslatedWorldPos, ResolvedView.PreViewTranslation));
+					const float CellWidth = 10.0f; // 10cm
+					int3 Coord = int3(round(AbsoluteWorldPosition / CellWidth));
+					uint HashIndex = GridHash(int4(Coord, 0)).x;
+					Color = HUEtoRGB((HashIndex & 0xFFFFFF) * 5.96046447754e-08);
+					Color *= FakeShading;
+					break;
+				}
+				case PATH_TRACER_DEBUG_VIZ_WORLD_TANGENT: 				Color = PackedPayload.UnpackWorldTangent() * 0.5 + 0.5; break;
+				case PATH_TRACER_DEBUG_VIZ_LIGHT_GRID_COUNT:
+				case PATH_TRACER_DEBUG_VIZ_LIGHT_GRID_AXIS:
+				{
 
-				float3 Color = LightGridVisualize(LightLoopCount, DebugMode == PATH_TRACER_DEBUG_VIZ_LIGHT_GRID_AXIS ? 4 : 1);
-				Color *= abs(dot(WorldNormal, Ray.Direction));
-				return MakeResult(Color, 0.0, Z);
-			}
-			case PATH_TRACER_DEBUG_VIZ_DECAL_GRID_COUNT:
-			case PATH_TRACER_DEBUG_VIZ_DECAL_GRID_AXIS:
-			{
-				FDecalLoopCount DecalLoopCount = DecalGridLookup(TranslatedWorldPos);
-				float3 Color = PackedPayload.IsDecalReceiver() ? DecalGridVisualize(DecalLoopCount, DebugMode == PATH_TRACER_DEBUG_VIZ_DECAL_GRID_AXIS ? 3 : 1) : 0.18;
-				Color *= abs(dot(WorldNormal, Ray.Direction));
-				return MakeResult(Color, 0.0, Z);
-			}
-			case PATH_TRACER_DEBUG_VIZ_HITKIND:
-			{
-				return MakeResult((PackedPayload.IsFrontFace() ? float3(0, 1, 0) : float3(1, 0, 0)) * abs(dot(WorldNormal, Ray.Direction)), 0.0, Z);
+					Color = LightGridVisualize(LightLoopCount, DebugMode == PATH_TRACER_DEBUG_VIZ_LIGHT_GRID_AXIS ? 4 : 1);
+					Color *= FakeShading;
+					break;
+				}
+				case PATH_TRACER_DEBUG_VIZ_DECAL_GRID_COUNT:
+				case PATH_TRACER_DEBUG_VIZ_DECAL_GRID_AXIS:
+				{
+					FDecalLoopCount DecalLoopCount = DecalGridLookup(TranslatedWorldPos);
+					Color = PackedPayload.IsDecalReceiver() ? DecalGridVisualize(DecalLoopCount, DebugMode == PATH_TRACER_DEBUG_VIZ_DECAL_GRID_AXIS ? 3 : 1) : 0.18;
+					Color *= FakeShading;
+					break;
+				}
+				case PATH_TRACER_DEBUG_VIZ_HITKIND:
+				{
+					Color = (PackedPayload.IsFrontFace() ? float3(0, 1, 0) : float3(1, 0, 0)) * FakeShading;
+					break;
+				}
+				default:
+				{
+					break;
+				}
 			}
+
+			Result.Radiance.rgb += PathThroughput * Color;
+		}
+
+		if (PackedPayload.ShouldOutputDepth() && Result.Z == 0.0)
+		{
+			const float3 ViewZ = View.ViewToTranslatedWorld[2].xyz;
+			const float3 TranslatedWorldPos = Ray.Origin + PackedPayload.HitT * Ray.Direction;
+			Result.LinearDepth = dot(TranslatedWorldPos, ViewZ);
+			Result.Z = ConvertToDeviceZ(Result.LinearDepth);
+			Result.WorldNormal = PackedPayload.UnpackWorldNormal();
+		}
+
+		// account for local transparency change
+		PathThroughput *= PackedPayload.UnpackTransparencyColor();
+
+		// prepare next step around the loop (if any throughput is left)
+		// retrace the exact same ray with TMin one ulp past the hit we just found
+		Ray.TMin = ComputeNewTMin(Ray.Origin, Ray.Direction, PackedPayload.HitT);
+		if (!(Ray.TMin < Ray.TMax) || all(PathThroughput == 0))
+		{
+			break;
 		}
 	}
-	return MakeResult(0, 1, Z);
+
+	if (DebugMode == PATH_TRACER_DEBUG_VIZ_TRANSPARENCY_COUNT)
+	{
+		Result.Radiance.rgb = TransparentHitCount > 0 ? BlueGreenRedColorMap(saturate((TransparentHitCount - 1) / 16.0)) : 0.5;
+		if (OpaquePackedPayload.IsHit())
+		{
+			// add some fake shading to help with definition
+			Result.Radiance.rgb *= abs(dot(OpaquePackedPayload.UnpackWorldNormal(), Ray.Direction));
+		}
+	}
+
+	return Result;
 }
 
 #endif // !PATH_TRACER_PRIMARY_RAYS
diff --git a/Engine/Shaders/Shared/PathTracingDefinitions.h b/Engine/Shaders/Shared/PathTracingDefinitions.h
--- a/Engine/Shaders/Shared/PathTracingDefinitions.h
+++ b/Engine/Shaders/Shared/PathTracingDefinitions.h
@@ -104,4 +104,5 @@
 #define PATH_TRACER_DEBUG_VIZ_DECAL_GRID_COUNT 			22
 #define PATH_TRACER_DEBUG_VIZ_DECAL_GRID_AXIS			23
 #define PATH_TRACER_DEBUG_VIZ_VOLUME_LIGHT_COUNT        24
-#define PATH_TRACER_DEBUG_VIZ_HITKIND                   25
\ No newline at end of file
+#define PATH_TRACER_DEBUG_VIZ_HITKIND                   25
+#define PATH_TRACER_DEBUG_VIZ_TRANSPARENCY_COUNT		26
\ No newline at end of file
diff --git a/Engine/Source/Runtime/Renderer/Private/PathTracing.cpp b/Engine/Source/Runtime/Renderer/Private/PathTracing.cpp
--- a/Engine/Source/Runtime/Renderer/Private/PathTracing.cpp
+++ b/Engine/Source/Runtime/Renderer/Private/PathTracing.cpp
@@ -87,6 +87,7 @@ struct FPathTracingDebugModeInfo
 		REGISTER_VIZ_TYPE(TEXT("DecalGridAxis") 	   	, PATH_TRACER_DEBUG_VIZ_DECAL_GRID_AXIS			);
 		REGISTER_VIZ_TYPE(TEXT("VolumeLightCount") 		, PATH_TRACER_DEBUG_VIZ_VOLUME_LIGHT_COUNT		);
 		REGISTER_VIZ_TYPE(TEXT("HitKind")            	, PATH_TRACER_DEBUG_VIZ_HITKIND             	);
+		REGISTER_VIZ_TYPE(TEXT("TransparencyCount")     , PATH_TRACER_DEBUG_VIZ_TRANSPARENCY_COUNT		);
 #undef REGISTER_VIZ_TYPE
 	}
 
@@ -123,6 +124,12 @@ TAutoConsoleVariable<bool> CVarPathTracingVisualizeEnableEmissive(
 	ECVF_RenderThreadSafe
 );
 
+TAutoConsoleVariable<bool> CVarPathTracingVisualizeIncludeTranslucent(
+	TEXT("r.PathTracing.Visualize.IncludeTranslucent"),
+	false,
+	TEXT("Indicates if translucent materials should be included in selected visualization mode."),
+	ECVF_RenderThreadSafe
+);
 
 TAutoConsoleVariable<int32> CVarPathTracingVisualizeSamplesPerPixel(
 	TEXT("r.PathTracing.Visualize.SamplesPerPixel"),
@@ -427,9 +434,9 @@ TAutoConsoleVariable<int32> CVarPathTracingSamplerType(
 );
 
 #if WITH_MGPU
-TAutoConsoleVariable<int32> CVarPathTracingMultiGPU(
+TAutoConsoleVariable<bool> CVarPathTracingMultiGPU(
 	TEXT("r.PathTracing.MultiGPU"),
-	0,
+	false,
 	TEXT("Run the path tracer using all available GPUs when enabled\n")
 	TEXT("Using this functionality in the editor requires -MaxGPUCount=N setting on the command line."),
 	ECVF_RenderThreadSafe
@@ -1426,7 +1433,8 @@ class FPathTracingDebugRG : public FGlobalShader
 
 	class FUseSER : SHADER_PERMUTATION_BOOL("PATH_TRACER_USE_SER");
 	class FUsePrimaryRays : SHADER_PERMUTATION_BOOL("PATH_TRACER_PRIMARY_RAYS");
-	using FPermutationDomain = TShaderPermutationDomain<FUseSER, FUsePrimaryRays>;
+	class FIncludeTranslucent : SHADER_PERMUTATION_BOOL("PATH_TRACER_INCLUDE_TRANSLUCENT");
+	using FPermutationDomain = TShaderPermutationDomain<FUseSER, FUsePrimaryRays, FIncludeTranslucent>;
 
 	static bool ShouldCompilePermutation(const FGlobalShaderPermutationParameters& Parameters)
 	{
@@ -1477,6 +1485,8 @@ class FPathTracingDebugRG : public FGlobalShader
 		SHADER_PARAMETER_RDG_BUFFER_UAV(RWStructuredBuffer<FMarkovChainState>, MarkovChainStateGrid)
 		SHADER_PARAMETER_STRUCT_REF(FViewUniformShaderParameters, ViewUniformBuffer)
 		SHADER_PARAMETER_RDG_TEXTURE(Texture2D<uint>, RandomSequenceSpaceFillingCurve)
+		SHADER_PARAMETER_TEXTURE(Texture2D, PreIntegratedGF)
+		SHADER_PARAMETER_SAMPLER(SamplerState, PreIntegratedGFSampler)
 		SHADER_PARAMETER(int32, DebugMode)
 		SHADER_PARAMETER(int32, SamplesPerPixel)
 		SHADER_PARAMETER(int32, SampleIndex)
@@ -3302,6 +3312,7 @@ void FDeferredShadingSceneRenderer::PreparePathTracing(const FViewInfo& View, co
 			FPathTracingDebugRG::FPermutationDomain PermutationVector;
 			PermutationVector.Set<FPathTracingDebugRG::FUseSER>(bUseSER && bUsePrimaryRays);
 			PermutationVector.Set<FPathTracingDebugRG::FUsePrimaryRays>(bUsePrimaryRays);
+			PermutationVector.Set<FPathTracingDebugRG::FIncludeTranslucent>(CVarPathTracingVisualizeIncludeTranslucent.GetValueOnRenderThread());
 			auto RayGenShader = GetGlobalShaderMap(ViewFamily.GetShaderPlatform())->GetShader<FPathTracingDebugRG>(PermutationVector);
 			OutRayGenShaders.Add(RayGenShader.GetRayTracingShader());
 		}
@@ -3516,7 +3527,7 @@ void FDeferredShadingSceneRenderer::RenderPathTracing(
 
 		FRDGTextureRef DiffuseColorTexture = GraphBuilder.CreateTexture(AlbedoNormalDesc, TEXT("PathTracer.DiffuseColor"));
 		FRDGTextureRef SpecularColorTexture = GraphBuilder.CreateTexture(AlbedoNormalDesc, TEXT("PathTracer.SpecularColor"));
-		FRDGTextureRef NormalTexture = GraphBuilder.CreateTexture(AlbedoNormalDesc, TEXT("PathTracer.Normal_Aleks"));
+		FRDGTextureRef NormalTexture = GraphBuilder.CreateTexture(AlbedoNormalDesc, TEXT("PathTracer.WorldNormal"));
 		FRDGTextureRef RoughnessTexture = GraphBuilder.CreateTexture(RoughnessDesc, TEXT("PathTracer.Roughness"));
 		FRDGTextureRef LinearDepthTexture = GraphBuilder.CreateTexture(RoughnessDesc, TEXT("PathTracer.LinearDepth"));
 
@@ -3559,6 +3570,8 @@ void FDeferredShadingSceneRenderer::RenderPathTracing(
 				PassParameters->MarkovChainStateGridSize = MarkovChainStateGridHashSize;
 				PassParameters->ViewUniformBuffer = View.ViewUniformBuffer;
 				PassParameters->RandomSequenceSpaceFillingCurve = GSystemTextures.GetSpaceFillingCurveTexture(GraphBuilder);
+				PassParameters->PreIntegratedGF = GSystemTextures.PreintegratedGF->GetRHI();
+				PassParameters->PreIntegratedGFSampler = TStaticSamplerState<SF_Bilinear, AM_Clamp, AM_Clamp, AM_Clamp>::GetRHI();
 				PassParameters->DebugMode = DebugMode;
 
 				PassParameters->SamplesPerPixel = SamplesPerPixel;
@@ -3602,6 +3615,7 @@ void FDeferredShadingSceneRenderer::RenderPathTracing(
 			FPathTracingDebugRG::FPermutationDomain PermutationVector;
 			PermutationVector.Set<FPathTracingDebugRG::FUseSER>(bUseSER && bUsePrimaryRays);
 			PermutationVector.Set<FPathTracingDebugRG::FUsePrimaryRays>(bUsePrimaryRays);
+			PermutationVector.Set<FPathTracingDebugRG::FIncludeTranslucent>(CVarPathTracingVisualizeIncludeTranslucent.GetValueOnRenderThread());
 			TShaderMapRef<FPathTracingDebugRG> RayGenShader(View.ShaderMap, PermutationVector);
 			ClearUnusedGraphResources(RayGenShader, PassParameters);
 			GraphBuilder.AddPass(
@@ -3844,7 +3858,7 @@ void FDeferredShadingSceneRenderer::RenderPathTracing(
 	}
 
 #if WITH_MGPU
-	Config.UseMultiGPU = CVarPathTracingMultiGPU.GetValueOnRenderThread() != 0;
+	Config.UseMultiGPU = CVarPathTracingMultiGPU.GetValueOnRenderThread();
 	// TODO: Figure out how to support adaptive sampling in multi-gpu cases (this is complicated due to the swizzled layout of the variance texture)
 	Config.UseMultiGPU &= !Config.UseAdaptiveSampling;
 #else
