Fix shader compilation error.

diff --git a/Engine/Shaders/Private/PathTracing/PathTracingDebug.usf b/Engine/Shaders/Private/PathTracing/PathTracingDebug.usf
--- a/Engine/Shaders/Private/PathTracing/PathTracingDebug.usf
+++ b/Engine/Shaders/Private/PathTracing/PathTracingDebug.usf
@@ -944,7 +944,7 @@ FPathTracingResult PathTracingDebugCore(FRayDesc Ray, uint2 LaunchIndex)
 				case PATH_TRACER_DEBUG_VIZ_METALLIC:					Color = F0RGBToMetallic(PackedPayload.UnpackSpecularColor()); break;
 				case PATH_TRACER_DEBUG_VIZ_SPECULAR:					Color = F0RGBToDielectricSpecular(PackedPayload.UnpackSpecularColor()); break;
 				case PATH_TRACER_DEBUG_VIZ_ROUGHNESS:					Color = PackedPayload.UnpackRoughnessAniso().xyz; break;
-				case PATH_TRACER_DEBUG_VIZ_ANISOTROPY:					Color = AnisoDebugColor(PackedPayload.UnpackRoughnessAniso().w; break;
+				case PATH_TRACER_DEBUG_VIZ_ANISOTROPY:					Color = AnisoDebugColor(PackedPayload.UnpackRoughnessAniso().w); break;
 #else
 				case PATH_TRACER_DEBUG_VIZ_BASE_COLOR:					Color = PackedPayload.UnpackBaseColor(); break;
 				case PATH_TRACER_DEBUG_VIZ_DIFFUSE_COLOR:				Color = PackedPayload.UnpackBaseColor() * (1 - PackedPayload.UnpackMetallic()); break;
