Add Substrate visualization support for blendable GBuffer path.

diff --git a/Engine/Shaders/Private/RayTracing/RayTracingDebug.usf b/Engine/Shaders/Private/RayTracing/RayTracingDebug.usf
--- a/Engine/Shaders/Private/RayTracing/RayTracingDebug.usf
+++ b/Engine/Shaders/Private/RayTracing/RayTracingDebug.usf
@@ -17,6 +17,7 @@
 #include "../Visualization.ush"
 #include "../ShaderPrint.ush"
 #include "../Substrate/Substrate.ush"
+#include "../Substrate/SubstrateEvaluation.ush"
 #include "../Substrate/SubstrateVisualizeCommon.ush"
 #include "/Engine/Shared/SubstrateVisualizeDefinitions.h"
 
diff --git a/Engine/Shaders/Private/Substrate/SubstrateRead.ush b/Engine/Shaders/Private/Substrate/SubstrateRead.ush
--- a/Engine/Shaders/Private/Substrate/SubstrateRead.ush
+++ b/Engine/Shaders/Private/Substrate/SubstrateRead.ush
@@ -174,11 +174,13 @@ FSubstrateGBufferBSDF SubstrateReadGBufferBSDF(in FScreenSpaceData ScreenSpaceDa
 				SharedLocalBasisIndex, SharedLocalBasisTypes).InlinedBSDF;
 		}
 
+		#if SUBSTRATE_INLINE_SHADING
 		BSDF.Coverage = 1;
 		BSDF.Emissive = 0;
 		BSDF.ThicknessCm = (SSSType == SSS_TYPE_WRAP || SSSType == SSS_TYPE_TWO_SIDED_WRAP) ? SUBSTRATE_SIMPLEVOLUME_THICKNESS_CM : SUBSTRATE_LAYER_DEFAULT_THICKNESS_CM; // Identity rescaling in InternalGetSubstrateSlabBSDF 
 		BSDF.TmpMFP = SSSMFP;
 		BSDF.TopLayerDataWeight = 1;
+		#endif
 
 		BSDF.OperatorIndex = 0;
 		BSDF.LuminanceWeightV = 1.0f;
@@ -196,4 +198,11 @@ FSubstrateGBufferBSDF SubstrateReadGBufferBSDF(in FScreenSpaceData ScreenSpaceDa
 	return Out;
 }
 
+FSubstrateBSDF UnpackSubstrateBSDFIn(inout FSubstrateAddressing SubstrateAddressing)
+{
+	FScreenSpaceData ScreenSpaceData = GetScreenSpaceDataUint(SubstrateAddressing.PixelCoords);
+	FSubstrateGBufferBSDF GBufferBSDF = SubstrateReadGBufferBSDF(ScreenSpaceData);
+	return GBufferBSDF.BSDF;
+}
+
 #endif // SUBSTRATE_ENABLED
diff --git a/Engine/Shaders/Private/Substrate/SubstrateVisualize.usf b/Engine/Shaders/Private/Substrate/SubstrateVisualize.usf
--- a/Engine/Shaders/Private/Substrate/SubstrateVisualize.usf
+++ b/Engine/Shaders/Private/Substrate/SubstrateVisualize.usf
@@ -55,15 +55,6 @@ void MaterialPrintInfoCS(uint3 DispatchThreadId : SV_DispatchThreadID)
 
 	FShaderPrintContext Context = InitShaderPrintContext(true, uint2(50, 50));
 
-#if SUBSTRATE_GBUFFER_FORMAT==0
-
-	SerializeNullPixelDebugDataEntry();
-
-	Context = InitShaderPrintContext(true, uint2(50, 25));
-	Print(Context, TEXT("Blendable GBuffer does not support pixel inspection yet. Please use `r.PostProcessing.GBufferPicking 1` instead."), FontOrange);
-
-#else
-
 	// Print pixel material properties. Simple print for simple compilers such as dx11/fxc which cannot handle loop...
 #if SUBSTRATE_MAX_CLOSURE_COUNT!= 8
 #error Unsupported SUBSTRATE_MAX_CLOSURE_COUNT in visualisation shader.
@@ -80,8 +71,6 @@ void MaterialPrintInfoCS(uint3 DispatchThreadId : SV_DispatchThreadID)
 	// Serialise out to buffer for print.
 	SerializeSubstratePixelDebugDataEntry(PixelPos, PrimaryView.FrameNumber);
 
-#endif
-
 	// Print pixel footprint
 	const float3 WorldPositionX = GetWorldPositionFromPixelPos(PixelPos + uint2(1, 0));
 	const float3 WorldPositionY = GetWorldPositionFromPixelPos(PixelPos + uint2(0, 1));
diff --git a/Engine/Shaders/Private/Substrate/SubstrateVisualizeCommon.ush b/Engine/Shaders/Private/Substrate/SubstrateVisualizeCommon.ush
--- a/Engine/Shaders/Private/Substrate/SubstrateVisualizeCommon.ush
+++ b/Engine/Shaders/Private/Substrate/SubstrateVisualizeCommon.ush
@@ -5,6 +5,7 @@
 #if SUBSTRATE_ENABLED
 
 #include "/Engine/Private/ShaderPrint.ush"
+#include "/Engine/Private/Substrate/SubstrateRead.ush"
 
 #ifndef SUBSTRATE_MATERIALCONTAINER_IS_VIEWRESOURCE
 #error SUBSTRATE_MATERIALCONTAINER_IS_VIEWRESOURCE needs to be defined
@@ -119,46 +120,24 @@ void PrintAddress(inout FShaderPrintContext Context, inout FSubstrateAddressing
 
 void SubstratePrintBSDF(
 	inout FShaderPrintContext Context, 
-	inout FSubstrateAddressing SubstrateAddressing, 
-	FSubstratePixelHeader Header, 
-	FSubstrateSubsurfaceHeader SSSHeader,
+	FSubstrateBSDFContext BSDFContext,
 	float3 WorldPosition, 
-	float3 V, 
-	inout float2 RectMax, 
-	FSubstrateRaytracingPayload Payload)
+	float3 V)
 {
-	#if SUBSTRATE_MATERIALCONTAINER_IS_VIEWRESOURCE
-	const FSubstrateBSDF BSDF = UnpackSubstrateBSDFIn(Substrate.MaterialTextureArray, SubstrateAddressing, Header);
-	#else
-	const FSubstrateBSDF BSDF = UnpackSubstrateBSDFIn(Payload, SubstrateAddressing, Header); // Take the parameter of the first BSDF
-	#endif
-
 	const float LineSize = 5.f;
-	float3 WorldNormal = 0.0f;
+	const float3 WorldNormal = BSDFContext.N;
+
 	// Draw Referential
-	#if SUBSTRATE_MATERIALCONTAINER_IS_VIEWRESOURCE
-	{
-		const float3 DummyL = float3(0, 0, 1);
-		FSubstrateBSDFContext BSDFContext = SubstrateCreateBSDFContext(Header, BSDF, SubstrateAddressing, V, DummyL);
-		DrawReferentialTWS(WorldPosition, BSDFContext.X, BSDFContext.Y, BSDFContext.N, float3(1, 1, 0));
-		WorldNormal = BSDFContext.N;
-	}
-	#else
-	{
-		const float3x3 TangentBasis = SubstrateGetBSDFSharedBasis(Header, BSDF, SubstrateAddressing);
-		DrawReferentialTWS(WorldPosition, TangentBasis[0], TangentBasis[1], TangentBasis[2], float3(1, 1, 0));
-		WorldNormal = TangentBasis[2];
-	}
-	#endif
+	DrawReferentialTWS(WorldPosition, BSDFContext.X, BSDFContext.Y, BSDFContext.N, float3(1, 1, 0));
 
-	switch (BSDF_GETTYPE(BSDF))
+	switch (BSDF_GETTYPE(BSDFContext.BSDF))
 	{
 		case SUBSTRATE_BSDF_TYPE_SLAB:
 		{
-			const bool bHasHaziness = BSDF_GETHASHAZINESS(BSDF);
+			const bool bHasHaziness = BSDF_GETHASHAZINESS(BSDFContext.BSDF);
 			if (bHasHaziness)
 			{
-				FHaziness Haziness = UnpackHaziness(SLAB_HAZINESS(BSDF));
+				FHaziness Haziness = UnpackHaziness(SLAB_HAZINESS(BSDFContext.BSDF));
 				const bool  bHazeAsSimpleClearCoat = Haziness.bSimpleClearCoat;
 
 			#if CLEAR_COAT_BOTTOM_NORMAL
@@ -180,14 +159,14 @@ void SubstratePrintBSDF(
 		case SUBSTRATE_BSDF_TYPE_EYE:
 		{
 			// Tangent basis for the iris
-			const float3x3 IrisTangentBasis = GetTangentBasis(EYE_IRISNORMAL(BSDF));
+			const float3x3 IrisTangentBasis = GetTangentBasis(EYE_IRISNORMAL(BSDFContext.BSDF));
 			//DrawReferentialTWS(WorldPosition, IrisTangentBasis[0], IrisTangentBasis[1], IrisTangentBasis[2], float3(0, 1, 1));
-			AddLineTWS(WorldPosition, WorldPosition + EYE_IRISNORMAL(BSDF) * LineSize, float4(0, 1, 1, 1));
+			AddLineTWS(WorldPosition, WorldPosition + EYE_IRISNORMAL(BSDFContext.BSDF) * LineSize, float4(0, 1, 1, 1));
 
 			// Tangent basis for the iris plane
-			const float3x3 IrisPlaneTangentBasis = GetTangentBasis(EYE_IRISPLANENORMAL(BSDF));
+			const float3x3 IrisPlaneTangentBasis = GetTangentBasis(EYE_IRISPLANENORMAL(BSDFContext.BSDF));
 			//DrawReferentialTWS(WorldPosition, IrisPlaneTangentBasis[0], IrisPlaneTangentBasis[1], IrisPlaneTangentBasis[2], float3(1, 0, 1));
-			AddLineTWS(WorldPosition, WorldPosition + EYE_IRISPLANENORMAL(BSDF) * LineSize, float4(1, 0, 1, 1));
+			AddLineTWS(WorldPosition, WorldPosition + EYE_IRISPLANENORMAL(BSDFContext.BSDF) * LineSize, float4(1, 0, 1, 1));
 
 		}
 		break;
@@ -205,33 +184,16 @@ void SubstratePrintBSDF(
 	}
 }
 
-void SubstratePrintMaterialProperties_Internal(
-	inout FShaderPrintContext Context, 
-	uint2 InCoord, 
-	float3 InWorldPosition, 
-	float3 V, 
-	uint InClosureIndex, 
-	inout FSubstrateAddressing SubstrateAddressing, 
-	FSubstratePixelHeader Header,
-	FSubstrateSubsurfaceHeader SSSHeader,
-	FSubstrateTopLayerData TopLayerData,
-	FSubstrateRaytracingPayload Payload,
-	const uint FootPrint_Start,
-	const uint FootPrint_PostHeader)
+#if SUBSTRATE_MATERIALCONTAINER_IS_VIEWRESOURCE && SUBSTRATE_GBUFFER_FORMAT == 1
+void SubstratePrintMaterialProperties(inout FShaderPrintContext Context, uint2 InCoord, float3 InWorldPosition, float3 V, uint InClosureIndex)
 {
-	const FFontColor FontBSDFCount = FontEmerald;
-	const FFontColor FontHeaderStateName = FontWhite;
-	const FFontColor FontHeaderPropName = FontWhite;
-	const FFontColor FontHeaderPropValue = FontSilver;
-
-	// Backgroun rect. for better print visibility. Track min/max point
-	float2 RectMin = Context.Pos;
-	float2 RectMax = Context.StartPos;
+	FSubstrateAddressing SubstrateAddressing = GetSubstratePixelDataByteOffset(InCoord, uint2(View.BufferSizeAndInvSize.xy), Substrate.MaxBytesPerPixel);
+	FSubstratePixelHeader Header = UnpackSubstrateHeaderIn(Substrate.MaterialTextureArray, SubstrateAddressing, Substrate.TopLayerTexture);
 
 	// BSDFs
 	if (InClosureIndex < Header.ClosureCount)
 	{
-		#if SUBSTRATE_MATERIALCONTAINER_IS_VIEWRESOURCE && SUBSTRATE_GBUFFER_FORMAT==1 && SUBSTRATE_MATERIAL_CLOSURE_COUNT > 1
+		#if SUBSTRATE_MATERIAL_CLOSURE_COUNT > 1
 		if (InClosureIndex > 0)
 		{
 			const uint AddressOffset = UnpackClosureOffsetAtIndex(Substrate.ClosureOffsetTexture[InCoord], InClosureIndex, Header.ClosureCount);
@@ -239,66 +201,35 @@ void SubstratePrintMaterialProperties_Internal(
 		}
 		#endif
 
-		RectMin = Context.Pos;
-		RectMax = Context.StartPos;
-		SubstratePrintBSDF(Context, SubstrateAddressing, Header, SSSHeader, InWorldPosition, V, RectMax, Payload);
+		const FSubstrateBSDF BSDF = UnpackSubstrateBSDFIn(Substrate.MaterialTextureArray, SubstrateAddressing, Header);
+		FSubstrateBSDFContext BSDFContext = SubstrateCreateBSDFContext(Header, BSDF, SubstrateAddressing, V, float3(0, 0, 1) /*DummyL*/);
+		SubstratePrintBSDF(Context, BSDFContext, InWorldPosition, V);
 	}
 }
-
-#if SUBSTRATE_MATERIALCONTAINER_IS_VIEWRESOURCE
+#elif SUBSTRATE_MATERIALCONTAINER_IS_VIEWRESOURCE && SUBSTRATE_GBUFFER_FORMAT == 0
 void SubstratePrintMaterialProperties(inout FShaderPrintContext Context, uint2 InCoord, float3 InWorldPosition, float3 V, uint InClosureIndex)
 {
-	FSubstrateAddressing SubstrateAddressing = GetSubstratePixelDataByteOffset(InCoord, uint2(View.BufferSizeAndInvSize.xy), Substrate.MaxBytesPerPixel);
-	const uint FootPrint_Start = SubstrateAddressing.ReadBytes;
-
-	FSubstratePixelHeader Header = UnpackSubstrateHeaderIn(Substrate.MaterialTextureArray, SubstrateAddressing, Substrate.TopLayerTexture);
-	const uint FootPrint_PostHeader = SubstrateAddressing.ReadBytes;
-
-	FSubstrateSubsurfaceHeader SSSHeader = SubstrateLoadSubsurfaceHeader(Substrate.MaterialTextureArray, Substrate.FirstSliceStoringSubstrateSSSData, SubstrateAddressing.PixelCoords);
-	FSubstrateRaytracingPayload DummyPayload;
-	FSubstrateTopLayerData TopLayerData = SubstrateUnpackTopLayerData(Substrate.TopLayerTexture.Load(uint3(InCoord, 0)));
-
-	SubstratePrintMaterialProperties_Internal(
-		Context, 
-		InCoord, 
-		InWorldPosition, 
-		V, 
-		InClosureIndex, 
-		SubstrateAddressing,
-		Header,
-		SSSHeader,
-		TopLayerData,
-		DummyPayload, 
-		FootPrint_Start, 
-		FootPrint_PostHeader);
+	// BSDFs
+	if (InClosureIndex < 1)
+	{
+		FScreenSpaceData ScreenSpaceData = GetScreenSpaceDataUint(InCoord);
+		FSubstrateGBufferBSDF GBufferBSDF = SubstrateReadGBufferBSDF(ScreenSpaceData);
+		FSubstrateBSDFContext BSDFContext = SubstrateCreateBSDFContext(GBufferBSDF.BSDF, GBufferBSDF.BSDFTangentBasis, V, float3(0, 0, 1) /*DummyL*/);
+		SubstratePrintBSDF(Context, BSDFContext, InWorldPosition, V);
+	}
 }
-#else
+#else // Raytracing payload
 void SubstratePrintMaterialProperties(inout FShaderPrintContext Context, uint2 InCoord, float3 InWorldPosition, float3 V, uint InClosureIndex, FSubstrateRaytracingPayload Payload)
 {
-	FSubstrateAddressing SubstrateAddressing = GetSubstratePixelDataByteOffset(InCoord, uint2(View.BufferSizeAndInvSize.xy), 0);	
-	const uint FootPrint_Start = SubstrateAddressing.ReadBytes;
-	
+	FSubstrateAddressing SubstrateAddressing = GetSubstratePixelDataByteOffset(InCoord, uint2(View.BufferSizeAndInvSize.xy), 0);		
 	FSubstratePixelHeader Header = UnpackSubstrateHeaderIn(Payload, SubstrateAddressing, Payload);
-	const uint FootPrint_PostHeader = SubstrateAddressing.ReadBytes;
-
-	FSubstrateSubsurfaceHeader SSSHeader = (FSubstrateSubsurfaceHeader)0;
-	FSubstrateTopLayerData TopLayerData = SubstrateUnpackTopLayerData(Payload.PackedTopLayerData);
 
 	for (uint ClosureIndex=0;ClosureIndex<Header.ClosureCount;ClosureIndex++)
 	{
-		SubstratePrintMaterialProperties_Internal(
-			Context, 
-			InCoord, 
-			InWorldPosition, 
-			V, 
-			InClosureIndex, 
-			SubstrateAddressing,
-			Header,
-			SSSHeader,
-			TopLayerData,
-			Payload,
-			FootPrint_Start,
-			FootPrint_PostHeader);
+		const FSubstrateBSDF BSDF = UnpackSubstrateBSDFIn(Payload, SubstrateAddressing, Header); // Take the parameter of the first BSDF
+		const float3x3 TangentBasis = SubstrateGetBSDFSharedBasis(Header, BSDF, SubstrateAddressing);
+		FSubstrateBSDFContext BSDFContext = SubstrateCreateBSDFContext(TangentBasis, BSDF, V, float3(0, 0, 1) /*DummyL*/, false/*bHasValidL*/);
+		SubstratePrintBSDF(Context, BSDFContext, InWorldPosition, V);
 	}
 }
 #endif
diff --git a/Engine/Shaders/Shared/SubstrateVisualizeDefinitions.h b/Engine/Shaders/Shared/SubstrateVisualizeDefinitions.h
--- a/Engine/Shaders/Shared/SubstrateVisualizeDefinitions.h
+++ b/Engine/Shaders/Shared/SubstrateVisualizeDefinitions.h
@@ -255,6 +255,121 @@ void SerializeSubstratePixelDebugData(FSubstrateDebugDataSerializer& S, FSubstra
 
 #ifndef __cplusplus
 
+
+#if SUBSTRATE_MATERIALCONTAINER_IS_VIEWRESOURCE && SUBSTRATE_GBUFFER_FORMAT == 0
+
+#include "/Engine/Private/Substrate/SubstrateRead.ush"
+#include "/Engine/Private/Substrate/SubstrateSubsurface.ush"
+
+FSubstratePixelHeader UnpackSubstrateHeaderIn(inout FSubstrateAddressing SubstrateAddressing)
+{	
+	const FGBufferData GBufferData = GetScreenSpaceDataUint(SubstrateAddressing.PixelCoords).GBuffer;
+
+	FSubstratePixelHeader Out = InitialiseSubstratePixelHeader();
+	if (GBufferData.ShadingModelID != SHADINGMODELID_UNLIT)
+	{
+		uint Mode = HEADER_MATERIALMODE_NONE;
+		if (GBufferData.ShadingModelID == SHADINGMODELID_SINGLELAYERWATER)
+		{
+			Mode = HEADER_MATERIALMODE_SLWATER;
+		}
+		else if (GBufferData.ShadingModelID == SHADINGMODELID_HAIR)
+		{
+			Mode = HEADER_MATERIALMODE_HAIR;
+		}
+		else if (GBufferData.ShadingModelID == SHADINGMODELID_EYE)
+		{
+			Mode = HEADER_MATERIALMODE_EYE;
+		}
+		else
+		{
+			const bool bHasAnisotropy = abs(GBufferData.Anisotropy) > 0;
+			const bool bComplexMaterial = bHasAnisotropy;
+			const bool bSimpleMaterial = !bComplexMaterial && GBufferData.ShadingModelID == SHADINGMODELID_DEFAULT_LIT;
+			const bool bSingleMaterial = !bComplexMaterial && !bSimpleMaterial;
+			const bool bScreenSpaceSubsurfaceScattering = GBufferData.ShadingModelID == SHADINGMODELID_SUBSURFACE_PROFILE || GBufferData.ShadingModelID == SHADINGMODELID_EYE;
+
+			Mode = bSimpleMaterial ? HEADER_MATERIALMODE_SLAB_SIMPLE : bSingleMaterial ? HEADER_MATERIALMODE_SLAB_SINGLE : HEADER_MATERIALMODE_SLAB_COMPLEX;
+		}
+
+		const bool bHasSSS = 
+		   GBufferData.ShadingModelID == SHADINGMODELID_SUBSURFACE
+		|| GBufferData.ShadingModelID == SHADINGMODELID_PREINTEGRATED_SKIN		
+		|| GBufferData.ShadingModelID == SHADINGMODELID_SUBSURFACE_PROFILE
+		|| GBufferData.ShadingModelID == SHADINGMODELID_TWOSIDED_FOLIAGE
+		|| GBufferData.ShadingModelID == SHADINGMODELID_EYE;
+
+		Out.ClosureCount = 1; 
+		Out.SetMaterialMode(Mode);
+		Out.SetHasSubsurface(bHasSSS);
+
+		// TODO
+		// Out.SetHasPrecShadowMask(...)
+		// Out.SetZeroPrecShadowMask(...)
+		// Out.SetIsFirstPerson(...)
+		// Out.SetCastContactShadow(...)
+		// Out.SetHasSubsurface(...)
+		// Out.SetDynamicIndirectShadowCasterRepresentation(...)
+
+		// TODO
+		#if SUBSTRATE_INLINE_SHADING
+		// Out.IrradianceAO.MaterialAO;
+		// Out.IrradianceAO.IndirectIrradiance;
+		// Out.IrradianceAO.DiffuseIndirectSampleOcclusion;
+		#endif
+	}
+
+	SubstrateAddressing.CurrentIndex = 0;
+	SubstrateAddressing.ReadBytes = 16;
+	return Out;
+}
+
+FSubstrateSubsurfaceHeader SubstrateLoadSubsurfaceHeader(inout FSubstrateAddressing SubstrateAddressing)
+{
+	const FGBufferData GBufferData = GetScreenSpaceDataUint(SubstrateAddressing.PixelCoords).GBuffer;
+
+	FSubstrateSubsurfaceHeader Out = (FSubstrateSubsurfaceHeader)0;
+
+	uint Mode = SSS_TYPE_WRAP;
+	if (GBufferData.ShadingModelID == SHADINGMODELID_SUBSURFACE || GBufferData.ShadingModelID == SHADINGMODELID_PREINTEGRATED_SKIN)
+	{
+		Mode = SSS_TYPE_WRAP;
+		SubstrateSubSurfaceHeaderSetWrapOpacity(Out, GBufferData.CustomData.a);
+	}
+	else if (GBufferData.ShadingModelID == SHADINGMODELID_TWOSIDED_FOLIAGE)
+	{
+		Mode = SSS_TYPE_TWO_SIDED_WRAP;
+		SubstrateSubSurfaceHeaderSetWrapOpacity(Out, GBufferData.CustomData.a);
+
+	}
+	if (GBufferData.ShadingModelID == SHADINGMODELID_SUBSURFACE_PROFILE ||
+		GBufferData.ShadingModelID == SHADINGMODELID_EYE)
+	{
+		Mode = SSS_TYPE_DIFFUSION_PROFILE;
+		SubstrateSubSurfaceHeaderSetProfile(Out, GBufferData.CustomData.a, ExtractSubsurfaceProfileInt(GBufferData));
+	}
+
+	SubstrateSubSurfaceHeaderSetSSSType(Out, Mode);
+	return Out;
+}
+
+FSubstrateTopLayerData SubstrateUnpackTopLayerData(inout FSubstrateAddressing SubstrateAddressing)
+{
+	const FGBufferData GBufferData = GetScreenSpaceDataUint(SubstrateAddressing.PixelCoords).GBuffer;
+
+	FSubstrateTopLayerData Out = (FSubstrateTopLayerData)0;
+	if (GBufferData.ShadingModelID != SHADINGMODELID_UNLIT)
+	{
+		Out.WorldNormal = GBufferData.WorldNormal;
+		Out.Roughness = GBufferData.Roughness;
+		Out.Material = 1;
+	}
+	return Out;
+}
+
+#endif
+
+
 #if SUBSTRATE_MATERIALCONTAINER_IS_VIEWRESOURCE
 FSubstratePixelDebugData ConvertToSerializableSubstratePixelDebugData(uint2 InCoord, uint GPUFrameNumber)
 #else
@@ -266,16 +381,21 @@ FSubstratePixelDebugData ConvertToSerializableSubstratePixelDebugData(uint2 InCo
 	FSubstrateAddressing SubstrateAddressing = GetSubstratePixelDataByteOffset(InCoord, uint2(View.BufferSizeAndInvSize.xy), Substrate.MaxBytesPerPixel);
 	const uint FootPrint_Start = SubstrateAddressing.ReadBytes;
 
-#if SUBSTRATE_MATERIALCONTAINER_IS_VIEWRESOURCE
+#if SUBSTRATE_MATERIALCONTAINER_IS_VIEWRESOURCE && SUBSTRATE_GBUFFER_FORMAT == 1
 	FSubstratePixelHeader Header = UnpackSubstrateHeaderIn(Substrate.MaterialTextureArray, SubstrateAddressing, Substrate.TopLayerTexture);
+#elif SUBSTRATE_MATERIALCONTAINER_IS_VIEWRESOURCE && SUBSTRATE_GBUFFER_FORMAT == 0
+	FSubstratePixelHeader Header = UnpackSubstrateHeaderIn(SubstrateAddressing);
 #else
 	FSubstratePixelHeader Header = UnpackSubstrateHeaderIn(Payload, SubstrateAddressing, Payload);
 #endif
 	const uint FootPrint_PostHeader = SubstrateAddressing.ReadBytes;
 
-#if SUBSTRATE_MATERIALCONTAINER_IS_VIEWRESOURCE
+#if SUBSTRATE_MATERIALCONTAINER_IS_VIEWRESOURCE && SUBSTRATE_GBUFFER_FORMAT == 1
 	FSubstrateSubsurfaceHeader SSSHeader = SubstrateLoadSubsurfaceHeader(Substrate.MaterialTextureArray, Substrate.FirstSliceStoringSubstrateSSSData, SubstrateAddressing.PixelCoords);
 	FSubstrateTopLayerData TopLayerData = SubstrateUnpackTopLayerData(Substrate.TopLayerTexture.Load(uint3(InCoord, 0)));
+#elif SUBSTRATE_MATERIALCONTAINER_IS_VIEWRESOURCE && SUBSTRATE_GBUFFER_FORMAT == 0
+	FSubstrateSubsurfaceHeader SSSHeader = SubstrateLoadSubsurfaceHeader(SubstrateAddressing);
+	FSubstrateTopLayerData TopLayerData = SubstrateUnpackTopLayerData(SubstrateAddressing);
 #else
 	FSubstrateSubsurfaceHeader SSSHeader = (FSubstrateSubsurfaceHeader)0;
 	FSubstrateTopLayerData TopLayerData = SubstrateUnpackTopLayerData(Payload.PackedTopLayerData);
@@ -321,8 +441,10 @@ FSubstratePixelDebugData ConvertToSerializableSubstratePixelDebugData(uint2 InCo
 	#endif
 	{
 		// Unpack BSDF data
-	#if SUBSTRATE_MATERIALCONTAINER_IS_VIEWRESOURCE
+	#if SUBSTRATE_MATERIALCONTAINER_IS_VIEWRESOURCE && SUBSTRATE_GBUFFER_FORMAT == 1
 		FSubstrateBSDF BSDF = UnpackSubstrateBSDF(Substrate.MaterialTextureArray, SubstrateAddressing, Header);
+	#elif SUBSTRATE_MATERIALCONTAINER_IS_VIEWRESOURCE && SUBSTRATE_GBUFFER_FORMAT == 0
+		FSubstrateBSDF BSDF = UnpackSubstrateBSDFIn(SubstrateAddressing);
 	#else
 		FSubstrateBSDF BSDF = UnpackSubstrateBSDFIn(Payload, SubstrateAddressing, Header);
 	#endif
diff --git a/Engine/Source/Runtime/Renderer/Private/Substrate/SubstrateVisualize.cpp b/Engine/Source/Runtime/Renderer/Private/Substrate/SubstrateVisualize.cpp
--- a/Engine/Source/Runtime/Renderer/Private/Substrate/SubstrateVisualize.cpp
+++ b/Engine/Source/Runtime/Renderer/Private/Substrate/SubstrateVisualize.cpp
@@ -51,6 +51,7 @@ class FMaterialPrintInfoCS : public FGlobalShader
 		SHADER_PARAMETER_STRUCT_REF(FViewUniformShaderParameters, ViewUniformBuffer)
 		SHADER_PARAMETER_RDG_UNIFORM_BUFFER(FSubstrateGlobalUniformParameters, Substrate)
 		SHADER_PARAMETER_STRUCT_INCLUDE(FSceneTextureParameters, SceneTextures)
+		SHADER_PARAMETER_RDG_UNIFORM_BUFFER(FSceneTextureUniformParameters, SceneTexturesStruct)
 		SHADER_PARAMETER_STRUCT_INCLUDE(ShaderPrint::FShaderParameters, ShaderPrintParameters)
 		SHADER_PARAMETER_RDG_BUFFER_UAV(RWStructuredBuffer<int>, SubstrateDebugDataUAV)
 	END_SHADER_PARAMETER_STRUCT()
@@ -96,6 +97,7 @@ class FVisualizeMaterialCountPS : public FGlobalShader
 		SHADER_PARAMETER_STRUCT_REF(FViewUniformShaderParameters, ViewUniformBuffer)
 		SHADER_PARAMETER_RDG_UNIFORM_BUFFER(FSubstrateGlobalUniformParameters, Substrate)
 		SHADER_PARAMETER_STRUCT_INCLUDE(FSceneTextureParameters, SceneTextures)
+		SHADER_PARAMETER_RDG_UNIFORM_BUFFER(FSceneTextureUniformParameters, SceneTexturesStruct)
 		SHADER_PARAMETER_STRUCT_INCLUDE(ShaderPrint::FShaderParameters, ShaderPrintParameters)
 		RENDER_TARGET_BINDING_SLOTS()
 	END_SHADER_PARAMETER_STRUCT()
@@ -141,6 +143,7 @@ class FSubstrateSystemInfoCS : public FGlobalShader
 		SHADER_PARAMETER_STRUCT_REF(FViewUniformShaderParameters, ViewUniformBuffer)
 		SHADER_PARAMETER_RDG_UNIFORM_BUFFER(FSubstrateGlobalUniformParameters, Substrate)
 		SHADER_PARAMETER_STRUCT_INCLUDE(FSceneTextureParameters, SceneTextures)
+		SHADER_PARAMETER_RDG_UNIFORM_BUFFER(FSceneTextureUniformParameters, SceneTexturesStruct)
 		SHADER_PARAMETER_RDG_BUFFER_UAV(RWStructuredBuffer<int>, SubstrateDebugDataUAV)
 	END_SHADER_PARAMETER_STRUCT()
 
@@ -182,6 +185,8 @@ class FMaterialDebugSubstrateTreeCS : public FGlobalShader
 		SHADER_PARAMETER(uint32, bOverrideCursorPosition)
 		SHADER_PARAMETER_STRUCT_REF(FViewUniformShaderParameters, ViewUniformBuffer)
 		SHADER_PARAMETER_RDG_UNIFORM_BUFFER(FSubstrateGlobalUniformParameters, Substrate)
+		SHADER_PARAMETER_STRUCT_INCLUDE(FSceneTextureParameters, SceneTextures)
+		SHADER_PARAMETER_RDG_UNIFORM_BUFFER(FSceneTextureUniformParameters, SceneTexturesStruct)
 		SHADER_PARAMETER_STRUCT_INCLUDE(ShaderPrint::FShaderParameters, ShaderPrintParameters)
 	END_SHADER_PARAMETER_STRUCT()
 
@@ -227,6 +232,8 @@ class FMaterialDebugSubstrateTreePS : public FGlobalShader
 		SHADER_PARAMETER(uint32, bOverrideCursorPosition)
 		SHADER_PARAMETER_STRUCT_REF(FViewUniformShaderParameters, ViewUniformBuffer)
 		SHADER_PARAMETER_RDG_UNIFORM_BUFFER(FSubstrateGlobalUniformParameters, Substrate)
+		SHADER_PARAMETER_STRUCT_INCLUDE(FSceneTextureParameters, SceneTextures)
+		SHADER_PARAMETER_RDG_UNIFORM_BUFFER(FSceneTextureUniformParameters, SceneTexturesStruct)
 		SHADER_PARAMETER_RDG_UNIFORM_BUFFER(FReflectionUniformParameters, ReflectionStruct)
 		SHADER_PARAMETER_STRUCT_REF(FReflectionCaptureShaderData, ReflectionCapture)
 		SHADER_PARAMETER_RDG_UNIFORM_BUFFER(FForwardLightUniformParameters, ForwardLightStruct)
@@ -958,6 +965,7 @@ static void AddVisualizeMaterialPropertiesPasses(FRDGBuilder& GraphBuilder, cons
 		PassParameters->ViewUniformBuffer = View.ViewUniformBuffer;
 		PassParameters->Substrate = Substrate::BindSubstrateGlobalUniformParameters(View);
 		PassParameters->SceneTextures = GetSceneTextureParameters(GraphBuilder, View);
+		PassParameters->SceneTexturesStruct = View.GetSceneTextures().UniformBuffer;
 		PassParameters->SubstrateDebugDataSizeInUints = NewSubstratePixelDebugBuffer.DebugDataSizeInUints;
 		PassParameters->SubstrateDebugDataUAV = NewSubstratePixelDebugBuffer.DebugDataUAV;
 		ShaderPrint::SetParameters(GraphBuilder, View.ShaderPrintData, PassParameters->ShaderPrintParameters);
@@ -987,6 +995,7 @@ static void AddVisualizeMaterialCountPasses(FRDGBuilder & GraphBuilder, const FV
 	PassParameters->bOverrideCursorPosition = WITH_EDITOR ? 0u : 1u;
 	PassParameters->Substrate = Substrate::BindSubstrateGlobalUniformParameters(View);
 	PassParameters->SceneTextures = GetSceneTextureParameters(GraphBuilder, View);
+	PassParameters->SceneTexturesStruct = View.GetSceneTextures().UniformBuffer;
 	PassParameters->RenderTargets[0] = FRenderTargetBinding(SceneColorTexture, ERenderTargetLoadAction::ELoad);
 	ShaderPrint::SetParameters(GraphBuilder, View.ShaderPrintData, PassParameters->ShaderPrintParameters);
 
@@ -1022,6 +1031,7 @@ static void AddVisualizeSystemInfoPasses(FRDGBuilder& GraphBuilder, const FViewI
 	PassParameters->ViewUniformBuffer = View.ViewUniformBuffer;
 	PassParameters->Substrate = Substrate::BindSubstrateGlobalUniformParameters(View);
 	PassParameters->SceneTextures = GetSceneTextureParameters(GraphBuilder, View);
+	PassParameters->SceneTexturesStruct = View.GetSceneTextures().UniformBuffer;
 
 	////////////////////////////////////////////////////////////////////////////////////////////////////
 	// Get a new pixel material buffer and render debug data into it
@@ -1343,6 +1353,8 @@ static void AddVisualizeAdvancedMaterialPasses(FRDGBuilder& GraphBuilder, const
 		FMaterialDebugSubstrateTreeCS::FParameters* PassParameters = GraphBuilder.AllocParameters<FMaterialDebugSubstrateTreeCS::FParameters>();
 		PassParameters->ViewUniformBuffer = View.ViewUniformBuffer;
 		PassParameters->Substrate = Substrate::BindSubstrateGlobalUniformParameters(View);
+		PassParameters->SceneTextures = GetSceneTextureParameters(GraphBuilder, View);
+		PassParameters->SceneTexturesStruct = View.GetSceneTextures().UniformBuffer;
 		PassParameters->bOverrideCursorPosition = WITH_EDITOR ? 0u : 1u;
 		ShaderPrint::SetParameters(GraphBuilder, View.ShaderPrintData, PassParameters->ShaderPrintParameters);
 
@@ -1354,6 +1366,8 @@ static void AddVisualizeAdvancedMaterialPasses(FRDGBuilder& GraphBuilder, const
 		FMaterialDebugSubstrateTreePS::FParameters* PassParameters = GraphBuilder.AllocParameters<FMaterialDebugSubstrateTreePS::FParameters>();
 		PassParameters->ViewUniformBuffer = View.ViewUniformBuffer;
 		PassParameters->Substrate = Substrate::BindSubstrateGlobalUniformParameters(View);
+		PassParameters->SceneTextures = GetSceneTextureParameters(GraphBuilder, View);
+		PassParameters->SceneTexturesStruct = View.GetSceneTextures().UniformBuffer;
 		PassParameters->bOverrideCursorPosition = WITH_EDITOR ? 0u : 1u;
 		PassParameters->ReflectionStruct = CreateReflectionUniformBuffer(GraphBuilder, View);
 		PassParameters->ReflectionCapture = View.ReflectionCaptureUniformBuffer;
