Fix anisotropy lighting with bottom normal when Substrate is enabled.

diff --git a/Engine/Shaders/Private/PostProcessGBufferHints.usf b/Engine/Shaders/Private/PostProcessGBufferHints.usf
--- a/Engine/Shaders/Private/PostProcessGBufferHints.usf
+++ b/Engine/Shaders/Private/PostProcessGBufferHints.usf
@@ -256,6 +256,13 @@ void MainCS(uint3 DispatchThreadId : SV_DispatchThreadID)
 			const float3 WorldPosition = ReconstructTranslatedWorldPositionFromDepth(UV, Data.Depth);
 			AddLineTWS(WorldPosition, WorldPosition + Data.WorldNormal, ColorRed, ColorYellow);
 		}
+
+		if (CLEAR_COAT_BOTTOM_NORMAL && Data.Anisotropy != 0)
+		{
+			const float2 UV = (ShaderPrintData.CursorCoord + 0.5f) * View.BufferSizeAndInvSize.zw;
+			const float3 WorldPosition = ReconstructTranslatedWorldPositionFromDepth(UV, Data.Depth);
+			AddLineTWS(WorldPosition, WorldPosition + Data.WorldNormal, ColorRed, ColorYellow);
+		}
 	}
 #endif
 }
diff --git a/Engine/Shaders/Private/Substrate/SubstrateEvaluation.ush b/Engine/Shaders/Private/Substrate/SubstrateEvaluation.ush
--- a/Engine/Shaders/Private/Substrate/SubstrateEvaluation.ush
+++ b/Engine/Shaders/Private/Substrate/SubstrateEvaluation.ush
@@ -444,6 +444,8 @@ void SubstrateSetupBSDFContext(inout FSubstrateBSDFContext BSDFContext, in bool
 		SphereMaxNoH(BSDFContext.Context, SphereSinAlpha, true);
 	}
 	BSDFContext.Context.NoV = saturate(max(abs(BSDFContext.Context.NoV), SUBSTRATE_EPSILON));
+	BSDFContext.SatNoV = BSDFContext.Context.NoV;
+	BSDFContext.SatNoL = saturate(BSDFContext.Context.NoL);
 }
 
 FSubstrateEvaluateResult SubstrateEvaluateBSDFCommon(FSubstrateBSDFContext BSDFContext, FShadowTerms ShadowTerms, FAreaLightIntegrateContext AreaLightContext, FSubstrateIntegrationSettings Settings, int IntegrationType)
@@ -494,7 +496,12 @@ FSubstrateEvaluateResult SubstrateEvaluateBSDFCommon(FSubstrateBSDFContext BSDFC
 					float3 BottomNormal = SubstrateGetSimpleCoatBottomNormal(Haziness.BottomNormalOct, BSDFContext.N);
 					bUseBottomNormal = true;
 
+					// Reconstruct bottom layer's tangent basis
 					BSDFContext.N = BottomNormal;
+					BSDFContext.X = BSDFContext.X;
+					BSDFContext.Y = normalize(cross(BSDFContext.N, BSDFContext.X));
+					BSDFContext.TangentBasis[0] = BSDFContext.X;
+					BSDFContext.TangentBasis[1] = BSDFContext.Y;
 					BSDFContext.TangentBasis[2] = BottomNormal;
 				}
 			}
diff --git a/Engine/Shaders/Private/Substrate/SubstrateLegacyConversion.ush b/Engine/Shaders/Private/Substrate/SubstrateLegacyConversion.ush
--- a/Engine/Shaders/Private/Substrate/SubstrateLegacyConversion.ush
+++ b/Engine/Shaders/Private/Substrate/SubstrateLegacyConversion.ush
@@ -346,7 +346,13 @@ FSubstrateData SubstrateConvertLegacyMaterialDynamic(
 	}
 	else if (ShadingModel == SHADINGMODELID_CLEAR_COAT)
 	{
+		// When using SUBSTRATE_GBUFFER_FORMAT, data are stored into fixed GBuffer layout which does not contain bit for bUseBottomNormal
+		// In this case we always force bUseBottomNormal if CLEAR_COAT_BOTTOM_NORMAL is enabled
+		#if SUBSTRATE_GBUFFER_FORMAT == 0
+		const bool bUseBottomNormal = CLEAR_COAT_BOTTOM_NORMAL;
+		#else
 		const bool bUseBottomNormal = CLEAR_COAT_BOTTOM_NORMAL && any(RawClearCoatNormal != RawNormal);
+		#endif
 
 		LegacySlab.BaseColor = BaseColor;
 		LegacySlab.Specular = Specular;
diff --git a/Engine/Shaders/Private/Substrate/SubstrateVisualizeCommon.ush b/Engine/Shaders/Private/Substrate/SubstrateVisualizeCommon.ush
--- a/Engine/Shaders/Private/Substrate/SubstrateVisualizeCommon.ush
+++ b/Engine/Shaders/Private/Substrate/SubstrateVisualizeCommon.ush
@@ -144,7 +144,8 @@ void SubstratePrintBSDF(
 				if (Haziness.HasBottomNormal())
 				{
 					const float3 BottomNormal = SubstrateGetSimpleCoatBottomNormal(Haziness.BottomNormalOct, WorldNormal);
-					AddLineTWS(WorldPosition, WorldPosition + BottomNormal * 20.0f, float4(0, 1, 0, 1), float4(0, 1, 0, 1));
+					AddLineTWS(WorldPosition, WorldPosition + BottomNormal * 10.0f, ColorYellow, ColorYellow);							// Core
+					AddLineTWS(WorldPosition + BottomNormal * 10.f, WorldPosition + BottomNormal * 12.0f, ColorPurple, ColorPurple);	// Tip				
 				}
 			#endif
 			}
