Merge IndirectAO setup for Substrate lighting code.

diff --git a/Engine/Shaders/Private/ReflectionEnvironmentPixelShader.usf b/Engine/Shaders/Private/ReflectionEnvironmentPixelShader.usf
--- a/Engine/Shaders/Private/ReflectionEnvironmentPixelShader.usf
+++ b/Engine/Shaders/Private/ReflectionEnvironmentPixelShader.usf
@@ -511,28 +511,10 @@ void ReflectionEnvironmentSkyLighting(
 			float CombinedScreenAndMaterialAO = SubstrateGetAO(SubstratePixelHeader) * AmbientOcclusion;
 			FSubstrateIrradianceAndOcclusion SubstrateIrradianceAndOcclusion = SubstrateGetIrradianceAndAO(SubstratePixelHeader);
 		#else
-
 			FSubstrateGBufferBSDF GBufferBSDF = SubstrateReadGBufferBSDF(ScreenSpaceData);
-			float CombinedScreenAndMaterialAO = 1.0; // AmbientOcclusion;
-			
 			FSubstrateBSDFContext SubstrateBSDFContext = SubstrateCreateBSDFContext(GBufferBSDF.BSDFTangentBasis, GBufferBSDF.BSDF, V, V/*L*/, false/*bHasValidL*/);
-
-			// Similar to SubstrateUnpackIrradianceAndOcclusion
-			FSubstrateIrradianceAndOcclusion SubstrateIrradianceAndOcclusion = InitIrradianceAndOcclusion();
-		#if GBUFFER_HAS_DIFFUSE_SAMPLE_OCCLUSION
-			SubstrateIrradianceAndOcclusion.DiffuseIndirectSampleOcclusion = GBuffer.DiffuseIndirectSampleOcclusion;
-			SubstrateIrradianceAndOcclusion.MaterialAO = saturate(1.0 - float(countbits(Out.DiffuseIndirectSampleOcclusion)) * rcp(float(INDIRECT_SAMPLE_COUNT)));
-			SubstrateIrradianceAndOcclusion.IndirectIrradiance = 1.f;
-		#elif ALLOW_STATIC_LIGHTING
-			SubstrateIrradianceAndOcclusion.MaterialAO = 1.f;
-			SubstrateIrradianceAndOcclusion.DiffuseIndirectSampleOcclusion = 0x0u;
-			SubstrateIrradianceAndOcclusion.IndirectIrradiance = GBuffer.IndirectIrradiance;
-		#else
-			SubstrateIrradianceAndOcclusion.MaterialAO = GBuffer.GBufferAO;
-			SubstrateIrradianceAndOcclusion.DiffuseIndirectSampleOcclusion = 0x0u;
-			SubstrateIrradianceAndOcclusion.IndirectIrradiance = 1.f;
-		#endif
-
+			float CombinedScreenAndMaterialAO = 1.0; // AmbientOcclusion;
+			FSubstrateIrradianceAndOcclusion SubstrateIrradianceAndOcclusion = SubstrateGetIrradianceAndAO(ScreenSpaceData.GBuffer);
 		#endif
 			const float3 BSDFThroughput = LuminanceWeight(SubstrateBSDFContext, SubstrateBSDFContext.BSDF); // Use the reflected direction
 	
diff --git a/Engine/Shaders/Private/Substrate/Substrate.ush b/Engine/Shaders/Private/Substrate/Substrate.ush
--- a/Engine/Shaders/Private/Substrate/Substrate.ush
+++ b/Engine/Shaders/Private/Substrate/Substrate.ush
@@ -854,6 +854,15 @@ FSubstrateIrradianceAndOcclusion SubstrateUnpackIrradianceAndOcclusion(uint In /
 	return Out;
 }
 
+FSubstrateIrradianceAndOcclusion SubstrateGetIrradianceAndAO(FGBufferData In)
+{	
+	FSubstrateIrradianceAndOcclusion Out = InitIrradianceAndOcclusion();
+	Out.MaterialAO 						= In.GBufferAO;
+	Out.DiffuseIndirectSampleOcclusion 	= In.DiffuseIndirectSampleOcclusion;
+	Out.IndirectIrradiance 				= In.IndirectIrradiance;
+	return Out;
+}
+
 ///////////////////////////////////////////////////////////////////////////////
 // Top layer
 
