Fix SSS type being incorrect with diffuse type when loading data from Format=0.

diff --git a/Engine/Shaders/Private/DeferredLightPixelShaders.usf b/Engine/Shaders/Private/DeferredLightPixelShaders.usf
--- a/Engine/Shaders/Private/DeferredLightPixelShaders.usf
+++ b/Engine/Shaders/Private/DeferredLightPixelShaders.usf
@@ -30,7 +30,6 @@
 #define SUBSTRATE_LOAD_FROM_MATERIALCONTAINER 0	// Load from a specified BSDF, not from MAterialBuffer.
 #define SUBSTRATE_INLINE_SHADING 1				// Inline shading for normal setup if needed.
 #define SUBSTRATE_SSS_MATERIAL_OVERRIDE 1		// Need for SSS dither.
-#define SUBSTRATE_OPAQUE_MATERIAL 1				// Needed for SSS to end up in alpha. 
 #include "Substrate/SubstrateRead.ush"
 #else 
 #define SUBSTRATE_LOAD_FROM_MATERIALCONTAINER 1
diff --git a/Engine/Shaders/Private/ReflectionEnvironmentPixelShader.usf b/Engine/Shaders/Private/ReflectionEnvironmentPixelShader.usf
--- a/Engine/Shaders/Private/ReflectionEnvironmentPixelShader.usf
+++ b/Engine/Shaders/Private/ReflectionEnvironmentPixelShader.usf
@@ -26,7 +26,6 @@
 #define SUBSTRATE_LOAD_FROM_MATERIALCONTAINER 0	// Load from a specified BSDF, not from MAterialBuffer.
 #define SUBSTRATE_INLINE_SHADING 1				// Inline shading for normal setup if needed.
 #define SUBSTRATE_SSS_MATERIAL_OVERRIDE 1		// Need for SSS dither.
-#define SUBSTRATE_OPAQUE_MATERIAL 1				// Needed for SSS to end up in alpha. 
 #include "Substrate/SubstrateRead.ush"
 #else 
 #define SUBSTRATE_LOAD_FROM_MATERIALCONTAINER 1
diff --git a/Engine/Shaders/Private/Substrate/Substrate.ush b/Engine/Shaders/Private/Substrate/Substrate.ush
--- a/Engine/Shaders/Private/Substrate/Substrate.ush
+++ b/Engine/Shaders/Private/Substrate/Substrate.ush
@@ -1522,7 +1522,8 @@ FSubstrateData InternalGetSubstrateSlabBSDF(
 	float ThicknessCm,
 	bool bIsThinSurface,
 	bool bIsAtTheBottomOfTopology,	// Make sure we do not affect MFP when wrap mode is selected. This is also an optimisation, disabling lots of code when we know about topology.
-	uint SharedLocalBasisIndex, inout uint SharedLocalBasisTypes)
+	uint SharedLocalBasisIndex, inout uint SharedLocalBasisTypes,
+	bool bSupportSSSDiffusion)
 {
 	// Curvature-based roughness for geometric AA
 	Roughness		= max(Roughness, InPixelFootprint.NormalCurvatureRoughness);
@@ -1538,12 +1539,10 @@ FSubstrateData InternalGetSubstrateSlabBSDF(
 
 	// In forward mode (blend mode != Opaque/Masked), diffusion is not supported. 
 	// In such a case we fallback onto no SSS (i.e., diffuse), as look-wise, this is the closest one.
-	#if SUBSTRATE_OPAQUE_MATERIAL == 0
-	if (SSSType == SSS_TYPE_DIFFUSION || SSSType == SSS_TYPE_DIFFUSION_PROFILE)
+	if (!bSupportSSSDiffusion && (SSSType == SSS_TYPE_DIFFUSION || SSSType == SSS_TYPE_DIFFUSION_PROFILE))
 	{
 		SSSType = SSS_TYPE_NONE;
 	}
-	#endif
 
 	if (!bIsAtTheBottomOfTopology)
 	{
@@ -1744,7 +1743,8 @@ FSubstrateData GetSubstrateSlabBSDF(
 		ThicknessCm,
 		bIsThinSurface,
 		bIsAtTheBottomOfTopology,
-		SharedLocalBasisIndex, SharedLocalBasisTypes);
+		SharedLocalBasisIndex, SharedLocalBasisTypes,
+		SUBSTRATE_OPAQUE_MATERIAL /*bSupportSSSDiffusion*/);
 	#else
 	return (FSubstrateData)0;
 	#endif
@@ -1784,7 +1784,8 @@ FSubstrateData GetSubstrateSlabBSDF(
 		ThicknessCm,
 		bIsThinSurface,
 		bIsAtTheBottomOfTopology,
-		SharedLocalBasisIndex, SharedLocalBasisTypes);
+		SharedLocalBasisIndex, SharedLocalBasisTypes,
+		SUBSTRATE_OPAQUE_MATERIAL /*bSupportSSSDiffusion*/);
 	#else
 	return (FSubstrateData)0;
 	#endif
diff --git a/Engine/Shaders/Private/Substrate/SubstrateRead.ush b/Engine/Shaders/Private/Substrate/SubstrateRead.ush
--- a/Engine/Shaders/Private/Substrate/SubstrateRead.ush
+++ b/Engine/Shaders/Private/Substrate/SubstrateRead.ush
@@ -171,7 +171,8 @@ FSubstrateGBufferBSDF SubstrateReadGBufferBSDF(in FScreenSpaceData ScreenSpaceDa
 				ThicknessCm,
 				bIsThinSurface,
 				bIsAtTheBottomOfTopology,
-				SharedLocalBasisIndex, SharedLocalBasisTypes).InlinedBSDF;
+				SharedLocalBasisIndex, SharedLocalBasisTypes,
+				true /*bSupportSSSDiffusion*/).InlinedBSDF;
 		}
 
 		#if SUBSTRATE_INLINE_SHADING
