Path Tracer: Add support for transmissive materials and indirect scale on lights for visualization mode

Implement russian roulette for early path termination.

Track SigmaT for proper attenuation through solid glass materials.

#rnx

diff --git a/Engine/Shaders/Private/PathTracing/PathTracingDebug.usf b/Engine/Shaders/Private/PathTracing/PathTracingDebug.usf
--- a/Engine/Shaders/Private/PathTracing/PathTracingDebug.usf
+++ b/Engine/Shaders/Private/PathTracing/PathTracingDebug.usf
@@ -98,8 +98,7 @@ float3 TraceShadowRay(FRayDesc Ray, float PathRoughness, float PreviousPathRough
 		// We didn't run the miss shader, therefore we must have hit something opaque (or reached full opacity)
 		return 0.0;
 	}
-	float3 Throughput = PackedPayload.GetRayThroughput();
-	return Throughput;
+	return PackedPayload.GetRayThroughput() * exp(-max(PackedPayload.GetTau(), 0.0));
 }
 
 struct FPathTracingResult {
@@ -208,7 +207,7 @@ int GetDominantSide(float3 N)
 
 #if PATH_TRACER_PRIMARY_RAYS
 
-FPackedPathTracingPayload TraceTransparentRay(FRayDesc Ray, int ScatterType, int Bounce, float PathRoughness, inout float3 PathThroughput, inout FRandomSequence RandSequence, inout FPathTracingResult Result, inout float3 LocalRadiance)
+FPackedPathTracingPayload TraceTransparentRay(FRayDesc Ray, int ScatterType, int Bounce, float PathRoughness, float3 SigmaT, inout float3 PathThroughput, inout FRandomSequence RandSequence, inout FPathTracingResult Result, inout float3 LocalRadiance)
 {
 	const uint RayFlags = Bounce == 0 ? RAY_FLAG_CULL_BACK_FACING_TRIANGLES : 0;
 	const uint MissShaderIndex = 0;
@@ -337,6 +336,9 @@ FPackedPathTracingPayload TraceTransparentRay(FRayDesc Ray, int ScatterType, int
 			Result.Radiance.a += Luminance(PathThroughput * (1.0 - PackedPayload.UnpackTransparencyColor()));
 		}
 
+		// Apply Beer's law
+		PathThroughput *= select(SigmaT > 0.0, exp(-SigmaT * (PackedPayload.HitT - Ray.TMin)), 1.0);
+
 		if (Bounce == 0 || EnableEmissive)
 		{
 			LocalRadiance += PathThroughput * PackedPayload.UnpackRadiance();
@@ -358,10 +360,10 @@ FPackedPathTracingPayload TraceTransparentRay(FRayDesc Ray, int ScatterType, int
 			float3 Contrib = PathThroughput * EstimateMaterialAlbedo(HitPayload);
 
 			// only allow albedo tracking through a very narrow range of roughness values
-			const float DenoiserMinRoughness = 0.01;
-			const float DenoiserMaxRoughness = 0.02;
+			const float DenoiserMinRoughness = 0.000;
+			const float DenoiserMaxRoughness = 0.004;
 			float DenoiserContrib = 1 - saturate((PathRoughness - DenoiserMinRoughness) / (DenoiserMaxRoughness - DenoiserMinRoughness));
-			if (DenoiserContrib > 0)
+			if (DenoiserContrib > 0 && Bounce == 0)
 			{
 				float3 DenoiserThroughput = PathThroughput * DenoiserContrib * HitPayload.BSDFOpacity;
 
@@ -433,6 +435,7 @@ FPathTracingResult PathTracingDebugCore(FRayDesc Ray, uint2 LaunchIndex, int Sam
 	FPathTracingResult Result = (FPathTracingResult) 0;
 	float3 PathThroughput = 1;
 	float PathRoughness = 0;
+	float3 SigmaT = 0; // TODO: initialize from pre-pass for rays starting inside glass/water
 	uint ScatterType = PATHTRACER_SCATTER_CAMERA;
 	FRandomSequence RandSequence = RandomSequenceCreate(uint3(LaunchIndex, View.StateFrameIndex), SampleIndex, SamplesPerPixel);
 
@@ -455,7 +458,7 @@ FPathTracingResult PathTracingDebugCore(FRayDesc Ray, uint2 LaunchIndex, int Sam
 	for (int Bounce = 0; Bounce <= MaxBounces; Bounce++)
 	{
 		float3 LocalRadiance = 0;
-		FPackedPathTracingPayload PackedPayload = TraceTransparentRay(Ray, ScatterType, Bounce, PathRoughness, PathThroughput, RandSequence, Result, LocalRadiance);
+		FPackedPathTracingPayload PackedPayload = TraceTransparentRay(Ray, ScatterType, Bounce, PathRoughness, SigmaT, PathThroughput, RandSequence, Result, LocalRadiance);
 
 		if (PackedPayload.IsMiss())
 		{
@@ -516,6 +519,7 @@ FPathTracingResult PathTracingDebugCore(FRayDesc Ray, uint2 LaunchIndex, int Sam
 		FRISContext TraceSample = InitRISContext(RandSample.z);
 		FRayDesc ShadowRaySample;
 		float3 ShadowRayContrib = 0;
+		float ShadowRayRoughness = 0;
 		uint ShadowRayLightId;
 #endif
 		FMaterialSample MaterialSample = SampleMaterial(-Ray.Direction, Payload, RandSample.xyz);
@@ -537,6 +541,11 @@ FPathTracingResult PathTracingDebugCore(FRayDesc Ray, uint2 LaunchIndex, int Sam
 				float2 LightDiffuseSpecularScale = float2(GetLightDiffuseScale(LightId), GetLightSpecularScale(LightId));
 				FMaterialEval MaterialEval = EvalMaterial(-Ray.Direction, LightSample.Direction, Payload, LightDiffuseSpecularScale);
 				float3 LightContrib = PathThroughput * LightSample.RadianceOverPdf * MaterialEval.Weight * MaterialEval.Pdf;
+				if (ScatterType == PATHTRACER_SCATTER_DIFFUSE || ScatterType == PATHTRACER_SCATTER_SPECULAR)
+				{
+					LightContrib *= GetLightIndirectScale(LightId, PathRoughness);
+				}
+
 #if USE_BSDF_SAMPLING
 				LightContrib *= MISWeightPower(LightSample.Pdf, MaterialEval.Pdf);
 #endif
@@ -558,7 +567,8 @@ FPathTracingResult PathTracingDebugCore(FRayDesc Ray, uint2 LaunchIndex, int Sam
 					ShadowRay.TMin = 0;
 					ShadowRay.TMax = LightSample.Distance;
 
-					ApplyRayBias(ShadowRay.Origin, PackedPayload.HitT, PackedPayload.UnpackWorldGeoNormal());
+					const float SignedPositionBias = Payload.IsMaterialTransmissive() ? sign(dot(Payload.WorldGeoNormal, LightSample.Direction)) : 1.0;
+					ApplyRayBias(ShadowRay.Origin, PackedPayload.HitT, SignedPositionBias * PackedPayload.UnpackWorldGeoNormal());
 
 #if LIGHT_LOOP_METHOD == 2
 					float SelectionWeight = LightContrib.x + LightContrib.y + LightContrib.z;
@@ -566,13 +576,14 @@ FPathTracingResult PathTracingDebugCore(FRayDesc Ray, uint2 LaunchIndex, int Sam
 					{
 						ShadowRaySample = ShadowRay;
 						ShadowRayContrib = LightContrib / SelectionWeight;
+						ShadowRayRoughness = GetAverageRoughness(Payload);
 						ShadowRayLightId = LightId;
 					}
 #else
 					// Trace shadow ray immediately
 					const bool bCastShadows = CastsShadow(LightId);
 					const uint MissShaderIndex = GetLightMissShaderIndex(LightId);
-					LocalRadiance += LightContrib * TraceShadowRay(ShadowRay, 1.0, 0.0, MissShaderIndex, bCastShadows);
+					LocalRadiance += LightContrib * TraceShadowRay(ShadowRay, GetAverageRoughness(Payload), PathRoughness, MissShaderIndex, bCastShadows);
 #endif
 				}
 			}
@@ -586,7 +597,7 @@ FPathTracingResult PathTracingDebugCore(FRayDesc Ray, uint2 LaunchIndex, int Sam
 				ShadowRay.TMin = 0;
 				ShadowRay.TMax = RAY_DEFAULT_T_MAX;
 
-				ApplyRayBias(ShadowRay.Origin, PackedPayload.HitT, PackedPayload.UnpackWorldGeoNormal());
+				ApplyRayBias(ShadowRay.Origin, PackedPayload.HitT, MaterialSample.PositionBiasSign * PackedPayload.UnpackWorldGeoNormal());
 
 				FLightHit LightResult = TraceLight(ShadowRay, LightId);
 				if (LightResult.IsHit())
@@ -594,6 +605,10 @@ FPathTracingResult PathTracingDebugCore(FRayDesc Ray, uint2 LaunchIndex, int Sam
 					ShadowRay.TMax = LightResult.HitT;
 					float3 LightContrib = PathThroughput * MaterialSample.Weight * LightResult.Radiance;
 					LightContrib *= MaterialSample.ScatterType == PATHTRACER_SCATTER_DIFFUSE ? GetLightDiffuseScale(LightId) : GetLightSpecularScale(LightId);
+					if (ScatterType == PATHTRACER_SCATTER_DIFFUSE || ScatterType == PATHTRACER_SCATTER_SPECULAR)
+					{
+						LightContrib *= GetLightIndirectScale(LightId, PathRoughness);
+					}
 					LightContrib *= MISWeightPower(MaterialSample.Pdf, LightResult.Pdf * LightPickPdf);
 #if LIGHT_LOOP_METHOD == 2
 					// Defer tracing shadow until later
@@ -602,12 +617,13 @@ FPathTracingResult PathTracingDebugCore(FRayDesc Ray, uint2 LaunchIndex, int Sam
 					{
 						ShadowRaySample = ShadowRay;
 						ShadowRayContrib = LightContrib / SelectionWeight;
+						ShadowRayRoughness = MaterialSample.Roughness;
 						ShadowRayLightId = LightId;
 					}
 #else
 					const bool bCastShadows = CastsShadow(LightId);
 					const uint MissShaderIndex = GetLightMissShaderIndex(LightId);
-					LocalRadiance += LightContrib * TraceShadowRay(ShadowRay, 1.0, 0.0, MissShaderIndex, bCastShadows);
+					LocalRadiance += LightContrib * TraceShadowRay(ShadowRay, MaterialSample.Roughness, PathRoughness, MissShaderIndex, bCastShadows);
 #endif
 				}
 			}
@@ -620,7 +636,7 @@ FPathTracingResult PathTracingDebugCore(FRayDesc Ray, uint2 LaunchIndex, int Sam
 			const bool bCastShadows = CastsShadow(ShadowRayLightId);
 			const uint MissShaderIndex = GetLightMissShaderIndex(ShadowRayLightId);
 			ShadowRayContrib *= TraceSample.GetNormalization();
-			LocalRadiance += ShadowRayContrib * TraceShadowRay(ShadowRaySample, 1.0, 0.0, MissShaderIndex, bCastShadows);
+			LocalRadiance += ShadowRayContrib * TraceShadowRay(ShadowRaySample, ShadowRayRoughness, PathRoughness, MissShaderIndex, bCastShadows);
 		}
 #endif
 
@@ -753,16 +769,13 @@ FPathTracingResult PathTracingDebugCore(FRayDesc Ray, uint2 LaunchIndex, int Sam
 
 		// Algorithm 2: Sampling and Ray Casting
 		// TODO: Adjust SurfaceProb based on roughness
-		ScatterType = PATHTRACER_SCATTER_DIFFUSE; // TODO: is there a good way to classify the scattering type for guided rays?
-		PathRoughness = 1.0;                      // TODO: return average roughness in EvalMaterial?
 		float3 Throughput = 0;
 		if (ScoresSum == 0 || RandSample.w < MaterialProb)
 		{
 			// BSDF Sampling
-			FMaterialSample MaterialSample = SampleMaterial(-Ray.Direction, Payload, RandSample.xyz);
 			Throughput = MaterialSample.Weight;
 			SampledDirection = MaterialSample.Direction;
-			ScatterType = MaterialSample.ScatterType;
+			ScatterType = ScatterType == PATHTRACER_SCATTER_CAMERA || ScatterType == PATHTRACER_SCATTER_REFRACT ? MaterialSample.ScatterType : ScatterType; // change if first bounce
 			PathRoughness = max(MaterialSample.Roughness, PathRoughness);
 			ChosenState = (FMarkovChainState)0;
 		}
@@ -780,8 +793,56 @@ FPathTracingResult PathTracingDebugCore(FRayDesc Ray, uint2 LaunchIndex, int Sam
 			}
 			FMaterialEval MaterialEval = EvalMaterial(-Ray.Direction, SampledDirection, Payload, 1.0);
 			Throughput = MaterialEval.Weight * MISWeightBalanced(MaterialProb * MaterialEval.Pdf, GuidedPdf - GuidedPdf * MaterialProb) / MaterialProb;
+			MaterialSample.PositionBiasSign = sign(dot(SampledDirection, Payload.WorldGeoNormal));
+			ScatterType = PATHTRACER_SCATTER_DIFFUSE; // TODO: is there a good way to classify the scattering type for guided rays?
+			PathRoughness = 1.0;                      // TODO: return average roughness in EvalMaterial?
+		}
+
+		// update the current extinction if we are crossing a boundary on glass or water
+		if (MaterialSample.PositionBiasSign < 0 && Payload.IsMaterialSolidGlass())
+		{
+			const float3 LocalSigmaT = Payload.GetExtinction();
+			if (Payload.IsFrontFace())
+			{
+				// entering
+				SigmaT += LocalSigmaT;
+			}
+			else
+			{
+				// exiting
+				SigmaT -= LocalSigmaT;
+				SigmaT = max(SigmaT, 0);
+			}
+		}
+
+		float3 NextPathThroughput = PathThroughput * Throughput;
+		// Russian roulette:
+		//   The probability of keeping the path should be roughly proportional to the weight at the current shade point,
+		//  but just using MaterialWeight would miss out on cases where the path throughput changes color (like in a cornell
+		//  box when bouncing between walls of different colors). So use the ratio of the brightest color channel in the
+		//  previous and next throughput.
+		//   The second tweak is to add a sqrt() around the probability to soften the termination probability (paths will last
+		//  a little longer). This allows paths to go a bit deeper than the naive heuristic while still allowing them to terminate
+		//  early. This makes RR effective from the very first bounce without needing to delay it.
+		float ContinuationProb = sqrt(saturate(max3(NextPathThroughput.x, NextPathThroughput.y, NextPathThroughput.z) / max3(PathThroughput.x, PathThroughput.y, PathThroughput.z)));
+		if (ContinuationProb < 1)
+		{
+			// If there is some chance we should terminate the ray, draw an extra random value
+			float RussianRouletteRand = RandSample.w;
+			if (RussianRouletteRand >= ContinuationProb)
+			{
+				// stochastically terminate the path
+				break;
+			}
+			PathThroughput = NextPathThroughput / ContinuationProb;
+		}
+		else
+		{
+			PathThroughput = NextPathThroughput;
 		}
 
+
+
 		// Setup next ray
 		PathThroughput *= Throughput;
 		PreviousN = PackedPayload.UnpackWorldSmoothNormal();
@@ -790,7 +851,7 @@ FPathTracingResult PathTracingDebugCore(FRayDesc Ray, uint2 LaunchIndex, int Sam
 		Ray.TMin = 0;
 		Ray.TMax = RAY_DEFAULT_T_MAX;
 
-		ApplyRayBias(Ray.Origin, PackedPayload.HitT, PackedPayload.UnpackWorldGeoNormal());
+		ApplyRayBias(Ray.Origin, PackedPayload.HitT, MaterialSample.PositionBiasSign * PackedPayload.UnpackWorldGeoNormal());
 	}
 	Result.Radiance.rgb *= View.PreExposure;
 	return Result;
