Path Tracer: Add additional visualization modes for SSS and cloth materials

#rnx

diff --git a/Engine/Shaders/Private/PathTracing/PathTracingDebug.usf b/Engine/Shaders/Private/PathTracing/PathTracingDebug.usf
--- a/Engine/Shaders/Private/PathTracing/PathTracingDebug.usf
+++ b/Engine/Shaders/Private/PathTracing/PathTracingDebug.usf
@@ -983,9 +983,11 @@ FPathTracingResult PathTracingDebugCore(FRayDesc Ray, uint2 LaunchIndex)
 		}
 
 		{
+			FPathTracingPayload HitPayload = UnpackPathTracingPayload(PackedPayload, Ray);
+			AdjustPayloadAfterUnpack(HitPayload, true);
+
 			const float3 TranslatedWorldPos = Ray.Origin + PackedPayload.HitT * Ray.Direction;
-			const float3 WorldNormal = PackedPayload.UnpackWorldNormal();
-			const uint PrimitiveLightingChannelMask = PackedPayload.GetPrimitiveLightingChannelMask();
+			const float3 WorldNormal = HitPayload.WorldNormal;
 
 			FLightLoopCount LightLoopCount = LightGridLookup(TranslatedWorldPos);
 
@@ -993,34 +995,34 @@ FPathTracingResult PathTracingDebugCore(FRayDesc Ray, uint2 LaunchIndex)
 			float3 Color = 0, ExtraAlpha = 0;
 			switch (DebugMode)
 			{
-				case PATH_TRACER_DEBUG_VIZ_RADIANCE: 					Color = PackedPayload.UnpackRadiance() * View.PreExposure; break;
+				case PATH_TRACER_DEBUG_VIZ_RADIANCE: 					Color = HitPayload.Radiance * View.PreExposure; break;
 				case PATH_TRACER_DEBUG_VIZ_WORLD_NORMAL: 				Color = WorldNormal * 0.5 + 0.5; break;
-				case PATH_TRACER_DEBUG_VIZ_WORLD_SMOOTH_NORMAL:			Color = PackedPayload.UnpackWorldSmoothNormal() * 0.5 + 0.5;
-				case PATH_TRACER_DEBUG_VIZ_WORLD_GEO_NORMAL:			Color = PackedPayload.UnpackWorldGeoNormal() * 0.5 + 0.5;
+				case PATH_TRACER_DEBUG_VIZ_WORLD_SMOOTH_NORMAL:			Color = HitPayload.WorldSmoothNormal * 0.5 + 0.5;
+				case PATH_TRACER_DEBUG_VIZ_WORLD_GEO_NORMAL:			Color = HitPayload.WorldGeoNormal * 0.5 + 0.5;
 #if SUBSTRATE_ENABLED
 				// TODO: implement more visualizations for this case
-				case PATH_TRACER_DEBUG_VIZ_BASE_COLOR:					Color = lerp(PackedPayload.UnpackDiffuseColor(), PackedPayload.UnpackSpecularColor(), F0RGBToMetallic(PackedPayload.UnpackSpecularColor())); break;
-				case PATH_TRACER_DEBUG_VIZ_DIFFUSE_COLOR:				Color = PackedPayload.UnpackDiffuseColor(); break;
-				case PATH_TRACER_DEBUG_VIZ_SPECULAR_COLOR:				Color = PackedPayload.UnpackSpecularColor(); break;
-				case PATH_TRACER_DEBUG_VIZ_METALLIC:					Color = F0RGBToMetallic(PackedPayload.UnpackSpecularColor()); break;
-				case PATH_TRACER_DEBUG_VIZ_SPECULAR:					Color = F0RGBToDielectricSpecular(PackedPayload.UnpackSpecularColor()); break;
-				case PATH_TRACER_DEBUG_VIZ_ROUGHNESS:					Color = PackedPayload.UnpackRoughnessAniso().xyz; break;
-				case PATH_TRACER_DEBUG_VIZ_ANISOTROPY:					Color = AnisoDebugColor(PackedPayload.UnpackRoughnessAniso().w); break;
+				case PATH_TRACER_DEBUG_VIZ_BASE_COLOR:					Color = lerp(HitPayload.DiffuseColor, HitPayload.SpecularColor, F0RGBToMetallic(HitPayload.SpecularColor)); break;
+				case PATH_TRACER_DEBUG_VIZ_DIFFUSE_COLOR:				Color = HitPayload.DiffuseColor; break;
+				case PATH_TRACER_DEBUG_VIZ_SPECULAR_COLOR:				Color = HitPayload.SpecularColor; break;
+				case PATH_TRACER_DEBUG_VIZ_METALLIC:					Color = F0RGBToMetallic(HitPayload.SpecularColor); break;
+				case PATH_TRACER_DEBUG_VIZ_SPECULAR:					Color = F0RGBToDielectricSpecular(HitPayload.SpecularColor); break;
+				case PATH_TRACER_DEBUG_VIZ_ROUGHNESS:					Color = HitPayload.RoughnessData; break;
+				case PATH_TRACER_DEBUG_VIZ_ANISOTROPY:					Color = AnisoDebugColor(HitPayload.Anisotropy); break;
 #else
-				case PATH_TRACER_DEBUG_VIZ_BASE_COLOR:					Color = PackedPayload.UnpackBaseColor(); break;
-				case PATH_TRACER_DEBUG_VIZ_DIFFUSE_COLOR:				Color = PackedPayload.UnpackBaseColor() * (1 - PackedPayload.UnpackMetallic()); break;
-				case PATH_TRACER_DEBUG_VIZ_SPECULAR_COLOR:				Color = lerp(0.08 * PackedPayload.UnpackSpecular(), PackedPayload.UnpackBaseColor(), PackedPayload.UnpackMetallic()); break;
-				case PATH_TRACER_DEBUG_VIZ_METALLIC:					Color = PackedPayload.UnpackMetallic(); break;
-				case PATH_TRACER_DEBUG_VIZ_SPECULAR:					Color = PackedPayload.UnpackSpecular(); break;
-				case PATH_TRACER_DEBUG_VIZ_ROUGHNESS:					Color = PackedPayload.UnpackRoughness(); break;
-				case PATH_TRACER_DEBUG_VIZ_ANISOTROPY:					Color = AnisoDebugColor(PackedPayload.UnpackAnisotropy()); break;
-				case PATH_TRACER_DEBUG_VIZ_CUSTOM_DATA0:				Color = PackedPayload.UnpackCustomData0().rgb * View.PreExposure; ExtraAlpha = PackedPayload.UnpackCustomData0().a; break;
-				case PATH_TRACER_DEBUG_VIZ_CUSTOM_DATA1:				Color = PackedPayload.UnpackCustomData1().rgb * View.PreExposure; ExtraAlpha = PackedPayload.UnpackCustomData1().a; break;
+				case PATH_TRACER_DEBUG_VIZ_BASE_COLOR:					Color = HitPayload.BaseColor; break;
+				case PATH_TRACER_DEBUG_VIZ_DIFFUSE_COLOR:				Color = HitPayload.BaseColor * (1 - HitPayload.Metallic); break;
+				case PATH_TRACER_DEBUG_VIZ_SPECULAR_COLOR:				Color = lerp(0.08 * HitPayload.Specular, HitPayload.BaseColor, HitPayload.Metallic); break;
+				case PATH_TRACER_DEBUG_VIZ_METALLIC:					Color = HitPayload.Metallic; break;
+				case PATH_TRACER_DEBUG_VIZ_SPECULAR:					Color = HitPayload.Specular; break;
+				case PATH_TRACER_DEBUG_VIZ_ROUGHNESS:					Color = HitPayload.Roughness; break;
+				case PATH_TRACER_DEBUG_VIZ_ANISOTROPY:					Color = AnisoDebugColor(HitPayload.Anisotropy); break;
+				case PATH_TRACER_DEBUG_VIZ_CUSTOM_DATA0:				Color = HitPayload.CustomData0.rgb * View.PreExposure; ExtraAlpha = HitPayload.CustomData0.a; break;
+				case PATH_TRACER_DEBUG_VIZ_CUSTOM_DATA1:				Color = HitPayload.CustomData1.rgb * View.PreExposure; ExtraAlpha = HitPayload.CustomData1.a; break;
 #endif
-				case PATH_TRACER_DEBUG_VIZ_OPACITY:						Color = PackedPayload.UnpackTransparencyColor(); break;
-				case PATH_TRACER_DEBUG_VIZ_IOR:							Color = PackedPayload.UnpackIor(); break;
-				case PATH_TRACER_DEBUG_VIZ_SHADING_MODEL:				Color = GetShadingModelColor(PackedPayload.GetShadingModelID()); break;
-				case PATH_TRACER_DEBUG_VIZ_LIGHTING_CHANNEL_MASK:		Color = GetLightingChannelMaskDebugColor(PackedPayload.GetPrimitiveLightingChannelMask()); break;
+				case PATH_TRACER_DEBUG_VIZ_OPACITY:						Color = HitPayload.TransparencyColor; break;
+				case PATH_TRACER_DEBUG_VIZ_IOR:							Color = HitPayload.Ior; break;
+				case PATH_TRACER_DEBUG_VIZ_SHADING_MODEL:				Color = GetShadingModelColor(HitPayload.ShadingModelID); break;
+				case PATH_TRACER_DEBUG_VIZ_LIGHTING_CHANNEL_MASK:		Color = GetLightingChannelMaskDebugColor(HitPayload.PrimitiveLightingChannelMask); break;
 				case PATH_TRACER_DEBUG_VIZ_WORLD_POSITION:
 				{
 					float3 AbsoluteWorldPosition = DFDemote(DFFastSubtract(TranslatedWorldPos, ResolvedView.PreViewTranslation));
@@ -1053,6 +1055,42 @@ FPathTracingResult PathTracingDebugCore(FRayDesc Ray, uint2 LaunchIndex)
 					Color = (PackedPayload.IsFrontFace() ? float3(0, 1, 0) : float3(1, 0, 0)) * FakeShading;
 					break;
 				}
+				case PATH_TRACER_DEBUG_VIZ_SSS_COLOR:
+				case PATH_TRACER_DEBUG_VIZ_SSS_RADIUS:
+				case PATH_TRACER_DEBUG_VIZ_SSS_WEIGHT:
+				case PATH_TRACER_DEBUG_VIZ_SSS_PHASE:
+				{
+					FSSSRandomWalkInfo SSS = GetMaterialSSSInfo(HitPayload, Ray.Direction);
+					if (any(SSS.Weight > 0))
+					{
+						if (DebugMode == PATH_TRACER_DEBUG_VIZ_SSS_COLOR)
+						{
+							Color = SSS.Color;
+						}
+						else if (DebugMode == PATH_TRACER_DEBUG_VIZ_SSS_RADIUS)
+						{
+							Color = SSS.Radius;
+						}
+						else if (DebugMode == PATH_TRACER_DEBUG_VIZ_SSS_WEIGHT)
+						{
+							Color = SSS.Weight;
+						}
+						else if (DebugMode == PATH_TRACER_DEBUG_VIZ_SSS_PHASE)
+						{
+							Color = BlueGreenRedColorMap(0.5 * SSS.G + 0.5);
+						}
+					}
+					break;
+				}
+#if SUBSTRATE_ENABLED
+				case PATH_TRACER_DEBUG_VIZ_FUZZ_COLOR:     	Color = (HitPayload.ShadingModelID == SHADINGMODELID_SUBSTRATE) ? HitPayload.FuzzColor     : 0.0; break;
+				case PATH_TRACER_DEBUG_VIZ_FUZZ_ROUGHNESS:	Color = (HitPayload.ShadingModelID == SHADINGMODELID_SUBSTRATE) ? HitPayload.FuzzRoughness : 0.0; break;
+				case PATH_TRACER_DEBUG_VIZ_FUZZ_AMOUNT:		Color = (HitPayload.ShadingModelID == SHADINGMODELID_SUBSTRATE) ? HitPayload.FuzzAmount    : 0.0; break;
+#else
+				case PATH_TRACER_DEBUG_VIZ_FUZZ_COLOR:     	Color = (HitPayload.ShadingModelID == SHADINGMODELID_CLOTH) ? HitPayload.GetClothColor()  : 0.0; break;
+				case PATH_TRACER_DEBUG_VIZ_FUZZ_ROUGHNESS: 	Color = (HitPayload.ShadingModelID == SHADINGMODELID_CLOTH) ? HitPayload.Roughness        : 0.0; break;
+				case PATH_TRACER_DEBUG_VIZ_FUZZ_AMOUNT:    	Color = (HitPayload.ShadingModelID == SHADINGMODELID_CLOTH) ? HitPayload.GetClothAmount() : 0.0; break;
+#endif
 				default:
 				{
 					break;
diff --git a/Engine/Shaders/Shared/PathTracingDefinitions.h b/Engine/Shaders/Shared/PathTracingDefinitions.h
--- a/Engine/Shaders/Shared/PathTracingDefinitions.h
+++ b/Engine/Shaders/Shared/PathTracingDefinitions.h
@@ -105,4 +105,11 @@
 #define PATH_TRACER_DEBUG_VIZ_DECAL_GRID_AXIS			23
 #define PATH_TRACER_DEBUG_VIZ_VOLUME_LIGHT_COUNT        24
 #define PATH_TRACER_DEBUG_VIZ_HITKIND                   25
-#define PATH_TRACER_DEBUG_VIZ_TRANSPARENCY_COUNT		26
\ No newline at end of file
+#define PATH_TRACER_DEBUG_VIZ_TRANSPARENCY_COUNT		26
+#define PATH_TRACER_DEBUG_VIZ_SSS_COLOR                 27
+#define PATH_TRACER_DEBUG_VIZ_SSS_RADIUS                28
+#define PATH_TRACER_DEBUG_VIZ_SSS_WEIGHT                29
+#define PATH_TRACER_DEBUG_VIZ_SSS_PHASE                 30
+#define PATH_TRACER_DEBUG_VIZ_FUZZ_COLOR                31
+#define PATH_TRACER_DEBUG_VIZ_FUZZ_ROUGHNESS            32
+#define PATH_TRACER_DEBUG_VIZ_FUZZ_AMOUNT               33
diff --git a/Engine/Source/Runtime/Renderer/Private/PathTracing.cpp b/Engine/Source/Runtime/Renderer/Private/PathTracing.cpp
--- a/Engine/Source/Runtime/Renderer/Private/PathTracing.cpp
+++ b/Engine/Source/Runtime/Renderer/Private/PathTracing.cpp
@@ -88,6 +88,14 @@ struct FPathTracingDebugModeInfo
 		REGISTER_VIZ_TYPE(TEXT("VolumeLightCount") 		, PATH_TRACER_DEBUG_VIZ_VOLUME_LIGHT_COUNT		);
 		REGISTER_VIZ_TYPE(TEXT("HitKind")            	, PATH_TRACER_DEBUG_VIZ_HITKIND             	);
 		REGISTER_VIZ_TYPE(TEXT("TransparencyCount")     , PATH_TRACER_DEBUG_VIZ_TRANSPARENCY_COUNT		);
+		REGISTER_VIZ_TYPE(TEXT("SSSColor")     			, PATH_TRACER_DEBUG_VIZ_SSS_COLOR				);
+		REGISTER_VIZ_TYPE(TEXT("SSSRadius")   			, PATH_TRACER_DEBUG_VIZ_SSS_RADIUS				);
+		REGISTER_VIZ_TYPE(TEXT("SSSWeight")   			, PATH_TRACER_DEBUG_VIZ_SSS_WEIGHT				);
+		REGISTER_VIZ_TYPE(TEXT("SSSPhase")   			, PATH_TRACER_DEBUG_VIZ_SSS_PHASE				);
+		REGISTER_VIZ_TYPE(TEXT("FuzzColor")   			, PATH_TRACER_DEBUG_VIZ_FUZZ_COLOR				);
+		REGISTER_VIZ_TYPE(TEXT("FuzzRoughness") 		, PATH_TRACER_DEBUG_VIZ_FUZZ_ROUGHNESS			);
+		REGISTER_VIZ_TYPE(TEXT("FuzzAmount")   			, PATH_TRACER_DEBUG_VIZ_FUZZ_AMOUNT				);
+
 #undef REGISTER_VIZ_TYPE
 	}
 
