Fix incorrect SSS profile ID handling with Format=0, causing missing transmittance.

Some part of the code where using the packed value (0..1) while other part where using the unpacked value (0..255). To avoid confusion, now all code use the un packed value (uint value). The conversion from packed to unpacked, is done when calling Substrate node creation.

#test QAGame editor dx12 Legacy/Substrate Format=0/Substrate Format=1

diff --git a/Engine/Shaders/Private/BasePassPixelShader.usf b/Engine/Shaders/Private/BasePassPixelShader.usf
--- a/Engine/Shaders/Private/BasePassPixelShader.usf
+++ b/Engine/Shaders/Private/BasePassPixelShader.usf
@@ -1910,7 +1910,7 @@ void FPixelShaderInOut_MainPS(
 			Export.Roughness,
 			Export.Anisotropy,
 			Export.SubsurfaceColor,
-			Export.SubsurfaceProfileID,
+			Export.SubsurfaceProfileIDFloat,
 			Dither,
 			Export.ShadingModelID
 			);
diff --git a/Engine/Shaders/Private/DebugProbes.usf b/Engine/Shaders/Private/DebugProbes.usf
--- a/Engine/Shaders/Private/DebugProbes.usf
+++ b/Engine/Shaders/Private/DebugProbes.usf
@@ -126,7 +126,7 @@ void WriteOutMaterial(
 		/*float3 F90							*/1.0f,
 		/*float Roughness						*/Roughness,
 		/*float Anisotropy						*/0.0f,
-		/*float SSSProfileID					*/0.0f,
+		/*uint SSSProfileID						*/0,
 		/*float3 SSSMFP							*/0.0f,
 		/*float SSSMFPScale						*/0.0f,
 		/*float SSSPhaseAniso					*/0.0f,
diff --git a/Engine/Shaders/Private/DeferredShadingCommon.ush b/Engine/Shaders/Private/DeferredShadingCommon.ush
--- a/Engine/Shaders/Private/DeferredShadingCommon.ush
+++ b/Engine/Shaders/Private/DeferredShadingCommon.ush
@@ -1096,6 +1096,11 @@ uint ExtractSubsurfaceProfileInt(FGBufferData BufferData)
 	return ExtractSubsurfaceProfileInt(BufferData.CustomData.r);
 }
 
+float ConvertSubsurfaceProfileToFloat(uint InProfileIdUint)
+{
+	return (0xFF & InProfileIdUint) * (1.f / 255.f);
+}
+
 #if SHADING_PATH_DEFERRED
 
 #if FEATURE_LEVEL >= FEATURE_LEVEL_SM5
diff --git a/Engine/Shaders/Private/PathTracing/PathTracingMaterialHitShader.usf b/Engine/Shaders/Private/PathTracing/PathTracingMaterialHitShader.usf
--- a/Engine/Shaders/Private/PathTracing/PathTracingMaterialHitShader.usf
+++ b/Engine/Shaders/Private/PathTracing/PathTracingMaterialHitShader.usf
@@ -607,7 +607,7 @@ RAY_TRACING_ENTRY_CLOSEST_HIT(PathTracingMaterialCHS,
 					}
 					else if (BSDF_GETSSSTYPE(BSDF) == SSS_TYPE_DIFFUSION_PROFILE)
 					{
-						const uint ProfileId = SubstrateSubsurfaceProfileIdTo8bits(SLAB_SSSPROFILEID(BSDF));
+						const uint ProfileId = SLAB_SSSPROFILEID(BSDF);
 						const float Scale = SLAB_SSSPROFILERADIUSSCALE(BSDF);
 						Payload.MeanFreePath = DecodeSSSProfileRadius(ProfileId, Payload.DiffuseColor, Scale);
 #if USE_SSS_PROFILE_ANISOTROPY
@@ -739,8 +739,7 @@ RAY_TRACING_ENTRY_CLOSEST_HIT(PathTracingMaterialCHS,
 
 					if (BSDF_GETSSSTYPE(BSDF) == SSS_TYPE_DIFFUSION_PROFILE)
 					{
-						const uint ProfileId = SubstrateSubsurfaceProfileIdTo8bits(EYE_SSSPROFILEID(BSDF));
-						Payload.MeanFreePath = DecodeSSSProfileRadius(ProfileId, Payload.DiffuseColor, 1.0);
+						Payload.MeanFreePath = DecodeSSSProfileRadius(EYE_SSSPROFILEID(BSDF), Payload.DiffuseColor, 1.0);
 						Payload.PhaseG = 0.0;
 					}
 
diff --git a/Engine/Shaders/Private/ShadingCommon.ush b/Engine/Shaders/Private/ShadingCommon.ush
--- a/Engine/Shaders/Private/ShadingCommon.ush
+++ b/Engine/Shaders/Private/ShadingCommon.ush
@@ -41,7 +41,7 @@
 #define IS_FIRST_PERSON_MASK			(1 << 6) // Technically not related to the SelectiveOutputMask, but IS_FIRST_PERSON_MASK shares a bit with ZERO_PRECSHADOW_MASK and consequently only works if ALLOW_STATIC_LIGHTING is 0.
 #define SKIP_VELOCITY_MASK				(1 << 7)
 
-#define SSS_PROFILE_ID_INVALID  256
+#define SSS_PROFILE_ID_INVALID  255
 #define SSS_PROFILE_ID_PERPIXEL 512
 
 // for debugging and to visualize
diff --git a/Engine/Shaders/Private/Substrate/Substrate.ush b/Engine/Shaders/Private/Substrate/Substrate.ush
--- a/Engine/Shaders/Private/Substrate/Substrate.ush
+++ b/Engine/Shaders/Private/Substrate/Substrate.ush
@@ -985,7 +985,7 @@ float GetVarianceFromRoughness(float Roughness)
 uint PackSSSProfile(float InProfileId, float InRadius, half InThickness)
 {
 	uint Out = 0; 	
-	Out |= (PackR8(InProfileId));
+	Out |= (0xFF & uint(InProfileId));
 	Out |= (PackR8(InRadius)<<8);
 	Out |= (Pack10F(InThickness)<<16);
 	return Out;
@@ -993,8 +993,8 @@ uint PackSSSProfile(float InProfileId, float InRadius, half InThickness)
 
 void UnpackSSSProfile(uint In, inout float OutProfileId, inout float OutRadius, inout half OutThickness)
 {
-	OutProfileId = UnpackR8(In);
-	OutRadius = UnpackR8(In >> 8);
+	OutProfileId = BitFieldExtractU32(In, 8, 0);
+	OutRadius = UnpackG8(In);
 	OutThickness = Unpack10F((In >> 16) & 0x3FF);
 }
 
@@ -1512,7 +1512,7 @@ FSubstrateData InternalGetSubstrateSlabBSDF(
 	float3 RawNormal,
 	float3 DiffuseAlbedo, float3 F0, float3 F90,
 	float Roughness, float Anisotropy,
-	float SSSProfileID, float3 SSSMFP, float SSSMFPScale, float SSSPhaseAnisotropy, uint SSSType,
+	uint SSSProfileID, float3 SSSMFP, float SSSMFPScale, float SSSPhaseAnisotropy, uint SSSType,
 	float3 Emissive, 
 	float SecondRoughness, float SecondRoughnessWeight, float SecondRoughnessAsSimpleClearCoat,
 	float ClearCoatUseSecondNormal, float3 ClearCoatBottomNormal,
@@ -1533,7 +1533,7 @@ FSubstrateData InternalGetSubstrateSlabBSDF(
 	FSubstrateData SubstrateData = GetInitialisedSubstrateData();
 
 	SSSMFPScale = saturate(SSSMFPScale);
-	const bool bSSSProfilePlugged = SSSProfileID > 0.0f;
+	const bool bSSSProfilePlugged = SSSProfileID > 0;
 	const bool bMFPPlugged = any(SSSMFP > 0);
 	const bool bIsThin = bIsAtTheBottomOfTopology && (bMFPPlugged && bIsThinSurface); // Enable thin only if MFP is plugged, as otherwise there will be no back-tramission
 
@@ -1564,11 +1564,9 @@ FSubstrateData InternalGetSubstrateSlabBSDF(
 	if (!SUBSTRATE_FASTPATH && SSSType == SSS_TYPE_DIFFUSION_PROFILE)
 	{
 		// Average roughness for dual specular.
-		const uint SubsurfaceProfileUInt = SubstrateSubsurfaceProfileIdTo8bits(SSSProfileID);
+		GetSubsurfaceProfileDualSpecular(SSSProfileID, Roughness, SSSMFPScale, Roughness, Haziness.Roughness, Haziness.Weight);
 
-		GetSubsurfaceProfileDualSpecular(SubsurfaceProfileUInt, Roughness, SSSMFPScale, Roughness, Haziness.Roughness, Haziness.Weight);
-
-		SSSMFP = GetSubsurfaceProfileMFPInCm(SubsurfaceProfileUInt).xyz * SSSMFPScale;
+		SSSMFP = GetSubsurfaceProfileMFPInCm(SSSProfileID).xyz * SSSMFPScale;
 	}
 	else
 	{
@@ -1667,7 +1665,7 @@ FSubstrateData InternalGetSubstrateSlabBSDF(
 
 	if (SSSType == SSS_TYPE_DIFFUSION_PROFILE)
 	{
-		SLAB_SSSPROFILEID	(SubstrateData.InlinedBSDF)		= SSSProfileID;
+		SLAB_SSSPROFILEID	(SubstrateData.InlinedBSDF)			= SSSProfileID;
 		SLAB_SSSPROFILERADIUSSCALE(SubstrateData.InlinedBSDF)	= SSSMFPScale;
 		SLAB_SSSPROFILETHICKNESSCM(SubstrateData.InlinedBSDF)	= ThicknessCm;
 	}
@@ -1715,7 +1713,7 @@ FSubstrateData GetSubstrateSlabBSDF(
 	float3 RawNormal,
 	float3 DiffuseAlbedo, float3 F0, float3 F90,
 	float Roughness, float Anisotropy,
-	float SSSProfileID, float3 SSSMFP, float SSSMFPScale, float SSSPhaseAnisotropy, uint SSSType,
+	uint SSSProfileID, float3 SSSMFP, float SSSMFPScale, float SSSPhaseAnisotropy, uint SSSType,
 	float3 Emissive,
 	float SecondRoughness, float SecondRoughnessWeight, float SecondRoughnessAsSimpleClearCoat,
 	float ClearCoatUseSecondNormal, float3 ClearCoatBottomNormal,
@@ -1756,7 +1754,7 @@ FSubstrateData GetSubstrateSlabBSDF(
 	float3 RawNormal,
 	float3 DiffuseAlbedo, float3 F0, float3 F90,
 	float Roughness, float Anisotropy,
-	float SSSProfileID, float3 SSSMFP, float SSSMFPScale, float SSSPhaseAnisotropy, uint SSSType,
+	uint SSSProfileID, float3 SSSMFP, float SSSMFPScale, float SSSPhaseAnisotropy, uint SSSType,
 	float3 Emissive,
 	float SecondRoughness, float SecondRoughnessWeight, float SecondRoughnessAsSimpleClearCoat,
 	float ClearCoatUseSecondNormal, float3 ClearCoatBottomNormal,
@@ -1866,7 +1864,7 @@ FSubstrateData GetSubstrateHairBSDF(float3 BaseColor, float Scatter, float Specu
 	return SubstrateData;
 }
 
-FSubstrateData GetSubstrateEyeBSDF(float3 DiffuseAlbedo, float Roughness, float IrisMask, float IrisDistance, float3 IrisNormal, float3 IrisPlaneNormal, float SSSProfileId, float3 Emissive, uint SharedLocalBasisIndex)
+FSubstrateData GetSubstrateEyeBSDF(float3 DiffuseAlbedo, float Roughness, float IrisMask, float IrisDistance, float3 IrisNormal, float3 IrisPlaneNormal, uint SSSProfileId, float3 Emissive, uint SharedLocalBasisIndex)
 {
 	FSubstrateData SubstrateData = GetInitialisedSubstrateData();
 	BSDF_SETTYPE					(SubstrateData.InlinedBSDF, SUBSTRATE_BSDF_TYPE_EYE);
@@ -2500,11 +2498,11 @@ void FSubstratePixelHeader::SubstrateConvertFromDBuffer(
 		1.0f,				// F90
 		In.Roughness,		// Roughness
 		0.f,				// Anisotropy
-		0.f,				// SSSProfileID
+		0,					// SSSProfileID
 		0.f,				// SSSMFP
 		0.f,				// SSSMFPScale
 		0.f,				// SSSPhaseAniso
-		SSS_TYPE_NONE,	// UseSSSDiffusion
+		SSS_TYPE_NONE,		// UseSSSDiffusion
 		0.f,				// Emissive
 		0.f,				// SecondRoughness
 		0.f,				// SecondRoughnessWeight
@@ -3685,7 +3683,7 @@ uint SubstrateGetBSDFSubSurfaceProfileID(in FSubstrateBSDF BSDF)
 	{
 		if (BSDF_GETSSSTYPE(BSDF) == SSS_TYPE_DIFFUSION_PROFILE)
 		{
-			return SubstrateSubsurfaceProfileIdTo8bits(SLAB_SSSPROFILEID(BSDF));
+			return SLAB_SSSPROFILEID(BSDF);
 		}
 		break;
 	}
@@ -4567,7 +4565,7 @@ FSubstrateBSDF UnpackSubstrateBSDFIn(FSubstrateMaterialContainer SubstrateBuffer
 				SLAB_F0(OutBSDF)			= OutF0;
 				SLAB_ROUGHNESS(OutBSDF)		= SubstrateUnpackTopLayerData(SubstrateHeader.PackedTopLayerData).Roughness;
 
-				SLAB_SSSPROFILEID(OutBSDF)			= UnpackR8(Data2);
+				SLAB_SSSPROFILEID(OutBSDF)			= BitFieldExtractU32(Data2, 8, 0);
 				SLAB_SSSPROFILERADIUSSCALE(OutBSDF)	= UnpackG8(Data2);
 				SLAB_HAZINESS(OutBSDF)				= BitFieldExtractU32(Data2, 16, 16);
 
diff --git a/Engine/Shaders/Private/Substrate/SubstrateEvaluation.ush b/Engine/Shaders/Private/Substrate/SubstrateEvaluation.ush
--- a/Engine/Shaders/Private/Substrate/SubstrateEvaluation.ush
+++ b/Engine/Shaders/Private/Substrate/SubstrateEvaluation.ush
@@ -943,7 +943,7 @@ FSubstrateEvaluateResult SubstrateEvaluateBSDFCommon(FSubstrateBSDFContext BSDFC
 				{
 					ThicknessInCm = BSDF_GETISTHIN(BSDFContext.BSDF) ? SLAB_SSSPROFILETHICKNESSCM(BSDFContext.BSDF) : ThicknessInCm;
 
-					const uint ProfileId = SubstrateSubsurfaceProfileIdTo8bits(SLAB_SSSPROFILEID(BSDFContext.BSDF)); // TODO move this into the PackSubstrateOut( function, to avoid this decode here
+					const uint ProfileId = SLAB_SSSPROFILEID(BSDFContext.BSDF);
 					const FTransmissionProfileParams TransmissionParams = GetTransmissionProfileParams(ProfileId);
 					TransmissionThroughput	= GetTransmissionProfile(ProfileId, ThicknessInCm).rgb;
 
diff --git a/Engine/Shaders/Private/Substrate/SubstrateExport.ush b/Engine/Shaders/Private/Substrate/SubstrateExport.ush
--- a/Engine/Shaders/Private/Substrate/SubstrateExport.ush
+++ b/Engine/Shaders/Private/Substrate/SubstrateExport.ush
@@ -142,7 +142,7 @@ void PackSubstrateOut(
 	// The SSSData we accumulate to allow blending of Wrap and SubstrateDiffusion types.
 	uint SSSDataType = SSS_TYPE_NONE;
 	float SSSDataAniso = 0.0f;
-	float SSSDataProfileID = 0.0f;
+	uint  SSSDataProfileID = 0;
 	float SSSDataRadiusScale = 0.0f;
 	float3 SSSDataMFP = 0.0f;
 	float3 SSSDataBaseColor = 0.0f;
@@ -614,7 +614,7 @@ void PackSubstrateOut(
 						{
 							uint Data0 = 0; 
 							uint Data1 = 0;
-							uint Data2 = PackRGBA8(float4(SLAB_SSSPROFILEID(BSDF), SLAB_SSSPROFILERADIUSSCALE(BSDF), 0, 0)) | ((SLAB_HAZINESS(BSDF) & 0xFFFF)<<16);
+							uint Data2 = (PackR8(SLAB_SSSPROFILERADIUSSCALE(BSDF))<<8) | (uint(SLAB_SSSPROFILEID(BSDF)) & 0xFF) | ((SLAB_HAZINESS(BSDF) & 0xFFFF)<<16);
 							PackDiffuseAlbedoF0(Dither, SLAB_DIFFUSEALBEDO(BSDF), SLAB_F0(BSDF), Data0, Data1);
 
 							SUBSTRATE_STORE_UINT1(Out | Data0);
@@ -780,7 +780,7 @@ void PackSubstrateOut(
 								{
 									// Simple volume overrides SSS profile because it takes over for when not at the bottom of the BSDF layer, or during forward rendering.
 									// Note: Since it is a profile, we can't rescale its MFP value. So we store the 'thin' thickness value for computing the correct transmission during evaluation
-									SUBSTRATE_STORE_UINT1(PackSSSProfile(SLAB_SSSPROFILEID(BSDF), SLAB_SSSPROFILERADIUSSCALE(BSDF), BSDF_GETTHICKNESSCM(BSDF)));
+									SUBSTRATE_STORE_UINT1(PackSSSProfile(	SLAB_SSSPROFILEID(BSDF), SLAB_SSSPROFILERADIUSSCALE(BSDF), BSDF_GETTHICKNESSCM(BSDF)));
 								}
 								else
 								{
@@ -891,7 +891,7 @@ void PackSubstrateOut(
 				}
 				else if (SSSDataType == SSS_TYPE_DIFFUSION_PROFILE)
 				{
-					SubstrateSubSurfaceHeaderSetProfile(SSSData.Header, SSSDataRadiusScale, SubstrateSubsurfaceProfileIdTo8bits(SSSDataProfileID));
+					SubstrateSubSurfaceHeaderSetProfile(SSSData.Header, SSSDataRadiusScale, SSSDataProfileID);
 				}
 				else
 				{
@@ -928,7 +928,7 @@ struct FExportResult
 	float  Metallic;
 	float4 CustomData;
 	float3 SubsurfaceColor;
-	uint   SubsurfaceProfileID;
+	float  SubsurfaceProfileIDFloat;
 	float  SubsurfaceProfileRadiusScale;
 	float3 WorldTangent;
 	float  ShadingModelID;
@@ -984,7 +984,7 @@ FExportResult SubstrateMaterialExportOut(
 				Out.Metallic				= SubstrateGetBSDFMetallic(CurrentBSDF);
 
 				Out.SubsurfaceColor			= SubstrateGetBSDFSubSurfaceColor(CurrentBSDF);
-				Out.SubsurfaceProfileID		= SubstrateGetBSDFSubSurfaceProfileID(CurrentBSDF);
+				Out.SubsurfaceProfileIDFloat= ConvertSubsurfaceProfileToFloat(SubstrateGetBSDFSubSurfaceProfileID(CurrentBSDF));
 				Out.WorldTangent			= T;
 				Out.ShadingModelID			= SubstrateGetLegacyShadingModels(CurrentBSDF);
 
@@ -996,7 +996,7 @@ FExportResult SubstrateMaterialExportOut(
 				else if (SSSType == SSS_TYPE_DIFFUSION_PROFILE)
 				{
 					Out.SubsurfaceProfileRadiusScale = SLAB_SSSPROFILERADIUSSCALE(CurrentBSDF);
-					Out.CustomData.rgb = EncodeSubsurfaceProfile(SLAB_SSSPROFILEID(CurrentBSDF));
+					Out.CustomData.rgb = EncodeSubsurfaceProfile(Out.SubsurfaceProfileIDFloat);
 					Out.CustomData.a   = Out.SubsurfaceProfileRadiusScale;
 				}
 				else if (SSSType == SSS_TYPE_WRAP || SSSType == SSS_TYPE_TWO_SIDED_WRAP)
@@ -1033,7 +1033,7 @@ FExportResult SubstrateMaterialExportOut(
 					const float IrisMask		= saturate(EYE_IRISMASK(CurrentBSDF));
 					const float IrisDistance	= saturate(EYE_IRISDISTANCE(CurrentBSDF));
 
-					Out.CustomData.x = EncodeSubsurfaceProfile(Out.SubsurfaceProfileID).x;
+					Out.CustomData.x = EncodeSubsurfaceProfile(Out.SubsurfaceProfileIDFloat).x;
 					Out.CustomData.w = 1.0f - IrisMask;	// Opacity
 
 					const float3 GBufferWorldNormal = N;
diff --git a/Engine/Shaders/Private/Substrate/SubstrateLegacyConversion.ush b/Engine/Shaders/Private/Substrate/SubstrateLegacyConversion.ush
--- a/Engine/Shaders/Private/Substrate/SubstrateLegacyConversion.ush
+++ b/Engine/Shaders/Private/Substrate/SubstrateLegacyConversion.ush
@@ -38,7 +38,7 @@ struct FSubstrateLegacyParameters
 	float  Metallic;
 	float  Roughness;
 	float  Anisotropy;
-	float  SSSProfileID;
+	uint   SSSProfileID;
 	float3 SSSMFP;
 	float  SSSMFPScale;
 	float  SSSPhaseAniso;
@@ -75,7 +75,7 @@ FSubstrateLegacyParameters InitSubstrateLegacyParameters(FSubstratePixelFootprin
 	Out.Metallic = 0.f;
 	Out.Roughness = 0.5f;
 	Out.Anisotropy = 0.f;
-	Out.SSSProfileID = 0.f;
+	Out.SSSProfileID = 0;
 	Out.SSSMFP = Zeros;
 	Out.SSSMFPScale = 1.f;
 	Out.SSSPhaseAniso = 0.f;
@@ -150,7 +150,7 @@ FSubstrateData SubstrateConvertLegacyMaterialDynamic(
 	FSubstratePixelFootprint PixelFootprint,
 	float3 BaseColor, float Specular, float Metallic,
 	float Roughness, float Anisotropy,
-	float3 SubSurfaceColor, float SubSurfaceProfileId,
+	float3 SubSurfaceColor, uint SubSurfaceProfileId,
 	float ClearCoat, float ClearCoatRoughness,
 	float3 Emissive,
 	float Opacity,
@@ -310,7 +310,7 @@ FSubstrateData SubstrateConvertLegacyMaterialDynamic(
 		LegacySlab.Roughness = Roughness;
 		LegacySlab.Anisotropy = Anisotropy;
 		LegacySlab.SSSMFP = max(1e-05f, MFPInCm); // Ensure the MFP is not null to force the material as have SSS
-		LegacySlab.SSSProfileID = 0.f;
+		LegacySlab.SSSProfileID = 0;
 		LegacySlab.SSSMFPScale = OneMinusOpacity;
 		LegacySlab.SSSPhaseAniso = 0.93f;
 		LegacySlab.SSSType = SSS_TYPE_DIFFUSION;
@@ -468,7 +468,7 @@ FSubstrateData SubstrateConvertLegacyMaterialStatic(
 	FSubstratePixelFootprint PixelFootprint,
 	float3 BaseColor, float Specular, float Metallic,
 	float Roughness, float Anisotropy,
-	float3 SubSurfaceColor, float SubSurfaceProfileId,
+	float3 SubSurfaceColor, uint SubSurfaceProfileId,
 	float ClearCoat, float ClearCoatRoughness,
 	float3 Emissive,
 	float Opacity,
diff --git a/Engine/Shaders/Private/Substrate/SubstrateMaterialClassification.usf b/Engine/Shaders/Private/Substrate/SubstrateMaterialClassification.usf
--- a/Engine/Shaders/Private/Substrate/SubstrateMaterialClassification.usf
+++ b/Engine/Shaders/Private/Substrate/SubstrateMaterialClassification.usf
@@ -207,16 +207,15 @@ void TileMainCS(uint2 DispatchThreadId : SV_DispatchThreadID, uint LinearIndex :
 				const uint Data1 = MaterialTextureArrayUAV[uint3(PixelCoord, 1)];
 				const uint Data2 = MaterialTextureArrayUAV[uint3(PixelCoord, 2)];
 
-				const float4 CustomData = UnpackRGBA8(Data2);
-				float ProfileId = CustomData.x;
-				float RadiusScale = CustomData.y;
+				uint ProfileId = Data2 & 0xFF;
+				float RadiusScale = UnpackG8(Data2);
 
 				const uint PackedDiffuse20Bits = (Data1 & 0xFFFFF);
 				const float3 BaseColor = UnpackR7G7B6Gamma2(PackedDiffuse20Bits);
 
 				FSubstrateSubsurfaceHeader SSSHeader = (FSubstrateSubsurfaceHeader)0;
 				SubstrateSubSurfaceHeaderSetSSSType(SSSHeader, SSS_TYPE_DIFFUSION_PROFILE);
-				SubstrateSubSurfaceHeaderSetProfile(SSSHeader, RadiusScale, SubstrateSubsurfaceProfileIdTo8bits(ProfileId));
+				SubstrateSubSurfaceHeaderSetProfile(SSSHeader, RadiusScale, ProfileId);
 
 				FSubstrateSubsurfaceExtras SSSExtras = (FSubstrateSubsurfaceExtras)0;
 				SubstrateSubsurfaceExtrasSetBaseColor(SSSExtras, BaseColor);
diff --git a/Engine/Shaders/Private/Substrate/SubstrateRead.ush b/Engine/Shaders/Private/Substrate/SubstrateRead.ush
--- a/Engine/Shaders/Private/Substrate/SubstrateRead.ush
+++ b/Engine/Shaders/Private/Substrate/SubstrateRead.ush
@@ -64,7 +64,7 @@ FSubstrateGBufferBSDF SubstrateReadGBufferBSDF(in FScreenSpaceData ScreenSpaceDa
 		float3		FuzzColor = 0.0f;
 		float		FuzzRoughness = 0.0f;
 
-		float		SSSProfileID = SSS_PROFILE_ID_INVALID;
+		uint		SSSProfileID = SSS_PROFILE_ID_INVALID;
 		uint		SSSType = SSS_TYPE_NONE;
 		float		ThicknessCm = 1.0f;
 		bool		bIsThinSurface = false;
@@ -89,7 +89,7 @@ FSubstrateGBufferBSDF SubstrateReadGBufferBSDF(in FScreenSpaceData ScreenSpaceDa
 		}
 		else if (GBuffer.ShadingModelID == SHADINGMODELID_EYE)
 		{
-			SSSProfileID = ExtractSubsurfaceProfileInt(GBuffer.CustomData.x);
+			SSSProfileID = ExtractSubsurfaceProfileInt(GBuffer);
 			const float IrisMask = 1.0f - GBuffer.CustomData.w;
 			float IrisDistance = 0.0;
 			float3 IrisNormal = GBuffer.WorldNormal;
@@ -117,7 +117,7 @@ FSubstrateGBufferBSDF SubstrateReadGBufferBSDF(in FScreenSpaceData ScreenSpaceDa
 			{
 				SSSType = SSS_TYPE_DIFFUSION;	//  SSS_TYPE_DIFFUSION_PROFILE will be selected in InternalGetSubstrateSlabBSDF based on SSSProfileID and SSSMFPScale.
 				SSSMFP = 0;
-				SSSProfileID = ExtractSubsurfaceProfileInt(GBuffer.CustomData.x);
+				SSSProfileID = ExtractSubsurfaceProfileInt(GBuffer);
 				SSSMFPScale = GBuffer.CustomData.a;
 				SSSPhaseAnisotropy = 0.93f;
 			}
diff --git a/Engine/Shaders/Private/Substrate/SubstrateSubsurface.ush b/Engine/Shaders/Private/Substrate/SubstrateSubsurface.ush
--- a/Engine/Shaders/Private/Substrate/SubstrateSubsurface.ush
+++ b/Engine/Shaders/Private/Substrate/SubstrateSubsurface.ush
@@ -203,19 +203,12 @@ void SubstrateStoreSubsurfaceHeader(RWTexture2DArray<uint> SubstrateBuffer, uint
 {
 	WriteDataToBuffer(SubstrateBuffer, HeaderBytes, uint3(PixelPos, FirstSliceStoringSubstrateSSSData + 0), QuadPixelWriteMask);
 }
+
 void SubstrateStoreSubsurfaceExtras(RWTexture2DArray<uint> SubstrateBuffer, uint FirstSliceStoringSubstrateSSSData, uint2 PixelPos, uint ExtraBytes, uint QuadPixelWriteMask=1)
 {
 	WriteDataToBuffer(SubstrateBuffer, ExtraBytes, uint3(PixelPos, FirstSliceStoringSubstrateSSSData + 1), QuadPixelWriteMask);
 }
 
-uint SubstrateSubsurfaceProfileIdTo8bits(float In)
-{
-	// Similar encoding than ExtractSubsurfaceProfileInt. Valid profile ID start at 1.
-	return uint(In * 255.0f + 0.5f);
-}
-
-
-
 ///////////////////////////////////////////////////////////////////////////////
 // Shadow and transmission 
 
diff --git a/Engine/Source/Runtime/Core/Private/UObject/FortniteMainBranchObjectVersion.cpp b/Engine/Source/Runtime/Core/Private/UObject/FortniteMainBranchObjectVersion.cpp
--- a/Engine/Source/Runtime/Core/Private/UObject/FortniteMainBranchObjectVersion.cpp
+++ b/Engine/Source/Runtime/Core/Private/UObject/FortniteMainBranchObjectVersion.cpp
@@ -14,6 +14,6 @@ TMap<FGuid, FGuid> FFortniteMainBranchObjectVersion::GetSystemGuids()
 	SystemGuids.Add(DevGuids.Niagara_LatestScriptCompileVersion, FGuid("97F0FD7149D90D42B65D6557AF477D1A"));
 	SystemGuids.Add(DevGuids.SkeletalMeshDerivedDataVersion, FGuid("C998F2F3B2884A02B7E290B8C25835F0"));
 	SystemGuids.Add(DevGuids.STATICMESH_DERIVEDDATA_VER, FGuid("3ABF32149F9D4D838750368AB95573BE"));
-	SystemGuids.Add(DevGuids.MaterialTranslationDDCVersion, FGuid("5A3429F7839C4DCF88F2CA3679C0D10E"));
+	SystemGuids.Add(DevGuids.MaterialTranslationDDCVersion, FGuid("352803E079F54E7AB52040F59068EF4E"));
 	return SystemGuids;
 }
diff --git a/Engine/Source/Runtime/Engine/Private/Materials/HLSLMaterialTranslator.cpp b/Engine/Source/Runtime/Engine/Private/Materials/HLSLMaterialTranslator.cpp
--- a/Engine/Source/Runtime/Engine/Private/Materials/HLSLMaterialTranslator.cpp
+++ b/Engine/Source/Runtime/Engine/Private/Materials/HLSLMaterialTranslator.cpp
@@ -14195,7 +14195,7 @@ int32 FHLSLMaterialTranslator::SubstrateSlabBSDF(
 				/*F90*/                              %s,\n\
 				/*Roughness*/                        %s,\n\
 				/*Anisotropy*/                       %s,\n\
-				/*SSSProfileId*/                     %s,\n\
+				/*SSSProfileId*/                     ExtractSubsurfaceProfileInt(%s),\n\
 				/*SSSMFP*/                           %s,\n\
 				/*SSSMFPScale*/                      %s,\n\
 				/*SSSPhaseAniso*/                    %s,\n\
@@ -14265,7 +14265,7 @@ int32 FHLSLMaterialTranslator::SubstrateSlabBSDF(
 			/*F90*/                              %s,\n\
 			/*Roughness*/                        %s,\n\
 			/*Anisotropy*/                       %s,\n\
-			/*SSSProfileId*/                     %s,\n\
+			/*SSSProfileId*/                     ExtractSubsurfaceProfileInt(%s),\n\
 			/*SSSMFP*/                           %s,\n\
 			/*SSSMFPScale*/                      %s,\n\
 			/*SSSPhaseAniso*/                    %s,\n\
@@ -14379,7 +14379,7 @@ int32 FHLSLMaterialTranslator::SubstrateConversionFromLegacy(
 				/*Roughness*/                         %s,\n\
 				/*Anisotropy*/                        %s,\n\
 				/*SubSurfaceColor*/                   %s,\n\
-				/*SubSurfaceProfileId*/               %s,\n\
+				/*SubSurfaceProfileId*/               ExtractSubsurfaceProfileInt(%s),\n\
 				/*ClearCoat*/                         %s,\n\
 				/*ClearCoatRoughness*/                %s,\n\
 				/*EmissiveColor*/                     %s,\n\
@@ -14448,7 +14448,7 @@ int32 FHLSLMaterialTranslator::SubstrateConversionFromLegacy(
 			/*Roughness*/                         %s,\n\
 			/*Anisotropy*/                        %s,\n\
 			/*SubSurfaceColor*/                   %s,\n\
-			/*SubSurfaceProfileId*/               %s,\n\
+			/*SubSurfaceProfileId*/               ExtractSubsurfaceProfileInt(%s),\n\
 			/*ClearCoat*/                         %s,\n\
 			/*ClearCoatRoughness*/                %s,\n\
 			/*EmissiveColor*/                     %s,\n\
@@ -14612,7 +14612,7 @@ int32 FHLSLMaterialTranslator::SubstrateHairBSDF(int32 BaseColor, int32 Scatter,
 int32 FHLSLMaterialTranslator::SubstrateEyeBSDF(int32 DiffuseAlbedo, int32 Roughness, int32 IrisMask, int32 IrisDistance, int32 IrisNormal, int32 IrisPlaneNormal, int32 SSSProfileId, int32 EmissiveColor, int32 CorneaNormal, const FString& SharedLocalBasisIndexMacro, FSubstrateOperator* PromoteToOperator)
 {
 	return AddCodeChunk(
-		MCT_Substrate, TEXT("Parameters.%s.PromoteParameterBlendedBSDFToOperator(GetSubstrateEyeBSDF(%s, %s, %s, %s, %s, %s, %s, %s, %s), %u, %u, %u, %u) /* Cornea:%s Iris:%s */"),
+		MCT_Substrate, TEXT("Parameters.%s.PromoteParameterBlendedBSDFToOperator(GetSubstrateEyeBSDF(%s, %s, %s, %s, %s, %s, ExtractSubsurfaceProfileInt(%s), %s, %s), %u, %u, %u, %u) /* Cornea:%s Iris:%s */"),
 		*GetParametersSubstrateTreeName(CurrentSubstrateCompilationContext),
 		*SubstrateGetCastParameterCode(DiffuseAlbedo, MCT_Float3),
 		*SubstrateGetCastParameterCode(Roughness, MCT_Float),
