Substrate - fixed missing static shadow when using forward shading.

#tests PC editor dx12
#rnx

diff --git a/Engine/Shaders/Private/Substrate/SubstrateForwardLighting.ush b/Engine/Shaders/Private/Substrate/SubstrateForwardLighting.ush
--- a/Engine/Shaders/Private/Substrate/SubstrateForwardLighting.ush
+++ b/Engine/Shaders/Private/Substrate/SubstrateForwardLighting.ush
@@ -74,7 +74,7 @@ float3 SubstrateForwardLighting(
 
 	float4 DynamicShadowFactors = 1;
 #if MATERIALBLENDING_SOLID || MATERIALBLENDING_MASKED
-	DynamicShadowFactors = GetForwardDynamicShadowFactors(ScreenUV);
+	DynamicShadowFactors = GetForwardDynamicShadowFactors(ScreenUV); // does that match between sub0/1?
 #endif
 
 	float SpecularScale = 1;
@@ -433,6 +433,12 @@ float3 SubstrateForwardLighting(
 
 					SubstrateBSDFContext.SubstrateUpdateBSDFContext(ToDirLight);
 
+					bool bIsDirLight = true;
+					// Get base precompute shadow factors.
+					float4 PrecomputedShadowFactor = GetPrecomputedShadowFactors(SubstratePixelHeader, TranslatedWorldPosition, bIsDirLight);
+					// Mix in the static light map shadow factors.
+					PrecomputedShadowFactor *= GetPrecomputedShadowMasks(LightmapVTPageTableResult, Interpolants, MaterialParameters, VolumetricLightmapBrickTextureUVs);
+
 					FSubstrateEvaluateResult BSDFEvaluate = (FSubstrateEvaluateResult)0;
 					Color += SubstrateForwardLightingCommon(
 						Dither,
@@ -444,7 +450,7 @@ float3 SubstrateForwardLighting(
 						InitRectTexture(),
 						DirLightingChannelMask,
 						PrimitiveLightingChannelMask,
-						GetPrecomputedShadowFactors(SubstratePixelHeader, TranslatedWorldPosition),
+						PrecomputedShadowFactor,
 						TranslatedWorldPosition,
 						SceneDepth,
 						LuminanceWeight(SubstrateBSDFContext, CurrentBSDF),
@@ -554,6 +560,12 @@ float3 SubstrateForwardLighting(
 					// Update the Substrate BSDF context accordgin to the new L
 					SubstrateBSDFContext.SubstrateUpdateBSDFContext(ToLight);
 
+					bool bIsDirLight = false;
+					// Get base precompute shadow factors.
+					float4 PrecomputedShadowFactor = GetPrecomputedShadowFactors(SubstratePixelHeader, TranslatedWorldPosition, bIsDirLight);
+					// Mix in the static light map shadow factors.
+					PrecomputedShadowFactor *= GetPrecomputedShadowMasks(LightmapVTPageTableResult, Interpolants, MaterialParameters, VolumetricLightmapBrickTextureUVs);
+
 					FSubstrateEvaluateResult BSDFEvaluate = (FSubstrateEvaluateResult)0;
 					Color += SubstrateForwardLightingCommon(
 						Dither,
@@ -565,7 +577,7 @@ float3 SubstrateForwardLighting(
 						LocalRectTexture,
 						LocalLightChannelMask,
 						PrimitiveLightingChannelMask,
-						GetPrecomputedShadowFactors(SubstratePixelHeader, TranslatedWorldPosition),
+						PrecomputedShadowFactor,
 						TranslatedWorldPosition,
 						SceneDepth,
 						LuminanceWeight(SubstrateBSDFContext, CurrentBSDF),
diff --git a/Engine/Shaders/Private/Substrate/SubstrateLightingCommon.ush b/Engine/Shaders/Private/Substrate/SubstrateLightingCommon.ush
--- a/Engine/Shaders/Private/Substrate/SubstrateLightingCommon.ush
+++ b/Engine/Shaders/Private/Substrate/SubstrateLightingCommon.ush
@@ -15,15 +15,21 @@
 #endif
 
 // Used by non-mobile path
-half4 GetPrecomputedShadowFactors(FSubstratePixelHeader SubstratePixelHeader, float3 TranslatedWorldPosition)
+half4 GetPrecomputedShadowFactors(FSubstratePixelHeader SubstratePixelHeader, float3 TranslatedWorldPosition, bool bIsDirLight) 
 {
 	// PrecomputedShadowFactors are not supported on mobile at the moment
 	float4 OutPrecomputedShadowFactors = SubstratePixelHeader.HasZeroPrecShadowMask() ? 0.0f : 1.0f;
 #if TRANSLUCENCY_LIGHTING_SURFACE_FORWARDSHADING && !SHADING_PATH_MOBILE
-	OutPrecomputedShadowFactors.x = ComputeDirectionalLightStaticShadowing(TranslatedWorldPosition).x; // Is that actually correct for forward & deferred?
+	if (bIsDirLight)
+	{
+		OutPrecomputedShadowFactors.x = ComputeDirectionalLightStaticShadowing(TranslatedWorldPosition).x; // Is that actually correct for forward & deferred?
+
+	}
+	else
 #else
-	OutPrecomputedShadowFactors.x = 1;
+		OutPrecomputedShadowFactors.x = 1;
 #endif
+
 	return half4(OutPrecomputedShadowFactors);
 }
 
