Substrate - optimisation for 44624571

#tests PC editor dx12
#rnx

diff --git a/Engine/Shaders/Private/Substrate/SubstrateForwardLighting.ush b/Engine/Shaders/Private/Substrate/SubstrateForwardLighting.ush
--- a/Engine/Shaders/Private/Substrate/SubstrateForwardLighting.ush
+++ b/Engine/Shaders/Private/Substrate/SubstrateForwardLighting.ush
@@ -433,7 +433,7 @@ float3 SubstrateForwardLighting(
 
 					SubstrateBSDFContext.SubstrateUpdateBSDFContext(ToDirLight);
 
-					bool bIsDirLight = true;
+					const bool bIsDirLight = true;
 					// Get base precompute shadow factors.
 					float4 PrecomputedShadowFactor = GetPrecomputedShadowFactors(SubstratePixelHeader, TranslatedWorldPosition, bIsDirLight);
 					// Mix in the static light map shadow factors.
@@ -471,6 +471,12 @@ float3 SubstrateForwardLighting(
 				// This prevents GPU hangs when the PS tries to read from uninitialized NumCulledLightsGrid buffer
 				const uint NumLightsInGridCell = min(CulledLightsGridHeader.NumLights, GetMaxLightsPerCell());
 
+				const bool bIsDirLight = false;
+				// Get base precompute shadow factors.
+				float4 PrecomputedShadowFactor = GetPrecomputedShadowFactors(SubstratePixelHeader, TranslatedWorldPosition, bIsDirLight);
+				// Mix in the static light map shadow factors.
+				PrecomputedShadowFactor *= GetPrecomputedShadowMasks(LightmapVTPageTableResult, Interpolants, MaterialParameters, VolumetricLightmapBrickTextureUVs);
+
 				LOOP
 				for (uint GridLightListIndex = 0; GridLightListIndex < NumLightsInGridCell; GridLightListIndex++)
 				{
@@ -560,12 +566,6 @@ float3 SubstrateForwardLighting(
 					// Update the Substrate BSDF context accordgin to the new L
 					SubstrateBSDFContext.SubstrateUpdateBSDFContext(ToLight);
 
-					bool bIsDirLight = false;
-					// Get base precompute shadow factors.
-					float4 PrecomputedShadowFactor = GetPrecomputedShadowFactors(SubstratePixelHeader, TranslatedWorldPosition, bIsDirLight);
-					// Mix in the static light map shadow factors.
-					PrecomputedShadowFactor *= GetPrecomputedShadowMasks(LightmapVTPageTableResult, Interpolants, MaterialParameters, VolumetricLightmapBrickTextureUVs);
-
 					FSubstrateEvaluateResult BSDFEvaluate = (FSubstrateEvaluateResult)0;
 					Color += SubstrateForwardLightingCommon(
 						Dither,
