RuntimeVirtualTextures could output PixelDepthOffset, but RVT never writes depth which causes issues on some platforms. Now disabled if VIRTUAL_TEXTURE_OUTPUT is enabled.

diff --git a/Engine/Shaders/Private/MaterialTemplate.ush b/Engine/Shaders/Private/MaterialTemplate.ush
--- a/Engine/Shaders/Private/MaterialTemplate.ush
+++ b/Engine/Shaders/Private/MaterialTemplate.ush
@@ -4663,7 +4663,7 @@ half3x3 AssembleTangentToWorld( half3 TangentToWorld0, half4 TangentToWorld2 )
 }
 
 // Whether the material shader should output pixel depth offset
-#define OUTPUT_PIXEL_DEPTH_OFFSET (WANT_PIXEL_DEPTH_OFFSET && !IS_NANITE_PASS && ((MATERIALBLENDING_SOLID || MATERIALBLENDING_MASKED) || (TRANSLUCENT_WRITING_VELOCITY)))
+#define OUTPUT_PIXEL_DEPTH_OFFSET (WANT_PIXEL_DEPTH_OFFSET && !VIRTUAL_TEXTURE_OUTPUT && !IS_NANITE_PASS && ((MATERIALBLENDING_SOLID || MATERIALBLENDING_MASKED) || (TRANSLUCENT_WRITING_VELOCITY)))
 
 // Whether to use the hidden d3d11 feature that supports depth writes with ZCull by only pushing into the screen
 //@todo - use for other SM5 platforms
