Substrate - fixed missing static shadow also noticed when using deferred shading.

#tests PC editor dx12
#rnx

diff --git a/Engine/Shaders/Private/BasePassPixelShader.usf b/Engine/Shaders/Private/BasePassPixelShader.usf
--- a/Engine/Shaders/Private/BasePassPixelShader.usf
+++ b/Engine/Shaders/Private/BasePassPixelShader.usf
@@ -1879,16 +1879,8 @@ void FPixelShaderInOut_MainPS(
 		GBuffer.GBufferAO = MaterialAO;
 		GBuffer.PerObjectGBufferData = GetPrimitive_PerObjectGBufferData(MaterialParameters.PrimitiveId);
 		GBuffer.Depth = MaterialParameters.ScreenPosition.w;
+		// GBuffer.SelectiveOutputMask is setup above.
 		EmissiveLuminance = Export.EmissiveLuminance;
-
-		#if USES_GBUFFER && !IS_NANITE_PASS
-		// This requires cleanup. Shader code that uses GBuffer.SelectiveOutputMask expects the outputmask to be in
-		// bits [4:7], but it gets packed as bits [0:3] in the flexible gbuffer since we might move it around.
-		if (Export.Anisotropy != 0)
-		{
-			GBuffer.SelectiveOutputMask = HAS_ANISOTROPY_MASK >> 4;
-		}
-		#endif
 		
 		// Patch material parameter to not have to change SetGBufferForShadingModel for now
 		MaterialParameters.WorldNormal = Export.WorldNormal;
diff --git a/Engine/Shaders/Private/ComputeShaderOutputCommon.ush b/Engine/Shaders/Private/ComputeShaderOutputCommon.ush
--- a/Engine/Shaders/Private/ComputeShaderOutputCommon.ush
+++ b/Engine/Shaders/Private/ComputeShaderOutputCommon.ush
@@ -29,7 +29,9 @@ uint GetQuadBinning()		{ return GetRootConstant1(); }
 uint GetHighPrecision()		{ return GetRootConstant2(); }
 uint GetDataByteOffset()	{ return GetRootConstant3(); }
 
-#if SUBSTRATE_OPAQUE_DEFERRED && SUBSTRATE_GBUFFER_FORMAT==1
+#define SUBSTRATE_OPAQUE_DEFERRED_ADAPTATIVE_GBUFFER (SUBSTRATE_OPAQUE_DEFERRED && SUBSTRATE_GBUFFER_FORMAT==1)
+
+#if SUBSTRATE_OPAQUE_DEFERRED_ADAPTATIVE_GBUFFER
 	#if SUBSTRATE_BASE_PASS_MRT_OUTPUT_COUNT != 3
 		#error Substrate SUBSTRATE_BASE_PASS_MRT_OUTPUT_COUNT has been updated but not the uint MRTs
 	#endif
@@ -96,7 +98,7 @@ FPixelShaderOut ShadePixel(const float2 SVPositionXY, uint QuadIndex, uint QuadP
 	return PixelShaderOut;
 }
 
-#if SUBSTRATE_OPAQUE_DEFERRED
+#if SUBSTRATE_OPAQUE_DEFERRED_ADAPTATIVE_GBUFFER
 
 void SubstrateExport(uint2 PixelPos, FPixelShaderOut ShadedPixel, uint Index)
 {
@@ -136,7 +138,7 @@ void ExportPixel(const uint2 PixelPos, FPixelShaderOut ShadedPixel)
 	ComputeShadingOutputs.OutTarget3[PixelPos] = ShadedPixel.MRT[3];
 #endif
 
-#if SUBSTRATE_OPAQUE_DEFERRED
+#if SUBSTRATE_OPAQUE_DEFERRED_ADAPTATIVE_GBUFFER
 
 	// In this case, here is the gbuffer pattern
 	// MRT0 is pixel color
@@ -154,7 +156,7 @@ void ExportPixel(const uint2 PixelPos, FPixelShaderOut ShadedPixel)
 	}
 	SubstrateExport(PixelPos, ShadedPixel.SubstrateTopLayerData);
 
-#else // SUBSTRATE_OPAQUE_DEFERRED
+#else // SUBSTRATE_OPAQUE_DEFERRED_ADAPTATIVE_GBUFFER
 
 #if PIXELSHADEROUTPUT_MRT4
 	ComputeShadingOutputs.OutTarget4[PixelPos] = ShadedPixel.MRT[4];
@@ -176,7 +178,7 @@ void ExportPixel(const uint2 PixelPos, FPixelShaderOut ShadedPixel)
 	ComputeShadingOutputs.OutSGGX[PixelPos] = ShadedPixel.SGGX;
 #endif
 
-#endif // SUBSTRATE_OPAQUE_DEFERRED
+#endif // SUBSTRATE_OPAQUE_DEFERRED_ADAPTATIVE_GBUFFER
 }
 
 #define VIS_HELPER_LANES 0
diff --git a/Engine/Shaders/Private/DeferredLightPixelShaders.usf b/Engine/Shaders/Private/DeferredLightPixelShaders.usf
--- a/Engine/Shaders/Private/DeferredLightPixelShaders.usf
+++ b/Engine/Shaders/Private/DeferredLightPixelShaders.usf
@@ -315,7 +315,7 @@ void DeferredLightPixelMain(
 		#if SUBSTRATE_LOAD_FROM_MATERIALCONTAINER
 			SubstrateShadowTermInputParameters.PrecomputedShadowFactors				= SubstrateReadPrecomputedShadowFactors(SubstratePixelHeader, PixelPos, SceneTexturesStruct.GBufferETexture);
 		#else
-			SubstrateShadowTermInputParameters.PrecomputedShadowFactors				= 1.0f;
+			SubstrateShadowTermInputParameters.PrecomputedShadowFactors				= GBuffer.PrecomputedShadowFactors;
 		#endif
 			SubstrateShadowTermInputParameters.TranslatedWorldPosition				= DerivedParams.TranslatedWorldPosition;
 			SubstrateShadowTermInputParameters.LightAttenuation						= GetLightAttenuationFromShadow(InputParams, SceneDepth, DerivedParams.TranslatedWorldPosition);
