Heterogeneous Volumes: Fix minor shader compilation warnings

#rnx

diff --git a/Engine/Shaders/Private/HeterogeneousVolumes/HeterogeneousVolumesAdaptiveVolumetricShadowMapSampling.ush b/Engine/Shaders/Private/HeterogeneousVolumes/HeterogeneousVolumesAdaptiveVolumetricShadowMapSampling.ush
--- a/Engine/Shaders/Private/HeterogeneousVolumes/HeterogeneousVolumesAdaptiveVolumetricShadowMapSampling.ush
+++ b/Engine/Shaders/Private/HeterogeneousVolumes/HeterogeneousVolumesAdaptiveVolumetricShadowMapSampling.ush
@@ -551,7 +551,7 @@ FAdaptiveVolumetricShadowMap GetAdaptiveVolumetricShadowMap()
 	for (int i = 0; i < 6; ++i)
 	{
 		ShadowMap.TranslatedWorldToShadow[i] = AVSM.TranslatedWorldToShadow[i];
-		ShadowMap.TranslatedWorldOrigin[i] = AVSM.TranslatedWorldOrigin[i];
+		ShadowMap.TranslatedWorldOrigin[i] = AVSM.TranslatedWorldOrigin[i].xyz;
 		ShadowMap.TranslatedWorldPlane[i] = AVSM.TranslatedWorldPlane[i];
 		ShadowMap.SplitDepths[i] = AVSM.SplitDepths[i];
 	}
@@ -578,7 +578,7 @@ FAdaptiveVolumetricShadowMap GetAdaptiveVolumetricShadowMapForCamera()
 	for (int i = 0; i < 6; ++i)
 	{
 		ShadowMap.TranslatedWorldToShadow[i] = CameraAVSM.TranslatedWorldToShadow[i];
-		ShadowMap.TranslatedWorldOrigin[i] = CameraAVSM.TranslatedWorldOrigin[i];
+		ShadowMap.TranslatedWorldOrigin[i] = CameraAVSM.TranslatedWorldOrigin[i].xyz;
 		ShadowMap.TranslatedWorldPlane[i] = CameraAVSM.TranslatedWorldPlane[i];
 		ShadowMap.SplitDepths[i] = CameraAVSM.SplitDepths[i];
 	}
diff --git a/Engine/Shaders/Private/HeterogeneousVolumes/HeterogeneousVolumesVisualizeCascades.usf b/Engine/Shaders/Private/HeterogeneousVolumes/HeterogeneousVolumesVisualizeCascades.usf
--- a/Engine/Shaders/Private/HeterogeneousVolumes/HeterogeneousVolumesVisualizeCascades.usf
+++ b/Engine/Shaders/Private/HeterogeneousVolumes/HeterogeneousVolumesVisualizeCascades.usf
@@ -34,7 +34,7 @@ FAdaptiveVolumetricShadowMap GetAdaptiveVolumetricShadowMap()
 	for (int i = 0; i < 6; ++i)
 	{
 		ShadowMap.TranslatedWorldToShadow[i] = AVSM.TranslatedWorldToShadow[i];
-		ShadowMap.TranslatedWorldOrigin[i] = AVSM.TranslatedWorldOrigin[i];
+		ShadowMap.TranslatedWorldOrigin[i] = AVSM.TranslatedWorldOrigin[i].xyz;
 		ShadowMap.TranslatedWorldPlane[i] = AVSM.TranslatedWorldPlane[i];
 		ShadowMap.SplitDepths[i] = AVSM.SplitDepths[i];
 	}
