Substrate - fixed missing indirect on Substrate Forward materials.

#tests PC editor dx12
#rnx

diff --git a/Engine/Shaders/Private/Substrate/SubstrateForwardLighting.ush b/Engine/Shaders/Private/Substrate/SubstrateForwardLighting.ush
--- a/Engine/Shaders/Private/Substrate/SubstrateForwardLighting.ush
+++ b/Engine/Shaders/Private/Substrate/SubstrateForwardLighting.ush
@@ -209,7 +209,6 @@ float3 SubstrateForwardLighting(
 					float3 SubsurfaceColorForIndirect = SubstrateEnvLight.DiffuseBackFaceWeight;
 
 					// Evaluate diffuse lighting
-					if (any((DiffuseColorForIndirect + SubsurfaceColorForIndirect) > 0.0f))
 					{
 						const bool bEvaluateBackface = any(SubsurfaceColorForIndirect > 0.0);
 
@@ -219,8 +218,10 @@ float3 SubstrateForwardLighting(
 					#if (MATERIALBLENDING_TRANSLUCENT || MATERIALBLENDING_ADDITIVE || MATERIALBLENDING_MODULATE) && PROJECT_SUPPORTS_LUMEN && !FORWARD_SHADING
 						if (IsLumenTranslucencyGIEnabled())
 						{
-							// Lumen Dynamic GI + shadowed Skylight
-							FTwoBandSHVectorRGB TranslucencyGISH = GetTranslucencyGIVolumeLighting(MaterialParameters.AbsoluteWorldPosition, ResolvedView.WorldToClip, true);
+							if (any((DiffuseColorForIndirect + SubsurfaceColorForIndirect) > 0.0f)) // Lumen evaluation can be skipped if all diffuse components are 0.
+							{
+								// Lumen Dynamic GI + shadowed Skylight
+								FTwoBandSHVectorRGB TranslucencyGISH = GetTranslucencyGIVolumeLighting(MaterialParameters.AbsoluteWorldPosition, ResolvedView.WorldToClip, true);
 
 							#if TRANSLUCENCY_LIGHTING_VOLUMETRIC_NONDIRECTIONAL || TRANSLUCENCY_LIGHTING_VOLUMETRIC_PERVERTEX_NONDIRECTIONAL
 
@@ -245,10 +246,11 @@ float3 SubstrateForwardLighting(
 								}
 								#endif
 							#endif
+							}
 						}
 						else
 					#endif
-						if (UseBasePassSkylight > 0)
+						if (UseBasePassSkylight > 0) // Non Lumen evaluation cannot be skipped if all diffuse components are 0 due to IndirectIrradiance which might still contribute via GetImageBasedReflectionSpecular
 						{
 							GetPrecomputedIndirectLightingAndSkyLight(
 								MaterialParameters,
