Substrate - respects SUPPORTS_ANISOTROPIC_MATERIALS at evaluation time similarly to legacy.

#tests PC editor dx12
#rnx

diff --git a/Engine/Shaders/Private/Substrate/SubstrateEvaluation.ush b/Engine/Shaders/Private/Substrate/SubstrateEvaluation.ush
--- a/Engine/Shaders/Private/Substrate/SubstrateEvaluation.ush
+++ b/Engine/Shaders/Private/Substrate/SubstrateEvaluation.ush
@@ -465,7 +465,7 @@ FSubstrateEvaluateResult SubstrateEvaluateBSDFCommon(FSubstrateBSDFContext BSDFC
 			float3 F0					= SLAB_F0(BSDFContext.BSDF);
 			float3 F90					= SLAB_F90(BSDFContext.BSDF);
 			const float SafeRoughness	= MakeRoughnessSafe(SLAB_ROUGHNESS(BSDFContext.BSDF));
-			const bool bHasAnisotropy	= BSDF_GETHASANISOTROPY(BSDFContext.BSDF) && !bIsRectLight;
+			const bool bHasAnisotropy	= BSDF_GETHASANISOTROPY(BSDFContext.BSDF) && !bIsRectLight && SUPPORTS_ANISOTROPIC_MATERIALS;
 			const bool bHaziness		= BSDF_GETHASHAZINESS(BSDFContext.BSDF);
 
 			// Specular occlusion is only used here once to affect the F90 source parameters. F0 will decrease naturally to 0.
@@ -1353,7 +1353,7 @@ FSubstrateImportanceSampleResult SubstrateImportanceSampleBSDF(FSubstrateBSDFCon
 			const float3 F0 			= SLAB_F0(BSDFContext.BSDF);
 			const float3 F90 			= SLAB_F90(BSDFContext.BSDF);
 
-			const bool bHasAnisotropy	= BSDF_GETHASANISOTROPY(BSDFContext.BSDF);
+			const bool bHasAnisotropy	= BSDF_GETHASANISOTROPY(BSDFContext.BSDF) && SUPPORTS_ANISOTROPIC_MATERIALS;
 			const bool bHaziness		= BSDF_GETHASHAZINESS(BSDFContext.BSDF);
 			const float SafeRoughness   = MakeRoughnessSafe(SLAB_ROUGHNESS(BSDFContext.BSDF));
 
@@ -1671,7 +1671,7 @@ FSubstrateEnvLightResult SubstrateEvaluateForEnvLight(FSubstrateBSDFContext BSDF
 				#endif
 
 				#if SUBSTRATE_COMPLEXPATH
-				const bool bHasAnisotropy = BSDF_GETHASANISOTROPY(BSDFContext.BSDF);
+				const bool bHasAnisotropy = BSDF_GETHASANISOTROPY(BSDFContext.BSDF) && SUPPORTS_ANISOTROPIC_MATERIALS;
 				if (bHasAnisotropy)
 				{
 					// Modified the BSDF normal (and roughness)
