Fix shader compilation error.

diff --git a/Engine/Shaders/Private/Substrate/SubstrateLightingCommon.ush b/Engine/Shaders/Private/Substrate/SubstrateLightingCommon.ush
--- a/Engine/Shaders/Private/Substrate/SubstrateLightingCommon.ush
+++ b/Engine/Shaders/Private/Substrate/SubstrateLightingCommon.ush
@@ -26,9 +26,11 @@ half4 GetPrecomputedShadowFactors(FSubstratePixelHeader SubstratePixelHeader, fl
 
 	}
 	else
-#else
-		OutPrecomputedShadowFactors.x = 1;
 #endif
+	{
+
+		OutPrecomputedShadowFactors.x = 1;
+	}
 
 	return half4(OutPrecomputedShadowFactors);
 }
