Reorder Substrate conversion block and merge SSS profile blocks.

diff --git a/Engine/Shaders/Private/Substrate/SubstrateExport.ush b/Engine/Shaders/Private/Substrate/SubstrateExport.ush
--- a/Engine/Shaders/Private/Substrate/SubstrateExport.ush
+++ b/Engine/Shaders/Private/Substrate/SubstrateExport.ush
@@ -993,39 +993,6 @@
 					Out.Metallic	= HAIR_SCATTER(CurrentBSDF);
 					Out.CustomData  = float4(0.0, 0.0, HAIR_BACKLIT(CurrentBSDF), 0.0);
 				}
-				else if (SSSType == SSS_TYPE_DIFFUSION_PROFILE)
-				{
-					Out.SubsurfaceProfileRadiusScale = SLAB_SSSPROFILERADIUSSCALE(CurrentBSDF);
-					Out.CustomData.rgb = EncodeSubsurfaceProfile(Out.SubsurfaceProfileIDFloat);
-					Out.CustomData.a   = Out.SubsurfaceProfileRadiusScale;
-				}
-				else if (SSSType == SSS_TYPE_WRAP || SSSType == SSS_TYPE_TWO_SIDED_WRAP)
-				{
-					Out.CustomData.rgb = EncodeSubsurfaceColor(Out.SubsurfaceColor);
-					Out.CustomData.a   = SubstrateSubSurfaceGetWrapOpacityFromAnisotropy(SLAB_SSSPHASEANISOTROPY(CurrentBSDF));
-				}
-				else if (BSDF_GETHASFUZZ(CurrentBSDF))
-				{
-					Out.CustomData.rgb = EncodeSubsurfaceColor(SLAB_FUZZ_COLOR(CurrentBSDF));
-					Out.CustomData.a   = SLAB_FUZZ_AMOUNT(CurrentBSDF); 
-					//MakeRoughnessSafe(SLAB_FUZZ_ROUGHNESS(CurrentBSDF), SUBSTRATE_MIN_FUZZ_ROUGHNESS);
-					Out.IndirectIrradianceScale *= 1 - Out.CustomData.a;
-				}
-				else if (BSDF_GETHASHAZINESS(CurrentBSDF))
-				{
-					FHaziness LocalHaziness = UnpackHaziness(SLAB_HAZINESS(CurrentBSDF));
-					if (LocalHaziness.bSimpleClearCoat)
-					{
-						Out.CustomData.x = LocalHaziness.Weight;
-						Out.CustomData.y = LocalHaziness.Roughness;
-
-					#if CLEAR_COAT_BOTTOM_NORMAL
-						Out.CustomData.wz = LocalHaziness.BottomNormalOct.xy;
-					#else
-						Out.CustomData.wz = DEFAULT_CLEAR_COAT_BOTTOM_NORMAL_OCT;
-					#endif
-					}
-				}
 				else if (BSDFType == SUBSTRATE_BSDF_TYPE_EYE)
 				{
 					const float3 IrisNormal		= EYE_IRISNORMAL(CurrentBSDF);
@@ -1090,8 +1057,12 @@
 						Out.Curvature = clamp(Out.Curvature, 0.001f, 1.0f);
 				#endif
 				}
-				else if (SSSType == SSS_TYPE_DIFFUSION_PROFILE) // Thin and not thin, Subsurface and TwoSided Foliage
+				else if (SSSType == SSS_TYPE_DIFFUSION_PROFILE)  // Thin and not thin, Subsurface and TwoSided Foliage
 				{
+					Out.SubsurfaceProfileRadiusScale = SLAB_SSSPROFILERADIUSSCALE(CurrentBSDF);
+					Out.CustomData.rgb = EncodeSubsurfaceProfile(Out.SubsurfaceProfileIDFloat);
+					Out.CustomData.a   = Out.SubsurfaceProfileRadiusScale;
+
 					const float3 GBufferWorldNormal = N;
 					const float Opacity = SLAB_SSSPROFILERADIUSSCALE(CurrentBSDF);
 					
@@ -1121,6 +1092,32 @@
 					}
 				#endif
 				}
+				else if (SSSType == SSS_TYPE_WRAP || SSSType == SSS_TYPE_TWO_SIDED_WRAP)
+				{
+					Out.CustomData.rgb = EncodeSubsurfaceColor(Out.SubsurfaceColor);
+					Out.CustomData.a   = SubstrateSubSurfaceGetWrapOpacityFromAnisotropy(SLAB_SSSPHASEANISOTROPY(CurrentBSDF));
+				}
+				else if (BSDF_GETHASFUZZ(CurrentBSDF))
+				{
+					Out.CustomData.rgb = EncodeSubsurfaceColor(SLAB_FUZZ_COLOR(CurrentBSDF));
+					Out.CustomData.a   = SLAB_FUZZ_AMOUNT(CurrentBSDF); 
+					Out.IndirectIrradianceScale *= 1 - Out.CustomData.a;
+				}
+				else if (BSDF_GETHASHAZINESS(CurrentBSDF))
+				{
+					FHaziness LocalHaziness = UnpackHaziness(SLAB_HAZINESS(CurrentBSDF));
+					if (LocalHaziness.bSimpleClearCoat)
+					{
+						Out.CustomData.x = LocalHaziness.Weight;
+						Out.CustomData.y = LocalHaziness.Roughness;
+
+					#if CLEAR_COAT_BOTTOM_NORMAL
+						Out.CustomData.wz = LocalHaziness.BottomNormalOct.xy;
+					#else
+						Out.CustomData.wz = DEFAULT_CLEAR_COAT_BOTTOM_NORMAL_OCT;
+					#endif
+					}
+				}
 			}
 		#undef CurrentBSDF
 		}
