Fix Subsurface and Two-Sided foliage SSS mode in Substrate when using Format=0.

diff --git a/Engine/Shaders/Private/Substrate/SubstrateRead.ush b/Engine/Shaders/Private/Substrate/SubstrateRead.ush
--- a/Engine/Shaders/Private/Substrate/SubstrateRead.ush
+++ b/Engine/Shaders/Private/Substrate/SubstrateRead.ush
@@ -129,7 +129,7 @@ FSubstrateGBufferBSDF SubstrateReadGBufferBSDF(in FScreenSpaceData ScreenSpaceDa
 					float3 Extinction = TransmittanceToExtinction(ExtractSubsurfaceColor(GBuffer), SUBSTRATE_SIMPLEVOLUME_THICKNESS_CM * CENTIMETER_TO_METER);
 					SSSMFP = (1.0f / max(SUBSTRATE_EPSILON, Extinction));
 				}
-				SSSPhaseAnisotropy = GBuffer.CustomData.a;		// TODO test, should be one minus opacity
+				SSSPhaseAnisotropy = 1.f - GBuffer.CustomData.a;
 				SSSMFPScale = 1.0f;
 			}
 			else if (GBuffer.ShadingModelID == SHADINGMODELID_TWOSIDED_FOLIAGE)
@@ -140,7 +140,7 @@ FSubstrateGBufferBSDF SubstrateReadGBufferBSDF(in FScreenSpaceData ScreenSpaceDa
 					float3 Extinction = TransmittanceToExtinction(ExtractSubsurfaceColor(GBuffer), SUBSTRATE_SIMPLEVOLUME_THICKNESS_CM * CENTIMETER_TO_METER);
 					SSSMFP = (1.0f / max(SUBSTRATE_EPSILON, Extinction));
 				}
-				SSSPhaseAnisotropy = GBuffer.CustomData.a;		// TODO test, should be one minus opacity
+				SSSPhaseAnisotropy = 1.f - GBuffer.CustomData.a;
 				SSSMFPScale = 1.0f;
 			}
 			else if (GBuffer.ShadingModelID == SHADINGMODELID_CLEAR_COAT)
