Fix incorrect translation for SSS profile introduced in CL 44648329

diff --git a/Engine/Shaders/Private/Substrate/SubstrateExport.ush b/Engine/Shaders/Private/Substrate/SubstrateExport.ush
--- a/Engine/Shaders/Private/Substrate/SubstrateExport.ush
+++ b/Engine/Shaders/Private/Substrate/SubstrateExport.ush
@@ -1071,9 +1071,6 @@ FExportResult SubstrateMaterialExportOut(
 					if (Opacity > SSSS_OPACITY_THRESHOLD_EPS)
 				#endif
 					{
-						Out.CustomData.rgb = EncodeSubsurfaceColor(Out.SubsurfaceColor);
-						Out.CustomData.a = Opacity;
-
 					#if SHADING_PATH_MOBILE
 						#if MATERIAL_SUBSURFACE_PROFILE_USE_CURVATURE
 						Out.Curvature = MobileShadingPathCurvature;
