Substrate - added support for r.RayTracing.SkyLight when gbufferformat=0.

#tests PC editor dx12
#rnx

diff --git a/Engine/Shaders/Private/RayTracing/RayTracingSkyLightRGS.usf b/Engine/Shaders/Private/RayTracing/RayTracingSkyLightRGS.usf
--- a/Engine/Shaders/Private/RayTracing/RayTracingSkyLightRGS.usf
+++ b/Engine/Shaders/Private/RayTracing/RayTracingSkyLightRGS.usf
@@ -23,10 +23,14 @@ RAY_TRACING_ENTRY_RAYGEN(SkyLightRGS)
 	uint2 DispatchThreadId = DispatchRaysIndex().xy + View.ViewRectMin.xy;
 	uint2 PixelCoord = GetPixelCoord(DispatchThreadId, UpscaleFactor);
 
-#if SUBSTRATE_ENABLED
+
+#if SUBSTRATE_ENABLED && SUBSTRATE_GBUFFER_FORMAT==1
+
 	RWSkyOcclusionMaskUAV[DispatchThreadId] = float4(0,0,0,1);
 	RWSkyOcclusionRayDistanceUAV[DispatchThreadId] = float2(200000, 0);
-#else // SUBSTRATE_ENABLED
+
+#else // SUBSTRATE_ENABLED && SUBSTRATE_GBUFFER_FORMAT==1
+
 	// Get G-Buffer surface data
 	float2 InvBufferSize = View.BufferSizeAndInvSize.zw;
 	float2 UV = (float2(PixelCoord) + 0.5) * InvBufferSize;
@@ -98,5 +102,7 @@ RAY_TRACING_ENTRY_RAYGEN(SkyLightRGS)
 
 	RWSkyOcclusionMaskUAV[DispatchThreadId] = float4(ClampToHalfFloatRange(DiffuseExitantRadiance.rgb), AmbientOcclusion);
 	RWSkyOcclusionRayDistanceUAV[DispatchThreadId] = float2(HitDistance, SamplesPerPixel);
-#endif // SUBSTRATE_ENABLED
+
+#endif // SUBSTRATE_ENABLED && SUBSTRATE_GBUFFER_FORMAT==1
+
 }
\ No newline at end of file
