Fix SSS roughness does not matching legacy roughness when using Format=0.

diff --git a/Engine/Shaders/Private/Substrate/SubstrateExport.ush b/Engine/Shaders/Private/Substrate/SubstrateExport.ush
--- a/Engine/Shaders/Private/Substrate/SubstrateExport.ush
+++ b/Engine/Shaders/Private/Substrate/SubstrateExport.ush
@@ -1059,6 +1059,10 @@ FExportResult SubstrateMaterialExportOut(
 				}
 				else if (SSSType == SSS_TYPE_DIFFUSION_PROFILE)  // Thin and not thin, Subsurface and TwoSided Foliage
 				{
+					// Roughness has been modify when calling GetSubstrateSlabBSDF(), but expending SSS Profile roughness
+					// Revert roughness to its original average value to get correct result
+					Out.Roughness = GetSubsurfaceProfileOriginalRoughness(SLAB_SSSPROFILEID(CurrentBSDF), Out.Roughness, SLAB_SSSPROFILERADIUSSCALE(CurrentBSDF));
+
 					Out.SubsurfaceProfileRadiusScale = SLAB_SSSPROFILERADIUSSCALE(CurrentBSDF);
 					Out.CustomData.rgb = EncodeSubsurfaceProfile(Out.SubsurfaceProfileIDFloat);
 					Out.CustomData.a   = Out.SubsurfaceProfileRadiusScale;
