Substrate - added support for r.RayTracing.SkyLight when gbufferformat=1.
This is basic support where many slabs are summarised into a single one.

#tests PC editor dx12
#rnx

diff --git a/Engine/Shaders/Private/RayTracing/CompositeSkyLightPS.usf b/Engine/Shaders/Private/RayTracing/CompositeSkyLightPS.usf
--- a/Engine/Shaders/Private/RayTracing/CompositeSkyLightPS.usf
+++ b/Engine/Shaders/Private/RayTracing/CompositeSkyLightPS.usf
@@ -4,6 +4,17 @@
 #include "../DeferredShadingCommon.ush"
 #include "../SceneTextureParameters.ush"
 
+#if SUBSTRATE_ENABLED && SUBSTRATE_GBUFFER_FORMAT==1
+#define SUBSTRATE_MATERIALCONTAINER_IS_VIEWRESOURCE 1
+#define SUBSTRATE_FASTPATH 0
+#define SUBSTRATE_SINGLEPATH 0
+#define SUBSTRATE_COMPLEXSPECIALPATH 0
+#define SUBSTRATE_SSS_MATERIAL_OVERRIDE 0		// SUBSTRATE_TODO: tile the pass to work with glint
+#define SUBSTRATE_GLINTS_ALLOWED 0
+#include "../Substrate/Substrate.ush"
+#include "../Substrate/SubstrateEvaluation.ush"
+#endif
+
 Texture2D<float4> SkyLightTexture;
 SamplerState SkyLightTextureSampler;
 
@@ -13,8 +24,45 @@ void CompositeSkyLightPS(
 )
 {
 	const float2 UV = OutUVAndScreenPos.xy;
+	const uint2 PixelCoord = UV * View.BufferSizeAndInvSize.xy;
+
+#if SUBSTRATE_ENABLED && SUBSTRATE_GBUFFER_FORMAT==1
+
+
+	FSubstrateAddressing SubstrateAddressing = GetSubstratePixelDataByteOffset(PixelCoord, uint2(View.BufferSizeAndInvSize.xy), Substrate.MaxBytesPerPixel);
+	FSubstratePixelHeader SubstratePixelHeader = UnpackSubstrateHeaderIn(Substrate.MaterialTextureArray, SubstrateAddressing, Substrate.TopLayerTexture);
+
+	const FSubstrateIntegrationSettings Settings = InitSubstrateIntegrationSettings(false /*bForceFullyRough*/, Substrate.bRoughDiffuse, Substrate.PeelLayersAboveDepth, Substrate.bRoughnessTracking);
+
+	const float3 NullV = float3(0, 0, 1);
+	const float3 NullL = float3(0, 0, 1);
+	float3 Albedo = 0.0f;
+	Substrate_for(uint ClosureIndex = 0, ClosureIndex < SubstratePixelHeader.ClosureCount, ++ClosureIndex)
+	{
+		FSubstrateBSDF CurrentBSDF = UnpackSubstrateBSDF(Substrate.MaterialTextureArray, SubstrateAddressing, SubstratePixelHeader);
+		FSubstrateBSDFContext Context = SubstrateCreateBSDFContext(SubstratePixelHeader, CurrentBSDF, SubstrateAddressing, NullV);
+
+		if (SubstrateIsBSDFVisible(CurrentBSDF))
+		{
+			FSubstrateAddressing NullSubstrateAddressing = (FSubstrateAddressing)0;	// Fake unused in SubstrateCreateBSDFContext when using Forward inline shading
+			FSubstrateBSDFContext SubstrateBSDFContext = SubstrateCreateBSDFContext(SubstratePixelHeader, CurrentBSDF, NullSubstrateAddressing, NullV, NullL);
+			FSubstrateEnvLightResult SubstrateEnvLight = SubstrateEvaluateForEnvLight(SubstrateBSDFContext, true /*bEnableSpecular*/, Settings);
+
+			// Use LuminanceWeightV instead of LuminanceWeight(..) as we only need to weight these value with the view transmittance, not the light transmittance;
+			const float3 Weight = CurrentBSDF.LuminanceWeightV;
+
+			Albedo += Weight * SubstrateGetBSDFDiffuseColor(CurrentBSDF);
+		}
+	}
+
+#else
+
 	FGBufferData GBufferData = GetGBufferDataFromSceneTextures(UV);
 	float3 Albedo = GBufferData.StoredBaseColor - GBufferData.StoredBaseColor * GBufferData.Metallic;
+
+#endif
+
+
 	float4 SkyLight = SkyLightTexture.Sample(SkyLightTextureSampler, UV);
 	// Apply albedo after denoising
 	SkyLight.rgb *= Albedo;
diff --git a/Engine/Shaders/Private/RayTracing/RayTracingSkyLightRGS.usf b/Engine/Shaders/Private/RayTracing/RayTracingSkyLightRGS.usf
--- a/Engine/Shaders/Private/RayTracing/RayTracingSkyLightRGS.usf
+++ b/Engine/Shaders/Private/RayTracing/RayTracingSkyLightRGS.usf
@@ -8,6 +8,17 @@
 #include "../HairStrands/HairStrandsRaytracing.ush"
 #endif
 
+#if SUBSTRATE_ENABLED && SUBSTRATE_GBUFFER_FORMAT==1
+#define SUBSTRATE_MATERIALCONTAINER_IS_VIEWRESOURCE 1
+#define SUBSTRATE_FASTPATH 0
+#define SUBSTRATE_SINGLEPATH 0
+#define SUBSTRATE_COMPLEXSPECIALPATH 0
+#define SUBSTRATE_SSS_MATERIAL_OVERRIDE 0		// SUBSTRATE_TODO: tile the pass to work with glint
+#define SUBSTRATE_GLINTS_ALLOWED 0
+#include "../Substrate/Substrate.ush"
+#include "../Substrate/SubstrateEvaluation.ush"
+#endif
+
 #include "RayTracingSkyLightEvaluation.ush"
 #include "RayTracingDeferredShadingCommon.ush"
 
@@ -23,11 +34,57 @@ RAY_TRACING_ENTRY_RAYGEN(SkyLightRGS)
 	uint2 DispatchThreadId = DispatchRaysIndex().xy + View.ViewRectMin.xy;
 	uint2 PixelCoord = GetPixelCoord(DispatchThreadId, UpscaleFactor);
 
+	float DeviceZ = SceneDepthTexture.Load(int3(PixelCoord, 0)).r;
+	float3 TranslatedWorldPosition;
+	float3 CameraDirection;
+	ReconstructTranslatedWorldPositionAndCameraDirectionFromDeviceZ(PixelCoord, DeviceZ, TranslatedWorldPosition, CameraDirection);
 
 #if SUBSTRATE_ENABLED && SUBSTRATE_GBUFFER_FORMAT==1
 
-	RWSkyOcclusionMaskUAV[DispatchThreadId] = float4(0,0,0,1);
-	RWSkyOcclusionRayDistanceUAV[DispatchThreadId] = float2(200000, 0);
+	FSubstrateAddressing SubstrateAddressing = GetSubstratePixelDataByteOffset(PixelCoord, uint2(View.BufferSizeAndInvSize.xy), Substrate.MaxBytesPerPixel);
+	FSubstratePixelHeader SubstratePixelHeader = UnpackSubstrateHeaderIn(Substrate.MaterialTextureArray, SubstrateAddressing, Substrate.TopLayerTexture);
+
+	// When Lumen is not used, only MaterialAO and ShadingID (see IsValid) are read, sourced form the single UINT read for the SubstratePixelHeader.
+	const uint BSDFType = SubstratePixelHeader.SubstrateGetBSDFType();
+
+	const FSubstrateIntegrationSettings Settings = InitSubstrateIntegrationSettings(false /*bForceFullyRough*/, Substrate.bRoughDiffuse, Substrate.PeelLayersAboveDepth, Substrate.bRoughnessTracking);
+
+	float3 Albedo = 0;
+	float3 DiffuseColor = 0;
+	float3 WorldNormal = 0;
+	Substrate_for(uint ClosureIndex = 0, ClosureIndex < SubstratePixelHeader.ClosureCount, ++ClosureIndex)
+	{
+		float3 NullV = float3(0, 0, 1);
+		float3 NullL = float3(0, 0, 1);
+		FSubstrateBSDF CurrentBSDF = UnpackSubstrateBSDF(Substrate.MaterialTextureArray, SubstrateAddressing, SubstratePixelHeader);
+		FSubstrateBSDFContext Context = SubstrateCreateBSDFContext(SubstratePixelHeader, CurrentBSDF, SubstrateAddressing, NullV);
+
+		if (SubstrateIsBSDFVisible(CurrentBSDF))
+		{
+			FSubstrateAddressing NullSubstrateAddressing = (FSubstrateAddressing)0;	// Fake unused in SubstrateCreateBSDFContext when using Forward inline shading
+			FSubstrateBSDFContext SubstrateBSDFContext = SubstrateCreateBSDFContext(SubstratePixelHeader, CurrentBSDF, NullSubstrateAddressing, NullV, NullL);
+			FSubstrateEnvLightResult SubstrateEnvLight = SubstrateEvaluateForEnvLight(SubstrateBSDFContext, true /*bEnableSpecular*/, Settings);
+
+			// Use LuminanceWeightV instead of LuminanceWeight(..) as we only need to weight these value with the view transmittance, not the light transmittance;
+			const float3 Weight = CurrentBSDF.LuminanceWeightV;
+
+			DiffuseColor  += Weight * SubstrateGetBSDFDiffuseColor(CurrentBSDF);
+			WorldNormal   += Weight * SubstrateBSDFContext.N;
+		}
+	}
+
+	FGBufferData GBufferData = (FGBufferData)0;
+	GBufferData.ShadingModelID = BSDFType == SUBSTRATE_BSDF_TYPE_HAIR ? SHADINGMODELID_HAIR : SHADINGMODELID_DEFAULT_LIT;
+	GBufferData.DiffuseColor = DiffuseColor;
+	GBufferData.WorldNormal = normalize(WorldNormal);
+	GBufferData.SpecularColor = 0;
+	GBufferData.Roughness = 0.5f;
+	// No anisotropy
+	GBufferData.SelectiveOutputMask = 0;
+	GBufferData.Anisotropy = 0.0f;
+
+	WorldNormal = GBufferData.WorldNormal;
+	Albedo = GBufferData.DiffuseColor;
 
 #else // SUBSTRATE_ENABLED && SUBSTRATE_GBUFFER_FORMAT==1
 
@@ -43,10 +100,6 @@ RAY_TRACING_ENTRY_RAYGEN(SkyLightRGS)
 	FGBufferData GBufferData = GetGBufferDataFromSceneTexturesLoad(PixelCoord);
 #endif
 
-	float DeviceZ = SceneDepthTexture.Load(int3(PixelCoord, 0)).r;
-	float3 TranslatedWorldPosition;
-	float3 CameraDirection;
-	ReconstructTranslatedWorldPositionAndCameraDirectionFromDeviceZ(PixelCoord, DeviceZ, TranslatedWorldPosition, CameraDirection);
 	float3 WorldNormal = GBufferData.WorldNormal;
 	float3 Albedo = GBufferData.DiffuseColor;
 
@@ -57,6 +110,8 @@ RAY_TRACING_ENTRY_RAYGEN(SkyLightRGS)
 		GBufferData.DiffuseColor = Albedo;
 	}
 
+#endif // SUBSTRATE_ENABLED && SUBSTRATE_GBUFFER_FORMAT==1
+
 	// Mask out depth values that are infinitely far away
 	bool IsFiniteDepth = DeviceZ > 0.0;
 	bool bTraceRay = (
@@ -103,6 +158,4 @@ RAY_TRACING_ENTRY_RAYGEN(SkyLightRGS)
 	RWSkyOcclusionMaskUAV[DispatchThreadId] = float4(ClampToHalfFloatRange(DiffuseExitantRadiance.rgb), AmbientOcclusion);
 	RWSkyOcclusionRayDistanceUAV[DispatchThreadId] = float2(HitDistance, SamplesPerPixel);
 
-#endif // SUBSTRATE_ENABLED && SUBSTRATE_GBUFFER_FORMAT==1
-
 }
\ No newline at end of file
diff --git a/Engine/Source/Runtime/Renderer/Private/RayTracing/RaytracingSkylight.cpp b/Engine/Source/Runtime/Renderer/Private/RayTracing/RaytracingSkylight.cpp
--- a/Engine/Source/Runtime/Renderer/Private/RayTracing/RaytracingSkylight.cpp
+++ b/Engine/Source/Runtime/Renderer/Private/RayTracing/RaytracingSkylight.cpp
@@ -43,6 +43,7 @@ static FAutoConsoleVariableRef CVarRayTracingSkyLight(
 #include "PostProcess/PostProcessing.h"
 #include "PostProcess/SceneFilterRendering.h"
 #include "HairStrands/HairStrandsRendering.h"
+#include "Substrate/Substrate.h"
 
 
 static int32 GRayTracingSkyLightSamplesPerPixel = -1;
@@ -260,6 +261,7 @@ class FRayTracingSkyLightRGS : public FGlobalShader
 		SHADER_PARAMETER_RDG_UNIFORM_BUFFER(FVirtualVoxelParameters, VirtualVoxel)
 		SHADER_PARAMETER_RDG_UNIFORM_BUFFER(FSceneUniformParameters, Scene)
 		SHADER_PARAMETER_RDG_UNIFORM_BUFFER(FNaniteRayTracingUniformParameters, NaniteRayTracing)
+		SHADER_PARAMETER_RDG_UNIFORM_BUFFER(FSubstrateGlobalUniformParameters, Substrate)
 		
 		SHADER_PARAMETER_STRUCT_INCLUDE(FSkyLightVisibilityRaysData, SkyLightVisibilityRaysData)
 		SHADER_PARAMETER_STRUCT_INCLUDE(FSceneTextureParameters, SceneTextures)
@@ -456,6 +458,7 @@ void FDeferredShadingSceneRenderer::RenderRayTracingSkyLight(
 		PassParameters->ViewUniformBuffer = View.ViewUniformBuffer;
 		PassParameters->Scene = GetSceneUniformBufferRef(GraphBuilder);
 		PassParameters->NaniteRayTracing = Nanite::GRayTracingManager.GetUniformBuffer();
+		PassParameters->Substrate = Substrate::BindSubstrateGlobalUniformParameters(View);
 
 		const bool bUseHairLighting = 
 			HairStrands::HasViewHairStrandsData(View) && 
@@ -640,6 +643,7 @@ class FCompositeSkyLightPS : public FGlobalShader
 		SHADER_PARAMETER_RDG_TEXTURE(Texture2D, SkyLightTexture)
 		SHADER_PARAMETER_SAMPLER(SamplerState, SkyLightTextureSampler)
 		SHADER_PARAMETER_STRUCT_INCLUDE(FViewShaderParameters, View)
+		SHADER_PARAMETER_RDG_UNIFORM_BUFFER(FSubstrateGlobalUniformParameters, Substrate)
 
 		SHADER_PARAMETER_STRUCT_INCLUDE(FSceneTextureParameters, SceneTextures)
 	END_SHADER_PARAMETER_STRUCT()
@@ -669,6 +673,7 @@ void FDeferredShadingSceneRenderer::CompositeRayTracingSkyLight(
 		PassParameters->SkyLightTexture = SkyLightRT;
 		PassParameters->SkyLightTextureSampler = TStaticSamplerState<SF_Point, AM_Clamp, AM_Clamp, AM_Clamp>::GetRHI();
 		PassParameters->View = Views[ViewIndex].GetShaderParameters();
+		PassParameters->Substrate = Substrate::BindSubstrateGlobalUniformParameters(View);
 		PassParameters->RenderTargets[0] = FRenderTargetBinding(SceneTextures.Color.Target, ERenderTargetLoadAction::ELoad);
 		PassParameters->SceneTextures = SceneTextureParameters;
 
