Fix eye shading missing Subsurface profile with Substrate when using Format=0.

diff --git a/Engine/Shaders/Private/Substrate/Substrate.ush b/Engine/Shaders/Private/Substrate/Substrate.ush
--- a/Engine/Shaders/Private/Substrate/Substrate.ush
+++ b/Engine/Shaders/Private/Substrate/Substrate.ush
@@ -3676,17 +3676,13 @@ float3 SubstrateGetBSDFSubSurfaceColor(in FSubstrateBSDF BSDF)
 
 uint SubstrateGetBSDFSubSurfaceProfileID(in FSubstrateBSDF BSDF)
 {
-	switch (BSDF_GETTYPE(BSDF))
-	{
-	case SUBSTRATE_BSDF_TYPE_SLAB:
-	case SUBSTRATE_BSDF_TYPE_EYE:
+	if (BSDF_GETSSSTYPE(BSDF) == SSS_TYPE_DIFFUSION_PROFILE)
 	{
-		if (BSDF_GETSSSTYPE(BSDF) == SSS_TYPE_DIFFUSION_PROFILE)
+		switch (BSDF_GETTYPE(BSDF))
 		{
-			return SLAB_SSSPROFILEID(BSDF);
+			case SUBSTRATE_BSDF_TYPE_SLAB:	return SLAB_SSSPROFILEID(BSDF);
+			case SUBSTRATE_BSDF_TYPE_EYE:	return EYE_SSSPROFILEID(BSDF);
 		}
-		break;
-	}
 	}
 	return 0;
 }
