Add extra  info to GBuffer picking for easing comparison.

diff --git a/Engine/Shaders/Private/PostProcessGBufferHints.usf b/Engine/Shaders/Private/PostProcessGBufferHints.usf
--- a/Engine/Shaders/Private/PostProcessGBufferHints.usf
+++ b/Engine/Shaders/Private/PostProcessGBufferHints.usf
@@ -209,6 +209,12 @@ void MainCS(uint3 DispatchThreadId : SV_DispatchThreadID)
 		Print(Ctx, TEXT("Ind. Sample Occ.  : "), FontWhite); Print(Ctx, Data.DiffuseIndirectSampleOcclusion, FontYellow); Newline(Ctx, RectMax);
 		Print(Ctx, TEXT("Ind. Irradiance   : "), FontWhite); Print(Ctx, Data.IndirectIrradiance, FontYellow); Newline(Ctx, RectMax);
 		Print(Ctx, TEXT("Prec.Shadow Factor: "), FontWhite); Print(Ctx, Data.PrecomputedShadowFactors, FontYellow); Newline(Ctx, RectMax);
+		Newline(Ctx);
+		Print(Ctx, TEXT("CastContactShadow : "), FontWhite); PrintBool(Ctx, CastContactShadow(Data)); Newline(Ctx, RectMax);
+		Print(Ctx, TEXT("HasDynIndShadCast : "), FontWhite); PrintBool(Ctx, HasDynamicIndirectShadowCasterRepresentation(Data)); Newline(Ctx, RectMax);
+		Print(Ctx, TEXT("IsFirstPerson     : "), FontWhite); PrintBool(Ctx, IsFirstPerson(Data)); Newline(Ctx, RectMax);
+		Print(Ctx, TEXT("HasPrecShadowMask : "), FontWhite); PrintBool(Ctx, HasPrecShadowMask(Data)); Newline(Ctx, RectMax);
+		Print(Ctx, TEXT("HasZeroPrecShadow : "), FontWhite); PrintBool(Ctx, HasZeroPrecShadowMask(Data)); Newline(Ctx, RectMax);
 		PrintBackgroundRect(Ctx, RectMin, RectMax);
 		Newline(Ctx);
 
