RuntimeVirtualTextures could output PixelDepthOffset, but RVT never writes depth which causes issues on some platforms. Now disabled if VIRTUAL_TEXTURE_OUTPUT is enabled.

#was-backout 44632270

diff --git a/Engine/Shaders/Private/MaterialTemplate.ush b/Engine/Shaders/Private/MaterialTemplate.ush
--- a/Engine/Shaders/Private/MaterialTemplate.ush
+++ b/Engine/Shaders/Private/MaterialTemplate.ush
@@ -4663,7 +4663,7 @@ half3x3 AssembleTangentToWorld( half3 TangentToWorld0, half4 TangentToWorld2 )
 }
 
 // Whether the material shader should output pixel depth offset
-#define OUTPUT_PIXEL_DEPTH_OFFSET (WANT_PIXEL_DEPTH_OFFSET && !IS_NANITE_PASS && ((MATERIALBLENDING_SOLID || MATERIALBLENDING_MASKED) || (TRANSLUCENT_WRITING_VELOCITY)))
+#define OUTPUT_PIXEL_DEPTH_OFFSET (WANT_PIXEL_DEPTH_OFFSET && !VIRTUAL_TEXTURE_OUTPUT && !IS_NANITE_PASS && ((MATERIALBLENDING_SOLID || MATERIALBLENDING_MASKED) || (TRANSLUCENT_WRITING_VELOCITY)))
 
 // Whether to use the hidden d3d11 feature that supports depth writes with ZCull by only pushing into the screen
 //@todo - use for other SM5 platforms
diff --git a/Engine/Shaders/Public/ShaderVersion.ush b/Engine/Shaders/Public/ShaderVersion.ush
--- a/Engine/Shaders/Public/ShaderVersion.ush
+++ b/Engine/Shaders/Public/ShaderVersion.ush
@@ -3,4 +3,4 @@
 // Each time the console command is executed it generates a new GUID. As this file is included
 // in Platform.ush (which should be included in any shader) it allows to invalidate the shader DDC.
 // 
-#pragma message("UESHADERMETADATA_VERSION AF8CDE1B-768B-460D-91A6-51D9CA796EBB")
+#pragma message("UESHADERMETADATA_VERSION 8EC21872-0146-4BEB-9866-E40FBAC8C957")
