Fix Eye shading casting incorrectly contact shadow when Substrate is enabled with Format=1

diff --git a/Engine/Shaders/Private/Substrate/Substrate.ush b/Engine/Shaders/Private/Substrate/Substrate.ush
--- a/Engine/Shaders/Private/Substrate/Substrate.ush
+++ b/Engine/Shaders/Private/Substrate/Substrate.ush
@@ -2404,7 +2404,7 @@ struct FSubstratePixelHeader
 	void SetZeroPrecShadowMask(bool bIn)						{ }
 	void SetIsFirstPerson(bool bIn)								{ this.State |= (bIn ? HEADER_MASK_ISFIRSTPERSON : 0u); }
 #endif
-	void SetCastContactShadow(bool bIn)							{ this.State |= (bIn ? HEADER_MASK_CASTCONTACTSHADOW : 0u); }
+	void SetCastContactShadow(bool bIn)							{ this.State  = (this.State & ~HEADER_MASK_CASTCONTACTSHADOW) | (bIn ? HEADER_MASK_CASTCONTACTSHADOW : 0u); }
 	void SetHasSubsurface(bool bIn)								{ this.State |= (bIn ? HEADER_MASK_HASSUBSURFACE : 0u); }
 	void SetDynamicIndirectShadowCasterRepresentation(bool bIn)	{ this.State |= (bIn ? HEADER_MASK_HASDYNINDIRECTSHADOWCASTER : 0u); }
 
diff --git a/Engine/Shaders/Private/Substrate/SubstrateExport.ush b/Engine/Shaders/Private/Substrate/SubstrateExport.ush
--- a/Engine/Shaders/Private/Substrate/SubstrateExport.ush
+++ b/Engine/Shaders/Private/Substrate/SubstrateExport.ush
@@ -274,6 +274,7 @@ void PackSubstrateOut(
 						const bool bHasSSS = BSDF_GETSSSTYPE(BSDF) == SSS_TYPE_DIFFUSION_PROFILE;
 						SubstratePixelHeader.SetHasSubsurface(bHasSSS);
 						SubstratePixelHeader.SetMaterialMode(HEADER_MATERIALMODE_EYE);
+						SubstratePixelHeader.SetCastContactShadow(false); // Eye shading disables contact shadow. See CastContactShadow(...)
 						break;
 					}
 
