Fix missing shadow casting on eye with Substrate with Format=1

diff --git a/Engine/Shaders/Private/DeferredLightingCommon.ush b/Engine/Shaders/Private/DeferredLightingCommon.ush
--- a/Engine/Shaders/Private/DeferredLightingCommon.ush
+++ b/Engine/Shaders/Private/DeferredLightingCommon.ush
@@ -80,7 +80,7 @@ float ShadowRayCast(
 		FSubstratePixelHeader SubstratePixelHeader = UnpackSubstrateHeaderIn(SubstrateRayCast.MaterialTextureArray, SubstrateAddressing, SubstrateRayCast.TopLayerTexture);
 		// We check if the hit pixel has CastContactShadow, and we always exclude EYE BSDFs from there as done by CastContactShadow(SampleGBuffer) in the legacy path below.
 		// SUBSTRATE_TODO: try to move that test to GBuffer export pass (set DoesCastContactShadow = false for EYE BSDFs).
-		bOutHitCastContactShadow = SubstratePixelHeader.DoesCastContactShadow() && !SubstratePixelHeader.SubstrateGetBSDFType() == SUBSTRATE_BSDF_TYPE_EYE;
+		bOutHitCastContactShadow = SubstratePixelHeader.DoesCastContactShadow() && SubstratePixelHeader.SubstrateGetBSDFType() != SUBSTRATE_BSDF_TYPE_EYE;
 #endif
 #else
 		FGBufferData SampleGBuffer = GetGBufferData(HitUV);
