Area light shading refactor. Area lights now provide DiffuseL and SpecularL which can be used for BxDFs.

diff --git a/Engine/Shaders/Private/AreaLightCommon.ush b/Engine/Shaders/Private/AreaLightCommon.ush
--- a/Engine/Shaders/Private/AreaLightCommon.ush
+++ b/Engine/Shaders/Private/AreaLightCommon.ush
@@ -10,7 +10,11 @@ struct FAreaLight
 	float		SphereSinAlphaSoft;
 	float		LineCosSubtended;
 
-	float3		FalloffColor;
+	half3		DiffuseL;	// Vector irradiance
+	half3		SpecularL;	// Closest to reflection vector
+	half		NoL;
+	float		Falloff;
+	half3		FalloffColor;
 
 	FRect		Rect;
 	FRectTexture Texture;
@@ -18,14 +22,6 @@ struct FAreaLight
 	uint		IsRectAndDiffuseMicroReflWeight;
 };
 
-struct FAreaLightIntegrateContext
-{
-	FAreaLight AreaLight;
-	float3 L;
-	float NoL;
-	float Falloff;
-};
-
 void SetIsRectLight(inout FAreaLight AreaLight, bool bIsRectLight)
 {
 	AreaLight.IsRectAndDiffuseMicroReflWeight = (AreaLight.IsRectAndDiffuseMicroReflWeight & 0xFFFFFFFE) | (bIsRectLight ? 0x1 : 0x0);
@@ -52,19 +48,19 @@ bool IsAreaLight(FAreaLight AreaLight)
 	return IsRectLight(AreaLight) || GetAreaLightDiffuseMicroReflWeight(AreaLight) < 1.0f;
 }
 
-FAreaLightIntegrateContext InitAreaLightIntegrateContext()
+FAreaLight InitAreaLight()
 {
-	// Manual initialization of the area light as the compiler is unhappy otherwise
-	FAreaLightIntegrateContext Out;
-	Out.AreaLight.SphereSinAlpha = 0;
-	Out.AreaLight.SphereSinAlphaSoft = 0;
-	Out.AreaLight.LineCosSubtended = 0;
-	Out.AreaLight.FalloffColor = 0;
-	Out.AreaLight.Rect = (FRect)0;
-	Out.AreaLight.IsRectAndDiffuseMicroReflWeight = 0;
-	Out.AreaLight.Texture = InitRectTexture();
-	Out.L = 0;
-	Out.NoL = 0;
-	Out.Falloff = 0;
-	return Out;
+	FAreaLight AreaLight;
+	AreaLight.SphereSinAlpha = 0;
+	AreaLight.SphereSinAlphaSoft = 0;
+	AreaLight.LineCosSubtended = 1;
+	AreaLight.DiffuseL = 0;
+	AreaLight.SpecularL = 0;
+	AreaLight.NoL = 1;
+	AreaLight.Falloff = 1;
+	AreaLight.FalloffColor = 1;
+	AreaLight.Rect = (FRect)0;
+	AreaLight.IsRectAndDiffuseMicroReflWeight = 0;
+	AreaLight.Texture = InitRectTexture();
+	return AreaLight;
 }
diff --git a/Engine/Shaders/Private/CapsuleLight.ush b/Engine/Shaders/Private/CapsuleLight.ush
--- a/Engine/Shaders/Private/CapsuleLight.ush
+++ b/Engine/Shaders/Private/CapsuleLight.ush
@@ -2,23 +2,6 @@
 
 #pragma once
 
-float SmoothMin( float a, float b, float k )
-{
-	float h = saturate( 0.5 + (0.5 / k) * (b - a) );
-	return lerp( b, a, h ) - k * (h - h*h);
-}
-
-float SmoothMax( float a, float b, float k )
-{
-	return SmoothMin( a, b, -k );
-}
-
-float SmoothClamp( float x, float Min, float Max, float k )
-{
-	return SmoothMin( SmoothMax( x, Min, k ), Max, k );
-	//return min( max( x, Min ), Max );
-}
-
 struct FCapsuleLight
 {
 	float3	LightPos[2];
@@ -41,7 +24,6 @@ float3 ClosestPointLineToPoint( float3 Line0, float3 Line1, float Length )
 {
 	float3 Line01 = Line1 - Line0;
 	return Line0 + Line01 * saturate( -dot( Line01, Line0 ) / Pow2( Length ) );
-	//return Line0 + Line01 * SmoothClamp( -dot( Line01, Line0 ) / Pow2( Length ), 0, 1, 0.5 );
 }
 
 // Closest point on line segment to ray
@@ -75,7 +57,7 @@ float3 SmallestAnglePointOnLineToRay( float3 Line0, float3 Line1, float Length,
 	return Line0 + t * Line01;
 }
 
-float3 LineIrradiance( float3 N, float3 Line0, float3 Line1, float DistanceBiasSqr, out float CosSubtended, out float BaseIrradiance, out float NoL )
+void LineIrradiance( float3 Line0, float3 Line1, float DistanceBiasSqr, out float CosSubtended, out float Falloff, out half3 L )
 {
 	float LengthSqr0 = dot( Line0, Line0 );
 	float LengthSqr1 = dot( Line1, Line1 );
@@ -84,11 +66,8 @@ float3 LineIrradiance( float3 N, float3 Line0, float3 Line1, float DistanceBiasS
 	float InvLength01 = InvLength0 * InvLength1;
 
 	CosSubtended = dot( Line0, Line1 ) * InvLength01;
-	BaseIrradiance = InvLength01 / ( CosSubtended * 0.5 + 0.5 + DistanceBiasSqr * InvLength01 );
-	NoL = 0.5 * ( dot(N, Line0) * InvLength0 + dot(N, Line1) * InvLength1 );
-	
-	float3 VectorIrradiance = ( BaseIrradiance * 0.5 ) * ( Line0 * InvLength0 + Line1 * InvLength1 );
-	return VectorIrradiance;
+	Falloff = InvLength01 / ( CosSubtended * 0.5 + 0.5 + DistanceBiasSqr * InvLength01 );
+	L = 0.5 * ( Line0 * InvLength0 + Line1 * InvLength1 );
 }
 
 // Alpha is half of angle of spherical cap
diff --git a/Engine/Shaders/Private/CapsuleLightIntegrate.ush b/Engine/Shaders/Private/CapsuleLightIntegrate.ush
--- a/Engine/Shaders/Private/CapsuleLightIntegrate.ush
+++ b/Engine/Shaders/Private/CapsuleLightIntegrate.ush
@@ -10,7 +10,6 @@
 #include "CapsuleLightSampling.ush"
 #include "SobolRandom.ush"
 
-// Should this be SH instead?
 float IntegrateLight( FCapsuleLight Capsule, bool bInverseSquared )
 {
 	float Falloff;
@@ -18,9 +17,9 @@ float IntegrateLight( FCapsuleLight Capsule, bool bInverseSquared )
 	BRANCH
 	if( Capsule.Length > 0 )
 	{
-		float NoL;
 		float LineCosSubtended = 1;
-		LineIrradiance( 0, Capsule.LightPos[0], Capsule.LightPos[1], Capsule.DistBiasSqr, LineCosSubtended, Falloff, NoL );
+		half3 L;
+		LineIrradiance( Capsule.LightPos[0], Capsule.LightPos[1], Capsule.DistBiasSqr, LineCosSubtended, Falloff, L );
 	}
 	else
 	{
@@ -34,13 +33,9 @@ float IntegrateLight( FCapsuleLight Capsule, bool bInverseSquared )
 	return Falloff;
 }
 
-FAreaLightIntegrateContext CreateCapsuleIntegrateContext(float Roughness, half3 N, half3 V, FCapsuleLight Capsule, bool bInverseSquared )
+FAreaLight CreateAreaLight(float Roughness, half3 N, half3 V, FCapsuleLight Capsule, bool bInverseSquared )
 {
-	FAreaLightIntegrateContext Out = InitAreaLightIntegrateContext();
-
-	float NoL;
-	float Falloff;
-	float LineCosSubtended = 1;
+	FAreaLight AreaLight = InitAreaLight();
 
 	// Clip to horizon
 	//float NoP0 = dot( N, Capsule.LightPos[0] );
@@ -51,26 +46,28 @@ FAreaLightIntegrateContext CreateCapsuleIntegrateContext(float Roughness, half3
 	BRANCH
 	if( Capsule.Length > 0 )
 	{
-		LineIrradiance( N, Capsule.LightPos[0], Capsule.LightPos[1], Capsule.DistBiasSqr, LineCosSubtended, Falloff, NoL );
+		LineIrradiance( Capsule.LightPos[0], Capsule.LightPos[1], Capsule.DistBiasSqr, AreaLight.LineCosSubtended, AreaLight.Falloff, AreaLight.DiffuseL );
+		AreaLight.NoL = dot( N, AreaLight.DiffuseL );
+		AreaLight.DiffuseL = normalize( AreaLight.DiffuseL );
 	}
 	else
 	{
 		float DistSqr = dot( Capsule.LightPos[0], Capsule.LightPos[0] );
-		Falloff = rcp( DistSqr + Capsule.DistBiasSqr );
+		AreaLight.Falloff = rcp( DistSqr + Capsule.DistBiasSqr );
 
-		float3 L = Capsule.LightPos[0] * rsqrt( DistSqr );
-		NoL = dot( N, L );
+		AreaLight.DiffuseL = Capsule.LightPos[0] * rsqrt( DistSqr );
+		AreaLight.NoL = dot( N, AreaLight.DiffuseL );
 	}
 
 	if( Capsule.Radius > 0 )
 	{
 		// TODO Use capsule area?
-		float SinAlphaSqr = saturate( Pow2( Capsule.Radius ) * Falloff );
-		NoL = SphereHorizonCosWrap( NoL, SinAlphaSqr );
+		float SinAlphaSqr = saturate( Pow2( Capsule.Radius ) * AreaLight.Falloff );
+		AreaLight.NoL = SphereHorizonCosWrap( AreaLight.NoL, SinAlphaSqr );
 	}
 
-	NoL = saturate( NoL );
-	Falloff = bInverseSquared ? Falloff : 1;
+	AreaLight.NoL = saturate( AreaLight.NoL );
+	AreaLight.Falloff = bInverseSquared ? AreaLight.Falloff : 1;
 
 	float3 ToLight = Capsule.LightPos[0];
 	if( Capsule.Length > 0 )
@@ -90,36 +87,29 @@ FAreaLightIntegrateContext CreateCapsuleIntegrateContext(float Roughness, half3
 
 	float DistSqr = dot( ToLight, ToLight );
 	float InvDist = rsqrt( DistSqr );
-	float3 L = ToLight * InvDist;
+	AreaLight.SpecularL = ToLight * InvDist;
 	
-	Roughness = max( Roughness, View.MinRoughness );
 	float a = Pow2( Roughness );
     
-	// Diffuse micro refelction contribution will softly fade out when th light become an area light.
+	// Diffuse micro reflection contribution will softly fade out when th light become an area light.
 	// We only based this assumption based on the light size.
 	const float SizeFadesOutDiffuseMicroRefl = 20.0;
 
-	Out.AreaLight.SphereSinAlpha = saturate( Capsule.Radius * InvDist * (1 - a) );
-	Out.AreaLight.SphereSinAlphaSoft = saturate( Capsule.SoftRadius * InvDist );
-	Out.AreaLight.LineCosSubtended = LineCosSubtended;
-	Out.AreaLight.FalloffColor = 1;
-	Out.AreaLight.Rect = (FRect)0;
-	Out.AreaLight.Texture = InitRectTexture();
-	Out.AreaLight.IsRectAndDiffuseMicroReflWeight = 0;
-	SetIsRectLight(Out.AreaLight, false);
-	SetAreaLightDiffuseMicroReflWeight(Out.AreaLight, saturate(1.0f - max(Capsule.Length, Capsule.Radius) / SizeFadesOutDiffuseMicroRefl));
-	Out.NoL = NoL;
-	Out.Falloff = Falloff;	
-	Out.L = L;
-	return Out;
+	AreaLight.SphereSinAlpha = saturate( Capsule.Radius * InvDist * (1 - a) );
+	AreaLight.SphereSinAlphaSoft = saturate( Capsule.SoftRadius * InvDist );
+
+	SetIsRectLight(AreaLight, false);
+	SetAreaLightDiffuseMicroReflWeight(AreaLight, saturate(1.0f - max(Capsule.Length, Capsule.Radius) / SizeFadesOutDiffuseMicroRefl));
+
+	return AreaLight;
 }
 
 
 FDirectLighting IntegrateBxDF( FGBufferData GBuffer, half3 N, half3 V, FCapsuleLight Capsule, FShadowTerms Shadow, bool bInverseSquared )
 {
 	GBuffer.Roughness = max( GBuffer.Roughness, View.MinRoughness );
-	FAreaLightIntegrateContext Context = CreateCapsuleIntegrateContext(GBuffer.Roughness, N, V, Capsule, bInverseSquared);
-	return IntegrateBxDF( GBuffer, N, V, Context.L, Context.Falloff, Context.NoL, Context.AreaLight, Shadow );
+	FAreaLight AreaLight = CreateAreaLight(GBuffer.Roughness, N, V, Capsule, bInverseSquared);
+	return IntegrateBxDF( GBuffer, N, V, AreaLight, Shadow );
 }
 
 FDirectLighting IntegrateBxDF( FGBufferData GBuffer, half3 N, half3 V, FCapsuleLight Capsule, FShadowTerms Shadow, uint2 SVPos )
@@ -150,20 +140,16 @@ FDirectLighting IntegrateBxDF( FGBufferData GBuffer, half3 N, half3 V, FCapsuleL
 		16,	// Light area
 	};
 	
-	uint2 SobolBase = SobolPixel( SVPos );
-	uint2 SobolFrame = SobolIndex( SobolBase, View.StateFrameIndexMod8, 3 );
-	
 	UNROLL
 	for( uint Set = 0; Set < NumSets; Set++ )
 	{
 		LOOP
 		for( uint i = 0; i < NumSamples[ Set ]; i++ )
 		{
-			uint2 Random = Rand3DPCG16( uint3( SVPos.xy, View.Random ^ Set ) ).xy;
+			FRandomSequence RandSequence = RandomSequenceCreate( uint3( SVPos.xy, View.StateFrameIndex ), i, NumSamples[ Set ] );
+			RandSequence.SampleSeed = StrongIntegerHash( Set );
 
-			float2 E = float2( SobolIndex( SobolFrame, i << 3 ) ) / 0x10000;
-			//float2 E = Hammersley( i, NumSamples[ Set ], Random );
-			//float2 E = CorrelatedMultiJitter2D( i, NumSamples[ Set ], Random.x );
+			float2 E = RandSequence.Get2D();
 			
 			float3 L, H;
 			if( Set == 0 )
@@ -243,10 +229,13 @@ FDirectLighting IntegrateBxDF( FGBufferData GBuffer, half3 N, half3 V, FCapsuleL
 			float NoH = saturate( dot(N, H) );
 			float VoH = saturate( dot(V, H) );
 
-			if( NoL > 0 && VoH > 0 )
+			//if( NoL > 0 && VoH > 0 )
+			if( VoH > 0 )
 			{
 				{
-					if (CapsuleTest(L, ToLight, CapsuleAxis, Pow2(Capsule.Radius), Capsule.Length) > 0)
+					float3 ToSphere = ClosestPointLineToRay( Capsule.LightPos[0], Capsule.LightPos[1], Capsule.Length, L );
+
+					if( length2( cross( L, ToSphere ) ) > Pow2( Capsule.Radius ) )
 					{
 						// Ray misses sphere
 						continue;
diff --git a/Engine/Shaders/Private/DeferredLightingCommon.ush b/Engine/Shaders/Private/DeferredLightingCommon.ush
--- a/Engine/Shaders/Private/DeferredLightingCommon.ush
+++ b/Engine/Shaders/Private/DeferredLightingCommon.ush
@@ -546,7 +546,7 @@ float3 GetIrradianceForLight(FDeferredLightData LightData, float3 WorldNormal,	f
 
 	if (LightData.bRadialLight)
 	{
-		FAreaLightIntegrateContext Context = (FAreaLightIntegrateContext) 0;
+		FAreaLight AreaLight = InitAreaLight();
 		float LightMask = GetLocalLightAttenuation(TranslatedWorldPosition, LightData, ToLight, L);
 
 		float Attenuation = 0.0f;
@@ -561,20 +561,20 @@ float3 GetIrradianceForLight(FDeferredLightData LightData, float3 WorldNormal,	f
 			if (IsRectVisible(Rect))
 			{
 				const FRectTexture SourceTexture = ConvertToRectTexture(LightData);
-				Context = CreateRectIntegrateContext(Roughness, WorldNormal, V, Rect, SourceTexture);
+				AreaLight = CreateAreaLight(Roughness, WorldNormal, V, Rect, SourceTexture);
 			}
 		}
 		else
 		{
 			FCapsuleLight Capsule = GetCapsule(ToLight, LightData);
 			Capsule.DistBiasSqr = 0;
-			Context = CreateCapsuleIntegrateContext(Roughness, WorldNormal, V, Capsule, LightData.bInverseSquared);
+			AreaLight = CreateAreaLight(Roughness, WorldNormal, V, Capsule, LightData.bInverseSquared);
 			Attenuation = IntegrateLight(Capsule, LightData.bInverseSquared);
 		}
 
 		CombinedAttenuation = Attenuation * LightMask;
-		AreaLightFalloffColor = Context.AreaLight.FalloffColor;
-		NoL = Context.NoL;
+		AreaLightFalloffColor = AreaLight.FalloffColor;
+		NoL = AreaLight.NoL;
 	}
 
 	return LightColor * AreaLightFalloffColor * (CombinedAttenuation * NoL * LightData.DiffuseScale);
diff --git a/Engine/Shaders/Private/RectLight.ush b/Engine/Shaders/Private/RectLight.ush
--- a/Engine/Shaders/Private/RectLight.ush
+++ b/Engine/Shaders/Private/RectLight.ush
@@ -159,139 +159,6 @@ float3 RectIrradianceLambert( float3 N, FRect Rect, out float BaseIrradiance, ou
 	return L;
 }
 
-float3 RectIrradianceApproxKaris( float3 N, FRect Rect, out float BaseIrradiance, out float NoL )
-{
-	float2 RectLocal;
-	RectLocal.x = SmoothClamp( dot( Rect.Axis[0], -Rect.Origin ), -Rect.Extent.x, Rect.Extent.x, 16 );
-	RectLocal.y = SmoothClamp( dot( Rect.Axis[1], -Rect.Origin ), -Rect.Extent.y, Rect.Extent.y, 16 );
-
-	float3 ClosestPoint = Rect.Origin;
-	ClosestPoint += Rect.Axis[0] * RectLocal.x;
-	ClosestPoint += Rect.Axis[1] * RectLocal.y;
-
-	float3 OppositePoint = 2 * Rect.Origin - ClosestPoint;
-	
-	float3 L0 = normalize( ClosestPoint );
-	float3 L1 = normalize( OppositePoint );
-	float3 L  = normalize( L0 + L1 );
-
-	// Intersect ray with plane
-	float Distance = dot( Rect.Axis[2], Rect.Origin ) / dot( Rect.Axis[2], L );
-	float DistanceSqr = Distance * Distance;
-
-	// right pyramid solid angle cosine weighted approx
-	//float Irradiance = 4 * RectExtent.x * RectExtent.y * rsqrt( ( Square( RectExtent.x ) + DistanceSqr ) * ( Square( RectExtent.y ) + DistanceSqr ) );
-	BaseIrradiance = 4 * Rect.Extent.x * Rect.Extent.y * rsqrt( ( (4 / PI) * Square( Rect.Extent.x ) + DistanceSqr ) * ( (4 / PI) * Square( Rect.Extent.y ) + DistanceSqr ) );
-	BaseIrradiance *= saturate( dot( Rect.Axis[2], L ) );
-
-	// Solid angle of sphere		= 2*PI * ( 1 - sqrt(1 - r^2 / d^2 ) )
-	// Cosine weighted integration	= PI * r^2 / d^2
-	// SinAlphaSqr = r^2 / d^2;
-	float SinAlphaSqr = BaseIrradiance * (1.0 / PI);
-
-	NoL = SphereHorizonCosWrap( dot( N, L ), SinAlphaSqr );
-
-	return L;
-}
-
-float3 RectIrradianceApproxLagarde( float3 N, FRect Rect, out float BaseIrradiance, out float NoL )
-{
-	float3 L = normalize( Rect.Origin );
-
-	float3 v0 = Rect.Origin - Rect.Axis[0] * Rect.Extent.x - Rect.Axis[1] * Rect.Extent.y;
-	float3 v1 = Rect.Origin + Rect.Axis[0] * Rect.Extent.x - Rect.Axis[1] * Rect.Extent.y;
-	float3 v2 = Rect.Origin + Rect.Axis[0] * Rect.Extent.x + Rect.Axis[1] * Rect.Extent.y;
-	float3 v3 = Rect.Origin - Rect.Axis[0] * Rect.Extent.x + Rect.Axis[1] * Rect.Extent.y;
-
-	float3 n0 = normalize( cross( v0, v1 ) );
-	float3 n1 = normalize( cross( v1, v2 ) );
-	float3 n2 = normalize( cross( v2, v3 ) );
-	float3 n3 = normalize( cross( v3, v0 ) );
-
-	float g0 = acos( dot( n0, n1 ) );
-	float g1 = acos( dot( n1, n2 ) );
-	float g2 = acos( dot( n2, n3 ) );
-	float g3 = acos( dot( n3, n0 ) );
-
-	// Solid angle
-	BaseIrradiance = g0 + g1 + g2 + g3 - 2*PI;
-
-	NoL = 0.2 * ( saturate( dot( N, L ) ) +
-		saturate( dot( N, normalize(v0) ) ) +
-		saturate( dot( N, normalize(v1) ) ) +
-		saturate( dot( N, normalize(v2) ) ) +
-		saturate( dot( N, normalize(v3) ) ) );
-
-	return L;
-}
-
-float3 RectIrradianceApproxDrobot( float3 N, FRect Rect, out float BaseIrradiance, out float NoL )
-{
-#if 0
-	// Drobot complex
-	float3 d0 = Rect.Origin;
-	d0 += RectX * clamp( dot( Rect.Axis[0], -Rect.Origin ), -Rect.Extent.x, Rect.Extent.x );
-	d0 += RectY * clamp( dot( Rect.Axis[1], -Rect.Origin ), -Rect.Extent.y, Rect.Extent.y );
-
-	float3 d1 = N - Rect.Axis[2] * saturate( 0.001 + dot( Rect.Axis[2], N ) );
-
-	d0 = normalize( d0 );
-	d1 = normalize( d1 );
-	float3 dh = normalize( d0 + d1 );
-#else
-	// Drobot simple
-	float clampCosAngle = 0.001 + saturate( dot( N, Rect.Axis[2] ) );
-	// clamp d0 to the positive hemisphere of surface normal
-	float3 d0 = normalize( -Rect.Axis[2] + N * clampCosAngle );
-	// clamp d1 to the negative hemisphere of light plane normal
-	float3 d1 = normalize( N - Rect.Axis[2] * clampCosAngle );
-	float3 dh = normalize( d0 + d1 );
-#endif
-
-	// Intersect ray with plane
-	float3 PointOnPlane = dh * ( dot( Rect.Axis[2], Rect.Origin ) / dot( Rect.Axis[2], dh ) );
-
-	float3 ClosestPoint = Rect.Origin;
-	ClosestPoint += Rect.Axis[0] * clamp( dot( Rect.Axis[0], PointOnPlane - Rect.Origin ), -Rect.Extent.x, Rect.Extent.x );
-	ClosestPoint += Rect.Axis[1] * clamp( dot( Rect.Axis[1], PointOnPlane - Rect.Origin ), -Rect.Extent.y, Rect.Extent.y );
-
-	float3 v0 = Rect.Origin - Rect.Axis[0] * Rect.Extent.x - Rect.Axis[1] * Rect.Extent.y;
-	float3 v1 = Rect.Origin + Rect.Axis[0] * Rect.Extent.x - Rect.Axis[1] * Rect.Extent.y;
-	float3 v2 = Rect.Origin + Rect.Axis[0] * Rect.Extent.x + Rect.Axis[1] * Rect.Extent.y;
-	float3 v3 = Rect.Origin - Rect.Axis[0] * Rect.Extent.x + Rect.Axis[1] * Rect.Extent.y;
-#if 1
-	float3 n0 = normalize( cross( v0, v1 ) );
-	float3 n1 = normalize( cross( v1, v2 ) );
-	float3 n2 = normalize( cross( v2, v3 ) );
-	float3 n3 = normalize( cross( v3, v0 ) );
-
-	float g0 = acos( dot( n0, n1 ) );
-	float g1 = acos( dot( n1, n2 ) );
-	float g2 = acos( dot( n2, n3 ) );
-	float g3 = acos( dot( n3, n0 ) );
-
-	float SolidAngle = g0 + g1 + g2 + g3 - 2*PI;
-
-	float3 L = normalize( ClosestPoint );
-#else
-	float DistanceSqr = dot( RectOrigin, RectOrigin );
-	float3 L = RectOrigin * rsqrt( DistanceSqr );
-
-	DistanceSqr = dot( ClosestPoint, ClosestPoint );
-	L = ClosestPoint * rsqrt( DistanceSqr );
-
-	float SolidAngle = PI / ( DistanceSqr / ( (4 / PI) * Rect.Extent.x * Rect.Extent.y ) + 1 );
-	//float SolidAngle = 4 * asin( RectExtent.x * RectExtent.y / sqrt( ( Square( RectExtent.x ) + DistanceSqr ) * ( Square(RectExtent.y) + DistanceSqr ) ) );
-	SolidAngle *= saturate( dot( -Rect.Axis[2], L ) );
-#endif
-
-	BaseIrradiance = SolidAngle;
-	NoL = saturate( dot( N, L ) );
-
-	return L;
-}
-
-
 float3 SampleSourceTexture( float3 L, FRect Rect, FRectTexture RectTexture)
 {
 #if USE_SOURCE_TEXTURE
diff --git a/Engine/Shaders/Private/RectLightIntegrate.ush b/Engine/Shaders/Private/RectLightIntegrate.ush
--- a/Engine/Shaders/Private/RectLightIntegrate.ush
+++ b/Engine/Shaders/Private/RectLightIntegrate.ush
@@ -1,7 +1,5 @@
 // Copyright Epic Games, Inc. All Rights Reserved.
 
-#define CHEAP_RECT	0
-
 #include "DeferredShadingCommon.ush"
 #include "MonteCarlo.ush"
 #include "AreaLightCommon.ush"
@@ -52,200 +50,27 @@ float IntegrateLight( FRect Rect )
 	float NoL;
 	float Falloff;
 
-#if !CHEAP_RECT	// Optimized Lambert
+	// Optimized Lambert
 	float3 L = RectIrradianceLambert( 0, Rect, Falloff, NoL );
-#elif 1	// Karis
-	float3 L = RectIrradianceApproxKaris( 0, Rect, Falloff, NoL );
-#elif 1	// Lagarde
-	float3 L = RectIrradianceApproxLagarde( 0, Rect, Falloff, NoL );
-#else	// Drobot
-	float3 L = RectIrradianceApproxDrobot( 0, Rect, Falloff, NoL );
-#endif
 	
 	return Falloff;
 }
 
-FAreaLightIntegrateContext CreateRectIntegrateContext( float Roughness, half3 N, half3 V, FRect Rect, FRectTexture SourceTexture )
+FAreaLight CreateAreaLight( float Roughness, half3 N, half3 V, FRect Rect, FRectTexture SourceTexture )
 {
-	float NoL = 0;
-	float Falloff = 0;
-	
-	FAreaLightIntegrateContext Out = InitAreaLightIntegrateContext();
-
-#if !CHEAP_RECT	// Optimized Lambert
-	float3 L = RectIrradianceLambert( N, Rect, Falloff, NoL );
-#elif 1	// Karis
-	float3 L = RectIrradianceApproxKaris( N, Rect, Falloff, NoL );
-#elif 1	// Lagarde
-	float3 L = RectIrradianceApproxLagarde( N, Rect, Falloff, NoL );
-#else	// Drobot
-	float3 L = RectIrradianceApproxDrobot( N, Rect, Falloff, NoL );
-#endif
-
-#if CHEAP_RECT
-	float3 R = reflect( -V, N );
-
-#if 0
-	float NoV = saturate( dot( N, V ) );
-
-	L = lerp( L, R, saturate( -2*dot( Rect.Axis[2], R ) ) );
-	L = normalize( L );
-
-	UNROLL
-	for( int k = 0; k < 2; k++ )
-	{
-#if 1
-		float NoL = dot( N, L );
-		float NoV = dot( N, V );
-		float VoL = dot( V, L );
-		float NoHInvLenH = ( NoL + NoV ) / ( 2 + 2*VoL );
-		float3 Gradient = ( ( N - L*NoL ) - ( V - L*VoL ) * NoHInvLenH ) * (2*NoHInvLenH);
-#else
-		float RoL = dot( R, L );
-		float3 Gradient = 2 * RoL * ( R - L * RoL );
-#endif
-
-		Gradient -= Rect.Axis[2] * dot( Rect.Axis[2], Gradient );
-		Gradient = lerp( Gradient, 0, (1 - NoV) * saturate( 2*dot( Rect.Axis[2], L ) ) );
-
-		L = ClampToRect( L + Gradient * ( 2.0 / ( 2.0 + k ) ), Rect );
-	}
-#elif 1
-	float3 Ls = L;
-
-	float3 v[4];
-	v[0] = Rect.Origin - Rect.Axis[0] * Rect.Extent.x - Rect.Axis[1] * Rect.Extent.y;
-	v[1] = Rect.Origin + Rect.Axis[0] * Rect.Extent.x - Rect.Axis[1] * Rect.Extent.y;
-	v[2] = Rect.Origin + Rect.Axis[0] * Rect.Extent.x + Rect.Axis[1] * Rect.Extent.y;
-	v[3] = Rect.Origin - Rect.Axis[0] * Rect.Extent.x + Rect.Axis[1] * Rect.Extent.y;
-
-	float3 e0 = v[0];
-	float3 e1 = v[1];
+	FAreaLight AreaLight = InitAreaLight();
 
-	float3 MinEdgeN = 0;
-	float  MinEdgeCos = 1;
-
-	UNROLL
-	for( uint i = 0; i < 4; i++ )
-	{
-		float3 v0 = v[i];
-		float3 v1 = v[ (i+1) % 4 ];
-
-		float3 EdgeN = normalize( cross( v0, v1 ) );
-		float  EdgeCos = dot( R, EdgeN );
-
-		if( EdgeCos < MinEdgeCos )
-		{
-			MinEdgeN   = EdgeN;
-			MinEdgeCos = EdgeCos;
-
-			e0 = v0;
-			e1 = v1;
-		}
-	}
-
-	if( MinEdgeCos > 0 )
-	{
-		Ls = R;
-	}
-	else
-	{
-#if 0
-		Ls = SmallestAnglePointOnLineToRay( e0, e1, length( e0 - e1 ), R );
-#else
-		float3 Rp = R - MinEdgeCos * MinEdgeN;
-		if(		dot( cross( Rp, e0 ), R ) < 0 )	Ls = e0;
-		else if(dot( cross( e1, Rp ), R ) < 0 )	Ls = e1;
-		else									Ls = Rp;
-#endif
-		Ls = normalize( Ls );
-	}
-
-	float a = Pow2( GBuffer.Roughness );
-	//L = lerp( Ls, L, a );
-	L = normalize( Ls );
-#else
-	L = R;
-	if( !RayHitRect( R, Rect ) )
-	{
-		float3 MaxL = R;
-		float  MaxNoH = -1;
-
-		uint NumSteps = 128;
-		for( uint i = 0; i < NumSteps; i++ )
-		{
-			float Theta = (2*PI) * i / (float)NumSteps;
-			float2 p;
-			p.x = cos( Theta );
-			p.y = sin( Theta );
-
-			p.xy /= dot( 1, abs(p) );
-			float2 PointInRect = float2( p.x + p.y, p.x - p.y ) * Rect.Extent;
-
-			//0.5 * sqrt( 2 + (2*sqrt(2.0)) * 
-
-			float3 ToRect = Rect.Origin;
-			ToRect += PointInRect.x * Rect.Axis[0];
-			ToRect += PointInRect.y * Rect.Axis[1];
-
-			L = normalize( ToRect );
-			
-			float RoL = dot( R, L );
-
-			BxDFContext Context;
-			Context.Init( N, V, L );
-
-			if( Context.NoH > MaxNoH )
-			{
-				MaxNoH = Context.NoH;
-				MaxL = L;
-			}
-		}
-
-		L = MaxL;
-	}
+	AreaLight.DiffuseL = RectIrradianceLambert( N, Rect, AreaLight.Falloff, AreaLight.NoL );
+	AreaLight.SpecularL = AreaLight.DiffuseL;
+	AreaLight.FalloffColor = SampleSourceTexture( AreaLight.DiffuseL, Rect, SourceTexture );
+	
+	AreaLight.Rect = Rect;
+	AreaLight.Texture = SourceTexture;
 
-	for( int k = 0; k < 0; k++ )
-	{
-		float NoL = dot( N, L );
-		float NoV = dot( N, V );
-		float VoL = dot( V, L );
-		float NoHInvLenH = ( NoL + NoV ) / ( 2 + 2*VoL );
-		float3 Gradient = ( ( N - L*NoL ) - ( V - L*VoL ) * NoHInvLenH ) * (2*NoHInvLenH);
+	SetIsRectLight(AreaLight, true);
+	SetAreaLightDiffuseMicroReflWeight(AreaLight, 0.0);
 
-		L = ClampToRect( L + Gradient * ( 2.0 / ( 2.0 + k ) ), Rect );
-	}
-#endif
-	
-	Out.AreaLight.SphereSinAlpha = sqrt( Falloff * (1.0 / PI) );
-	Out.AreaLight.SphereSinAlphaSoft = 0;
-	Out.AreaLight.LineCosSubtended = 1;
-	Out.AreaLight.FalloffColor = 1;
-	Out.AreaLight.Rect = Rect;
-	Out.AreaLight.Texture = SourceTexture;
-	Out.AreaLight.IsRectAndDiffuseMicroReflWeight = 0;
-	SetIsRectLight(Out.AreaLight, false);
-	SetAreaLightDiffuseMicroReflWeight(Out.AreaLight, 0.0);
-	Out.L = L;
-	Out.NoL = NoL;
-	Out.Falloff = Falloff;
-#else
-	float3 FalloffColor = SampleSourceTexture( L, Rect, SourceTexture );
-
-	Out.AreaLight.SphereSinAlpha = 0;
-	Out.AreaLight.SphereSinAlphaSoft = 0;
-	Out.AreaLight.LineCosSubtended = 1;
-	Out.AreaLight.FalloffColor = FalloffColor;
-	Out.AreaLight.Rect = Rect;
-	Out.AreaLight.Texture = SourceTexture;
-	Out.AreaLight.IsRectAndDiffuseMicroReflWeight = 0;
-	SetIsRectLight(Out.AreaLight, true);
-	SetAreaLightDiffuseMicroReflWeight(Out.AreaLight, 0.0);
-	Out.L = L;
-	Out.NoL = NoL;
-	Out.Falloff = Falloff;
-#endif
-	return Out;
+	return AreaLight;
 }
 
 FDirectLighting IntegrateBxDF(FGBufferData GBuffer, half3 N, half3 V, FRect Rect, FShadowTerms Shadow, FRectTexture SourceTexture)
@@ -258,9 +83,9 @@ FDirectLighting IntegrateBxDF(FGBufferData GBuffer, half3 N, half3 V, FRect Rect
 	FDirectLighting Out = (FDirectLighting)0;
 	if (IsRectVisible(Rect) && Distance > 0.001f)
 	{
-		FAreaLightIntegrateContext Context = CreateRectIntegrateContext(GBuffer.Roughness, N, V, Rect, SourceTexture);
+		FAreaLight AreaLight = CreateAreaLight(GBuffer.Roughness, N, V, Rect, SourceTexture);
 		GBuffer.Roughness = max(GBuffer.Roughness, 0.02);
-		Out = IntegrateBxDF(GBuffer, N, V, Context.L, Context.Falloff, Context.NoL, Context.AreaLight, Shadow);
+		Out = IntegrateBxDF(GBuffer, N, V, AreaLight, Shadow);
 	}
 	return Out;
 }
@@ -342,7 +167,8 @@ FDirectLighting IntegrateBxDF( FGBufferData GBuffer, half3 N, half3 V, FRect Rec
 			float NoH = saturate( dot(N, H) );
 			float VoH = saturate( dot(V, H) );
 
-			if( NoL > 0 && VoH > 0 )
+			//if( NoL > 0 && VoH > 0 )
+			if( VoH > 0 )
 			{
 				// Intersect ray with plane
 				float t = dot( Rect.Axis[2], Rect.Origin ) / dot( Rect.Axis[2], L );
diff --git a/Engine/Shaders/Private/ShadingModels.ush b/Engine/Shaders/Private/ShadingModels.ush
--- a/Engine/Shaders/Private/ShadingModels.ush
+++ b/Engine/Shaders/Private/ShadingModels.ush
@@ -83,43 +83,6 @@ half3 GetEnvBRDF(half3 SpecularColor, half Roughness, half NoV)
 }
 #endif // SHADING_PATH_MOBILE
 
-#if 0
-void StandardShadingShared( float3 DiffuseColor, float3 SpecularColor, float Roughness, float3 V, half3 N )
-{
-	float NoV = saturate( abs( dot(N, V) ) + 1e-5 );
-
-	// Diffuse_Lambert
-	Shared.DiffuseMul = DiffuseColor * (1.0 / PI);
-
-	// D_GGX, Vis_SmithJointApprox
-	float m = Roughness * Roughness;
-	Shared.m2 = m * m;
-	Shared.SpecularMul = (0.5 / PI) * Shared.m2;
-	Shared.VisMad = float2( 2 * NoV * ( 1 - m ) + m, NoV * m );
-	
-	// F_Schlick
-	Shared.SpecularMul *= saturate( 50.0 * SpecularColor.g );
-}
-
-void StandardShadingPerLight( Shared, float3 L, float3 V, half3 N )
-{
-	float3 H = normalize(V + L);			// 3 add, 2 mad, 4 mul, 1 rsqrt
-	float NoL = saturate( dot(N, L) );		// 2 mad, 1 mul
-	float NoH = saturate( dot(N, H) );		// 2 mad, 1 mul
-	float VoH = saturate( dot(V, H) );		// 2 mad, 1 mul
-
-	// D_GGX, Vis_SmithJointApprox
-	float d = ( NoH * Shared.m2 - NoH ) * NoH + 1;			// 2 mad
-	float v = NoL * Shared.VisMad.x + Shared.VisMad.y;		// 1 mad
-	float D_Vis = Shared.SpecularMul * rcp( d * d * v );	// 3 mul, 1 rcp
-
-	// F_Schlick
-	float Fc = pow( 1 - VoH, 5 );							// 1 sub, 3 mul
-	float3 F = Fc + (1 - Fc) * SpecularColor;				// 1 sub, 3 mad
-
-	return Shared.DiffuseMul + D_Vis * F;					// 3 mad
-}
-#endif
 struct FDirectLighting
 {
 	float3	Diffuse;
@@ -134,14 +97,14 @@ struct FShadowTerms
 	half	TransmissionThickness;
 	FHairTransmittanceData HairTransmittance;
 };
-FDirectLighting HairBxDF(FGBufferData GBuffer, half3 N, half3 V, half3 L, float Falloff, half NoL, FAreaLight AreaLight, FShadowTerms Shadow)
+FDirectLighting HairBxDF(FGBufferData GBuffer, half3 N, half3 V, FAreaLight AreaLight, FShadowTerms Shadow)
 {
-	const float3 BsdfValue = HairShading(GBuffer, L, V, N, Shadow.TransmissionShadow, Shadow.HairTransmittance, HAIR_BSDF_BACKLIT, 0, uint2(0, 0));
+	const float3 BsdfValue = HairShading(GBuffer, AreaLight.DiffuseL, V, N, Shadow.TransmissionShadow, Shadow.HairTransmittance, HAIR_BSDF_BACKLIT, 0, uint2(0, 0));
 
 	FDirectLighting Lighting;
 	Lighting.Diffuse = 0;
 	Lighting.Specular = 0;
-	Lighting.Transmission = AreaLight.FalloffColor * Falloff * BsdfValue;
+	Lighting.Transmission = AreaLight.FalloffColor * AreaLight.Falloff * BsdfValue;
 	return Lighting;
 }
 
@@ -186,46 +149,35 @@ float EnergyNormalization( inout float a2, float VoH, FAreaLight AreaLight )
 	return Energy;
 }
 
-float3 SpecularGGX(float Roughness, float Anisotropy, float3 SpecularColor, BxDFContext Context, float NoL, FAreaLight AreaLight)
+float3 SpecularGGX(float Roughness, float Anisotropy, float3 SpecularColor, BxDFContext Context, FAreaLight AreaLight)
 {
 	float Alpha = Roughness * Roughness;
-	float a2 = Alpha * Alpha;
-
-	FAreaLight Punctual = AreaLight;
-	Punctual.SphereSinAlpha = 0;
-	Punctual.SphereSinAlphaSoft = 0;
-	Punctual.LineCosSubtended = 1;
-	Punctual.Rect = (FRect)0;
-	Punctual.IsRectAndDiffuseMicroReflWeight = 0;
-
-	float Energy = EnergyNormalization(a2, Context.VoH, Punctual);
-
 	float ax = 0;
 	float ay = 0;
 	GetAnisotropicRoughness(Alpha, Anisotropy, ax, ay);
 
 	// Generalized microfacet specular
-	float3 D = D_GGXaniso(ax, ay, Context.NoH, Context.XoH, Context.YoH) * Energy;
-	float3 Vis = Vis_SmithJointAniso(ax, ay, Context.NoV, NoL, Context.XoV, Context.XoL, Context.YoV, Context.YoL);
+	float3 D = D_GGXaniso(ax, ay, Context.NoH, Context.XoH, Context.YoH);
+	float3 Vis = Vis_SmithJointAniso(ax, ay, Context.NoV, AreaLight.NoL, Context.XoV, Context.XoL, Context.YoV, Context.YoL);
 	float3 F = F_Schlick( SpecularColor, Context.VoH );
 
 	return (D * Vis) * F;
 }
 
-float3 SpecularGGX( float Roughness, float3 SpecularColor, BxDFContext Context, half NoL, FAreaLight AreaLight )
+float3 SpecularGGX( float Roughness, float3 SpecularColor, BxDFContext Context, FAreaLight AreaLight )
 {
 	float a2 = Pow4( Roughness );
 	float Energy = EnergyNormalization( a2, Context.VoH, AreaLight );
 	
 	// Generalized microfacet specular
 	float D = D_GGX( a2, Context.NoH ) * Energy;
-	float Vis = Vis_SmithJointApprox( a2, Context.NoV, NoL );
+	float Vis = Vis_SmithJointApprox( a2, Context.NoV, AreaLight.NoL );
 	float3 F = F_Schlick( SpecularColor, Context.VoH );
 
 	return (D * Vis) * F;
 }
 
-half3 DualSpecularGGX(half AverageRoughness, half Lobe0Roughness, half Lobe1Roughness, half LobeMix, half3 SpecularColor, BxDFContext Context, half NoL, FAreaLight AreaLight)
+half3 DualSpecularGGX(half AverageRoughness, half Lobe0Roughness, half Lobe1Roughness, half LobeMix, half3 SpecularColor, BxDFContext Context, FAreaLight AreaLight)
 {
 	float AverageAlpha2 = Pow4(AverageRoughness);
 	float Lobe0Alpha2 = Pow4(Lobe0Roughness);
@@ -236,13 +188,13 @@ half3 DualSpecularGGX(half AverageRoughness, half Lobe0Roughness, half Lobe1Roug
 
 	// Generalized microfacet specular
 	float D = lerp(D_GGX(Lobe0Alpha2, Context.NoH) * Lobe0Energy, D_GGX(Lobe1Alpha2, Context.NoH) * Lobe1Energy, LobeMix);
-	float Vis = Vis_SmithJointApprox(AverageAlpha2, Context.NoV, NoL); // Average visibility well approximates using two separate ones (one per lobe).
+	float Vis = Vis_SmithJointApprox(AverageAlpha2, Context.NoV, AreaLight.NoL); // Average visibility well approximates using two separate ones (one per lobe).
 	float3 F = F_Schlick(SpecularColor, Context.VoH);
 
 	return (D * Vis) * F;
 }
 
-FDirectLighting DefaultLitBxDF( FGBufferData GBuffer, half3 N, half3 V, half3 L, float Falloff, half NoL, FAreaLight AreaLight, FShadowTerms Shadow )
+FDirectLighting DefaultLitBxDF( FGBufferData GBuffer, half3 N, half3 V, FAreaLight AreaLight, FShadowTerms Shadow )
 {
 	BxDFContext Context;
 	FDirectLighting Lighting;
@@ -259,7 +211,7 @@ FDirectLighting DefaultLitBxDF( FGBufferData GBuffer, half3 N, half3 V, half3 L,
 	
 	
 	BRANCH
-	if (NoL > 0.0f)
+	if (AreaLight.NoL > 0.0f)
 	{
 #if SUPPORTS_ANISOTROPIC_MATERIALS
 		bool bHasAnisotropy = HasAnisotropy(GBuffer.SelectiveOutputMask);
@@ -273,7 +225,7 @@ FDirectLighting DefaultLitBxDF( FGBufferData GBuffer, half3 N, half3 V, half3 L,
 		{
 			half3 X = GBuffer.WorldTangent;
 			half3 Y = normalize(cross(N, X));
-			Init(Context, N, X, Y, V, L);
+			Init(Context, N, X, Y, V, AreaLight.SpecularL);
 
 			NoV = Context.NoV;
 			VoH = Context.VoH;
@@ -282,9 +234,9 @@ FDirectLighting DefaultLitBxDF( FGBufferData GBuffer, half3 N, half3 V, half3 L,
 		else
 		{
 #if SHADING_PATH_MOBILE
-			InitMobile(Context, N, V, L, NoL);
+			InitMobile(Context, N, V, AreaLight.SpecularL, AreaLight.NoL);
 #else
-			Init(Context, N, V, L);
+			Init(Context, N, V, AreaLight.SpecularL);
 #endif
 
 			NoV = Context.NoV;
@@ -299,17 +251,16 @@ FDirectLighting DefaultLitBxDF( FGBufferData GBuffer, half3 N, half3 V, half3 L,
 #if MATERIAL_ROUGHDIFFUSE
 		// Chan diffuse model with roughness == specular roughness. This is not necessarily a good modelisation of reality because when the mean free path is super small, the diffuse can in fact looks rougher. But this is a start.
 		// Also we cannot use the morphed context maximising NoH as this is causing visual artefact when interpolating rough/smooth diffuse response. 
-		Lighting.Diffuse = Diffuse_Chan(GBuffer.DiffuseColor, Pow4(GBuffer.Roughness), NoV, NoL, VoH, NoH, GetAreaLightDiffuseMicroReflWeight(AreaLight));
+		Lighting.Diffuse = Diffuse_Chan(GBuffer.DiffuseColor, Pow4(GBuffer.Roughness), NoV, AreaLight.NoL, VoH, NoH, GetAreaLightDiffuseMicroReflWeight(AreaLight));
 #else
 		Lighting.Diffuse = Diffuse_Lambert(GBuffer.DiffuseColor);
 #endif
-		Lighting.Diffuse *= AreaLight.FalloffColor * (Falloff * NoL);
+		Lighting.Diffuse *= AreaLight.FalloffColor * (AreaLight.Falloff * AreaLight.NoL);
 
 		BRANCH
 		if (bHasAnisotropy)
 		{
-			//Lighting.Specular = GBuffer.WorldTangent * .5f + .5f;
-			Lighting.Specular = AreaLight.FalloffColor * (Falloff * NoL) * SpecularGGX(GBuffer.Roughness, GBuffer.Anisotropy, GBuffer.SpecularColor, Context, NoL, AreaLight);
+			Lighting.Specular = AreaLight.FalloffColor * (AreaLight.Falloff * AreaLight.NoL) * SpecularGGX(GBuffer.Roughness, GBuffer.Anisotropy, GBuffer.SpecularColor, Context, AreaLight);
 		}
 		else
 		{
@@ -319,7 +270,7 @@ FDirectLighting DefaultLitBxDF( FGBufferData GBuffer, half3 N, half3 V, half3 L,
 			}
 			else
 			{
-				Lighting.Specular = AreaLight.FalloffColor * (Falloff * NoL) * SpecularGGX(GBuffer.Roughness, GBuffer.SpecularColor, Context, NoL, AreaLight);
+				Lighting.Specular = AreaLight.FalloffColor * (AreaLight.Falloff * AreaLight.NoL) * SpecularGGX(GBuffer.Roughness, GBuffer.SpecularColor, Context, AreaLight);
 			}
 		}
 
@@ -398,7 +349,7 @@ BxDFContext RefractClearCoatContext(BxDFContext Context)
 	return RefractedContext;
 }
 
-FDirectLighting ClearCoatBxDF( FGBufferData GBuffer, half3 N, half3 V, half3 L, float Falloff, half NoL, FAreaLight AreaLight, FShadowTerms Shadow )
+FDirectLighting ClearCoatBxDF( FGBufferData GBuffer, half3 N, half3 V, FAreaLight AreaLight, FShadowTerms Shadow )
 {
 	const float ClearCoat			= GBuffer.CustomData.x;
 	const float ClearCoatRoughness	= max(GBuffer.CustomData.y, 0.02f);
@@ -431,7 +382,7 @@ FDirectLighting ClearCoatBxDF( FGBufferData GBuffer, half3 N, half3 V, half3 L,
 	//////////////////////////////
 
 	// No anisotropy for the top layer
-	Init(Context, Nspec, V, L);
+	Init(Context, Nspec, V, AreaLight.SpecularL);
 	
 	// Modify SphereSinAlpha, knowing that it was previously manipulated by roughness of the under coat
 	// Note: the operation is not invertible for GBuffer.Roughness = 1.0, so roughness is clamped to 254.0/255.0
@@ -462,10 +413,10 @@ FDirectLighting ClearCoatBxDF( FGBufferData GBuffer, half3 N, half3 V, half3 L,
 		// Generalized microfacet specular
 		float a2 = Pow2(Alpha);
 		float ClearCoatEnergy = EnergyNormalization(a2, Context.VoH, AreaLight);
-		half Vis = Vis_SmithJointApprox(a2, Context.NoV, NoL);
+		half Vis = Vis_SmithJointApprox(a2, Context.NoV, AreaLight.NoL);
 		float D = D_GGX(a2, Context.NoH) * ClearCoatEnergy;
 		half Fr1 = (D * Vis) * F;
-		Lighting.Specular = ClearCoat * AreaLight.FalloffColor * (Falloff * NoL * Fr1);
+		Lighting.Specular = ClearCoat * AreaLight.FalloffColor * (AreaLight.Falloff * AreaLight.NoL * Fr1);
 	}
 	Lighting.Specular *= ComputeEnergyConservation(EnergyTermsCoat);
 
@@ -495,19 +446,19 @@ FDirectLighting ClearCoatBxDF( FGBufferData GBuffer, half3 N, half3 V, half3 L,
 		BRANCH
 		if (bHasAnisotropy)
 		{
-			Init(TempContext, N, X, Y, V, L);
+			Init(TempContext, N, X, Y, V, AreaLight.SpecularL);
 		}
 		else
 		{
-			Init(TempContext, Nspec, V, L);
+			Init(TempContext, Nspec, V, AreaLight.SpecularL);
 		}
 
 		// If bottom-normal, update normal-based dot products:
-		float3 H = normalize(V + L);
+		float3 H = normalize(V + AreaLight.SpecularL);
 		Context.NoH = saturate(dot(N, H));
 		Context.NoV = saturate(dot(N, V));
-		Context.NoL = saturate(dot(N, L));
-		Context.VoL = saturate(dot(V, L));
+		Context.NoL = saturate(dot(N, AreaLight.SpecularL));
+		Context.VoL = saturate(dot(V, AreaLight.SpecularL));
 		Context.VoH = saturate(dot(V, H));
 
 		Context.XoV = TempContext.XoV;
@@ -540,7 +491,7 @@ FDirectLighting ClearCoatBxDF( FGBufferData GBuffer, half3 N, half3 V, half3 L,
 		// Prepare the anisotropic Context to refract.
 		X = GBuffer.WorldTangent;
 		Y = normalize(cross(N, X));
-		Init(Context, Nspec, X, Y, V, L);
+		Init(Context, Nspec, X, Y, V, AreaLight.SpecularL);
 	}
 	BxDFContext BottomContext = RefractClearCoatContext(Context);
 	BottomContext.VoH = bIsRect ? BottomContext.NoV : BottomContext.VoH;
@@ -551,7 +502,7 @@ FDirectLighting ClearCoatBxDF( FGBufferData GBuffer, half3 N, half3 V, half3 L,
 	float3 Transmission = SimpleClearCoatTransmittance(BottomContext.NoL, BottomContext.NoV, GBuffer.Metallic, GBuffer.BaseColor);
 
 	// Default Lit
-	half3 DefaultDiffuse = (Falloff * NoL) * AreaLight.FalloffColor * Diffuse_Lambert(GBuffer.DiffuseColor) * ComputeEnergyPreservation(EnergyTermsBottom);
+	half3 DefaultDiffuse = (AreaLight.Falloff * AreaLight.NoL) * AreaLight.FalloffColor * Diffuse_Lambert(GBuffer.DiffuseColor) * ComputeEnergyPreservation(EnergyTermsBottom);
 	half3 RefractedDiffuse = FresnelCoeff * Transmission * DefaultDiffuse;
 	Lighting.Diffuse = lerp(DefaultDiffuse, RefractedDiffuse, ClearCoat);
 
@@ -570,20 +521,9 @@ FDirectLighting ClearCoatBxDF( FGBufferData GBuffer, half3 N, half3 V, half3 L,
 
 		// The energy factor for correcting area light shape needs to be computed "before" the actual D_GGX term, 
 		// as calling EnergyNormalization() modify the "a2" value.
-		float Energy = 0;
+		float Energy = 1;
 		BRANCH
-		if (bHasAnisotropy)
-		{
-			FAreaLight Punctual = AreaLight;
-			Punctual.SphereSinAlpha = 0;
-			Punctual.SphereSinAlphaSoft = 0;
-			Punctual.LineCosSubtended = 1;
-			Punctual.Rect = (FRect)0;
-			Punctual.IsRectAndDiffuseMicroReflWeight = 0;
-			
-			Energy = EnergyNormalization(a2, Context.VoH, Punctual);
-		}
-		else
+		if (!bHasAnisotropy)
 		{
 			Energy = EnergyNormalization(a2, Context.VoH, AreaLight);
 		}
@@ -601,7 +541,7 @@ FDirectLighting ClearCoatBxDF( FGBufferData GBuffer, half3 N, half3 V, half3 L,
 		else
 		{
 			// NoL is chosen to provide better parity with DefaultLit when ClearCoat=0
-			Vis2 = Vis_SmithJointApprox(a2, BottomContext.NoV, NoL);
+			Vis2 = Vis_SmithJointApprox(a2, BottomContext.NoV, AreaLight.NoL);
 			D2 = D_GGX(a2, BottomContext.NoH);
 		}
 		half3 F_Bot = F_Schlick(GBuffer.SpecularColor, BottomContext.VoH);
@@ -609,7 +549,7 @@ FDirectLighting ClearCoatBxDF( FGBufferData GBuffer, half3 N, half3 V, half3 L,
 		
 
 		// Note: reusing D and V from refracted context to save computation when ClearCoat < 1
-		float3 CommonSpecular = (Energy * Falloff * NoL * D2 * Vis2) * AreaLight.FalloffColor;
+		float3 CommonSpecular = (Energy * AreaLight.Falloff * AreaLight.NoL * D2 * Vis2) * AreaLight.FalloffColor;
 		float3 DefaultSpecular = F_DefaultLit;
 		float3 RefractedSpecular = FresnelCoeff * Transmission * F_Bot;
 		Lighting.Specular += CommonSpecular * lerp(DefaultSpecular, RefractedSpecular, ClearCoat);
@@ -642,13 +582,13 @@ void GetProfileDualSpecular(uint SubsurfaceProfileInt, half Roughness, half Opac
 #endif
 }
 
-FDirectLighting SubsurfaceProfileBxDF( FGBufferData GBuffer, half3 N, half3 V, half3 L, float Falloff, half NoL, FAreaLight AreaLight, FShadowTerms Shadow )
+FDirectLighting SubsurfaceProfileBxDF( FGBufferData GBuffer, half3 N, half3 V, FAreaLight AreaLight, FShadowTerms Shadow )
 {
 	BxDFContext Context;
 #if SHADING_PATH_MOBILE
-	InitMobile(Context, N, V, L, NoL);
+	InitMobile(Context, N, V, AreaLight.SpecularL, AreaLight.NoL);
 #else
-	Init( Context, N, V, L );
+	Init( Context, N, V, AreaLight.SpecularL );
 #endif
 	SphereMaxNoH( Context, AreaLight.SphereSinAlpha, true );
 	Context.NoV = saturate( abs( Context.NoV ) + 1e-5 );
@@ -670,7 +610,7 @@ FDirectLighting SubsurfaceProfileBxDF( FGBufferData GBuffer, half3 N, half3 V, h
 	FDirectLighting Lighting;
 #if SHADING_PATH_MOBILE
 	half Curvature = GBuffer.Curvature;
-	half UnClampedNoL = dot(GBuffer.WorldNormal, L);
+	half UnClampedNoL = dot(GBuffer.WorldNormal, AreaLight.DiffuseL);
 	half ShadowFactor = 1.0f - sqrt(Shadow.SurfaceShadow);
 	// Rotate the world normal based on the shadow value, it's just a experimental value 
 	half UnClampedRotatedNoL = max(UnClampedNoL - max(2.0f * UnClampedNoL, 0.4f) * ShadowFactor, -1.0f);
@@ -679,17 +619,17 @@ FDirectLighting SubsurfaceProfileBxDF( FGBufferData GBuffer, half3 N, half3 V, h
 	half3 Tint = GetSubsurfaceProfileTexture(SSSS_TINT_SCALE_OFFSET, SubsurfaceProfileId).rgb;
 	
 	// Needs to apply shadow at here, since the preintegrated value has already counted it, and we need to skip applying shadow outside.
-	Lighting.Diffuse = lerp(AreaLight.FalloffColor * (Falloff * NoL) * Shadow.SurfaceShadow, BurleyDiffuse.rgb, Tint) * Diffuse_Lambert(GBuffer.DiffuseColor);
+	Lighting.Diffuse = lerp(AreaLight.FalloffColor * (AreaLight.Falloff * AreaLight.NoL) * Shadow.SurfaceShadow, BurleyDiffuse.rgb, Tint) * Diffuse_Lambert(GBuffer.DiffuseColor);
 #else
 
 #if MATERIAL_ROUGHDIFFUSE
 	// Use Chan's diffuse model. It reduces flatness look and has better match, e.g., at the edge of human face skin when compared to GT.
-	const float3 DiffuseReflection = Diffuse_Chan(GBuffer.DiffuseColor, Pow4(GBuffer.Roughness), Context.NoV, NoL, Context.VoH, Context.NoH, GetAreaLightDiffuseMicroReflWeight(AreaLight));
+	const float3 DiffuseReflection = Diffuse_Chan(GBuffer.DiffuseColor, Pow4(GBuffer.Roughness), Context.NoV, AreaLight.NoL, Context.VoH, Context.NoH, GetAreaLightDiffuseMicroReflWeight(AreaLight));
 #else
-	const float3 DiffuseReflection = Diffuse_Burley(GBuffer.DiffuseColor, GBuffer.Roughness, Context.NoV, NoL, Context.VoH);
+	const float3 DiffuseReflection = Diffuse_Burley(GBuffer.DiffuseColor, GBuffer.Roughness, Context.NoV, AreaLight.NoL, Context.VoH);
 #endif
 
-	Lighting.Diffuse  = AreaLight.FalloffColor * (Falloff * NoL) * DiffuseReflection;
+	Lighting.Diffuse  = AreaLight.FalloffColor * (AreaLight.Falloff * AreaLight.NoL) * DiffuseReflection;
 #endif
 
 	if (IsRectLight(AreaLight))
@@ -700,7 +640,7 @@ FDirectLighting SubsurfaceProfileBxDF( FGBufferData GBuffer, half3 N, half3 V, h
 	}
 	else
 	{
-		Lighting.Specular = AreaLight.FalloffColor * (Falloff * NoL) * DualSpecularGGX(AverageRoughness, Lobe0Roughness, Lobe1Roughness, LobeMix, GBuffer.SpecularColor, Context, NoL, AreaLight);
+		Lighting.Specular = AreaLight.FalloffColor * (AreaLight.Falloff * AreaLight.NoL) * DualSpecularGGX(AverageRoughness, Lobe0Roughness, Lobe1Roughness, LobeMix, GBuffer.SpecularColor, Context, AreaLight);
 	}
 
 	Lighting.Diffuse  *= ComputeEnergyPreservation(EnergyTerms);
@@ -717,8 +657,8 @@ FDirectLighting SubsurfaceProfileBxDF( FGBufferData GBuffer, half3 N, half3 V, h
 	float3 Profile = GetTransmissionProfile(ProfileId, Thickness).rgb;
 
 	float3 RefracV = refract(V, -N, TransmissionParams.OneOverIOR);
-	float PhaseFunction = ApproximateHG( dot(-L, RefracV), TransmissionParams.ScatteringDistribution );
-	Lighting.Transmission = AreaLight.FalloffColor * Profile * (Falloff * PhaseFunction); // TODO: This probably should also include cosine term (NoL)
+	float PhaseFunction = ApproximateHG( dot(-AreaLight.DiffuseL, RefracV), TransmissionParams.ScatteringDistribution );
+	Lighting.Transmission = AreaLight.FalloffColor * Profile * (AreaLight.Falloff * PhaseFunction); // TODO: This probably should also include cosine term (NoL)
 
 #else // USE_TRANSMISSION
 
@@ -729,16 +669,16 @@ FDirectLighting SubsurfaceProfileBxDF( FGBufferData GBuffer, half3 N, half3 V, h
 	return Lighting;
 }
 
-FDirectLighting ClothBxDF( FGBufferData GBuffer, half3 N, half3 V, half3 L, float Falloff, half NoL, FAreaLight AreaLight, FShadowTerms Shadow )
+FDirectLighting ClothBxDF( FGBufferData GBuffer, half3 N, half3 V, FAreaLight AreaLight, FShadowTerms Shadow )
 {
 	const float3 FuzzColor	= ExtractSubsurfaceColor(GBuffer);
 	const float  Cloth		= saturate(GBuffer.CustomData.a);
 
 	BxDFContext Context;
 #if SHADING_PATH_MOBILE
-	InitMobile(Context, N, V, L, NoL);
+	InitMobile(Context, N, V, AreaLight.SpecularL, AreaLight.NoL);
 #else
-	Init( Context, N, V, L );
+	Init( Context, N, V, AreaLight.SpecularL );
 #endif
 	SphereMaxNoH( Context, AreaLight.SphereSinAlpha, true );
 	Context.NoV = saturate( abs( Context.NoV ) + 1e-5 );
@@ -747,25 +687,25 @@ FDirectLighting ClothBxDF( FGBufferData GBuffer, half3 N, half3 V, half3 L, floa
 	if(IsRectLight(AreaLight))
 		Spec1 = RectGGXApproxLTC( GBuffer.Roughness, GBuffer.SpecularColor, N, V, AreaLight.Rect, AreaLight.Texture );
 	else
-		Spec1 = AreaLight.FalloffColor * (Falloff * NoL) * SpecularGGX( GBuffer.Roughness, GBuffer.SpecularColor, Context, NoL, AreaLight );
+		Spec1 = AreaLight.FalloffColor * (AreaLight.Falloff * AreaLight.NoL) * SpecularGGX( GBuffer.Roughness, GBuffer.SpecularColor, Context, AreaLight );
 
 	const FBxDFEnergyTerms EnergyTerms1 = ComputeGGXSpecEnergyTerms(GBuffer.Roughness, Context.NoV, GBuffer.SpecularColor);
 	Spec1 *= ComputeEnergyConservation(EnergyTerms1);
 
 	// Cloth - Asperity Scattering - Inverse Beckmann Layer
 	float D2 = D_InvGGX( Pow4( GBuffer.Roughness ), Context.NoH );
-	float Vis2 = Vis_Cloth( Context.NoV, NoL );
+	float Vis2 = Vis_Cloth( Context.NoV, AreaLight.NoL );
 #if SHADING_PATH_MOBILE
 	Vis2 = saturate(Vis2);
 #endif
 	float3 F2 = F_Schlick( FuzzColor, Context.VoH );
-	float3 Spec2 = AreaLight.FalloffColor * (Falloff * NoL) * (D2 * Vis2) * F2;
+	float3 Spec2 = AreaLight.FalloffColor * (AreaLight.Falloff * AreaLight.NoL) * (D2 * Vis2) * F2;
 	
 	const FBxDFEnergyTermsA EnergyTerms2 = ComputeClothEnergyTermsA(GBuffer.Roughness, Context.NoV);
 	Spec2 *= ComputeEnergyConservation(EnergyTerms2);
 
 	FDirectLighting Lighting;
-	Lighting.Diffuse  = AreaLight.FalloffColor * (Falloff * NoL) * Diffuse_Lambert( GBuffer.DiffuseColor ); 
+	Lighting.Diffuse  = AreaLight.FalloffColor * (AreaLight.Falloff * AreaLight.NoL) * Diffuse_Lambert( GBuffer.DiffuseColor ); 
 	Lighting.Specular = lerp( Spec1, Spec2, Cloth );
 	Lighting.Transmission = 0;
 
@@ -774,16 +714,16 @@ FDirectLighting ClothBxDF( FGBufferData GBuffer, half3 N, half3 V, half3 L, floa
 	return Lighting;
 }
 
-FDirectLighting SubsurfaceBxDF(FGBufferData GBuffer, half3 N, half3 V, half3 L, float Falloff, half NoL, FAreaLight AreaLight, FShadowTerms Shadow )
+FDirectLighting SubsurfaceBxDF(FGBufferData GBuffer, half3 N, half3 V, FAreaLight AreaLight, FShadowTerms Shadow )
 {
-	FDirectLighting Lighting = DefaultLitBxDF( GBuffer, N, V, L, Falloff, NoL, AreaLight, Shadow);
+	FDirectLighting Lighting = DefaultLitBxDF( GBuffer, N, V, AreaLight, Shadow);
 
 	half3 SubsurfaceColor = ExtractSubsurfaceColor(GBuffer);
 	half Opacity = GBuffer.CustomData.a;
 
 	// to get an effect when you see through the material
 	// hard coded pow constant
-	half InScatter = pow(saturate(dot(L, -V)), 12) * lerp(3, .1f, Opacity);
+	half InScatter = pow(saturate(dot(AreaLight.DiffuseL, -V)), 12) * lerp(3, .1f, Opacity);
 
 	// Wrap around lighting, 
 	// * /(PI*2) to be energy consistent (hack do get some view dependnt and light dependent effect)
@@ -791,7 +731,7 @@ FDirectLighting SubsurfaceBxDF(FGBufferData GBuffer, half3 N, half3 V, half3 L,
 	// * Simplified version (with w=.5,  n=1.5): 
 	//     half WrappedDiffuse = 2 * pow(saturate((dot(N, L) + w) / (1.0f + w)), n) * (n + 1) / (2 * (1 + w));
 	//     NormalContribution = WrappedDiffuse * Opacity + 1 - Opacity;
-	const half WrappedDiffuse = pow(saturate(dot(N, L) * (1.f / 1.5f) + (0.5f / 1.5f)), 1.5f) * (2.5f / 1.5f);
+	const half WrappedDiffuse = pow(saturate(dot(N, AreaLight.DiffuseL) * (1.f / 1.5f) + (0.5f / 1.5f)), 1.5f) * (2.5f / 1.5f);
 	const half NormalContribution = lerp(1.f, WrappedDiffuse, Opacity);
 	const half BackScatter = GBuffer.GBufferAO * NormalContribution / (PI * 2);
 
@@ -805,7 +745,7 @@ FDirectLighting SubsurfaceBxDF(FGBufferData GBuffer, half3 N, half3 V, half3 L,
 	const half3 TransmittedColor = HSV_2_LinearRGB(half3(LinearRGB_2_HSV(RawTransmittedColor).xy, LinearRGB_2_HSV(SubsurfaceColor).z));
 
 	// lerp to never exceed 1 (energy conserving)
-	Lighting.Transmission = AreaLight.FalloffColor * (Falloff * lerp(BackScatter, 1, InScatter)) * lerp(TransmittedColor, SubsurfaceColor, Shadow.TransmissionThickness);
+	Lighting.Transmission = AreaLight.FalloffColor * (AreaLight.Falloff * lerp(BackScatter, 1, InScatter)) * lerp(TransmittedColor, SubsurfaceColor, Shadow.TransmissionThickness);
 
 	return Lighting;
 }
@@ -901,7 +841,7 @@ FDirectLighting ShadeSGGX( FGBufferData GBuffer, float Geo_a2, half3 V, half3 L
 }
 #endif
 
-FDirectLighting TwoSidedBxDF( FGBufferData GBuffer, half3 N, half3 V, half3 L, float Falloff, half AreaNoL, FAreaLight AreaLight, FShadowTerms Shadow )
+FDirectLighting TwoSidedBxDF( FGBufferData GBuffer, half3 N, half3 V, FAreaLight AreaLight, FShadowTerms Shadow )
 {
 	half3 SubsurfaceColor = ExtractSubsurfaceColor(GBuffer);
 
@@ -911,6 +851,7 @@ FDirectLighting TwoSidedBxDF( FGBufferData GBuffer, half3 N, half3 V, half3 L, f
 	half NormalLength = NormalLength2 * RcpNormalLength;
 	N = GBuffer.NormalDistribution * RcpNormalLength;
 
+	half3 L = AreaLight.DiffuseL;
 	half NoL = dot(N, L);
 	half NoV = dot(N, V);
 	half VoL = dot(V, L);
@@ -946,44 +887,40 @@ FDirectLighting TwoSidedBxDF( FGBufferData GBuffer, half3 N, half3 V, half3 L, f
 
 	float SphereVoL = dot( SphereV, SphereL );
 
-	AreaNoL = SGGX_Diffuse( SphereVoL );
-	//AreaNoL *= SGGX_ProjectedArea( Geo_a2, NoL );
-	AreaNoL *= TangentL.z / SphereL.z;
-	//AreaNoL /= 0.25 * Geo_Area;
-	AreaNoL *= lerp( PI, 4, saturate( Geo_a ) ) / Geo_Area;
+#if 0
+	float vNoV = SGGX_ProjectedArea( Geo_a2, NoV ) * 2 / Geo_Area;
+	float vNoL = SGGX_ProjectedArea( Geo_a2, NoL ) * 2 / Geo_Area;
+#else
+	float vNoV = 2 * TangentV.z / ( SphereV.z * Geo_Area );
+	float vNoL = 2 * TangentL.z / ( SphereL.z * Geo_Area );
+#endif
+
+	vNoL *= SGGX_Diffuse( SphereVoL );
+	vNoL *= lerp( 0.5 * PI, 2, saturate( Geo_a ) );
 
 	FDirectLighting Lighting = (FDirectLighting)0;
 
 	{
-		Lighting.Diffuse = AreaLight.FalloffColor * (Falloff * AreaNoL) * Diffuse_Lambert( GBuffer.DiffuseColor );
+		Lighting.Diffuse = AreaLight.FalloffColor * (AreaLight.Falloff * vNoL) * Diffuse_Lambert( GBuffer.DiffuseColor );
 	}
 
 	{
 		float Mat_a2 = Pow4( GBuffer.Roughness );
 		float Spec_a2 = SGGX_Convolve( Geo_a2, Mat_a2 );
 
-	#if 0
-		float NoV2 = SGGX_ProjectedArea( Geo_a2, NoV ) * 2 / Geo_Area;
-		float NoL2 = SGGX_ProjectedArea( Geo_a2, NoL ) * 2 / Geo_Area;
-	#else
-		float NoV2 = TangentV.z / SphereV.z * 2 / Geo_Area;
-		float NoL2 = TangentL.z / SphereL.z * 2 / Geo_Area;
-	#endif
-		NoL2 = AreaNoL;
-
 		float D = D_GGX( Spec_a2, NoH );
 
-		float Vis_SmithV = NoL2 * SGGX_ProjectedArea( Spec_a2, NoV );
-		float Vis_SmithL = NoV2 * SGGX_ProjectedArea( Spec_a2, NoL );
+		float Vis_SmithV = vNoL * SGGX_ProjectedArea( Spec_a2, NoV );
+		float Vis_SmithL = vNoV * SGGX_ProjectedArea( Spec_a2, NoL );
 
-		//Vis_SmithV = lerp( NoL2 * SGGX_ProjectedArea( Mat_a2, NoV2 ), Vis_SmithV, VoH );
-		Vis_SmithL = lerp( NoV2 * SGGX_ProjectedArea( Mat_a2, NoL2 ), Vis_SmithL, VoH );
+		//Vis_SmithV = lerp( vNoL * SGGX_ProjectedArea( Mat_a2, vNoV ), Vis_SmithV, VoH );
+		Vis_SmithL = lerp( vNoV * SGGX_ProjectedArea( Mat_a2, vNoL ), Vis_SmithL, VoH );
 
 		float Vis = 0.5 * rcp( Vis_SmithV + Vis_SmithL );
 
 		float3 F = F_Schlick( GBuffer.SpecularColor, VoH );
 
-		Lighting.Specular = AreaLight.FalloffColor * (Falloff * AreaNoL * D * Vis) * F;
+		Lighting.Specular = AreaLight.FalloffColor * (AreaLight.Falloff * vNoL * D * Vis) * F;
 	}
 
 	{
@@ -998,26 +935,26 @@ FDirectLighting TwoSidedBxDF( FGBufferData GBuffer, half3 N, half3 V, half3 L, f
 		// Scatter distribution
 		float Scatter = D_GGX( 0.6*0.6, saturate( -VoL ) );
 
-		Lighting.Transmission = AreaLight.FalloffColor * (Falloff * WrapNoL * Scatter) * SubsurfaceColor;
+		Lighting.Transmission = AreaLight.FalloffColor * (AreaLight.Falloff * WrapNoL * Scatter) * SubsurfaceColor;
 	}
 #else
-	FDirectLighting Lighting = DefaultLitBxDF( GBuffer, N, V, L, Falloff, AreaNoL, AreaLight, Shadow );
+	FDirectLighting Lighting = DefaultLitBxDF( GBuffer, N, V, AreaLight, Shadow );
 
 	// http://blog.stevemcauley.com/2011/12/03/energy-conserving-wrapped-diffuse/
 	half Wrap = 0.5;
-	half WrapNoL = saturate( ( -dot(N, L) + Wrap ) / Square( 1 + Wrap ) );
+	half WrapNoL = saturate( ( -dot(N, AreaLight.DiffuseL) + Wrap ) / Square( 1 + Wrap ) );
 
 	// Scatter distribution
-	half VoL = dot(V, L);
+	half VoL = dot(V, AreaLight.DiffuseL);
 	float Scatter = D_GGX( 0.6*0.6, saturate( -VoL ) );
 
-	Lighting.Transmission = AreaLight.FalloffColor * (Falloff * WrapNoL * Scatter) * SubsurfaceColor;
+	Lighting.Transmission = AreaLight.FalloffColor * (AreaLight.Falloff * WrapNoL * Scatter) * SubsurfaceColor;
 #endif
 
 	return Lighting;
 }
 
-FDirectLighting EyeBxDF( FGBufferData GBuffer, half3 N, half3 V, half3 L, float Falloff, half NoL, FAreaLight AreaLight, FShadowTerms Shadow )
+FDirectLighting EyeBxDF( FGBufferData GBuffer, half3 N, half3 V, FAreaLight AreaLight, FShadowTerms Shadow )
 {
 #if IRIS_NORMAL
 	const float2 CausticNormalDelta	= float2( GBuffer.StoredMetallic, GBuffer.StoredSpecular )	* 2 - (256.0/255.0);
@@ -1041,9 +978,9 @@ FDirectLighting EyeBxDF( FGBufferData GBuffer, half3 N, half3 V, half3 L, float
 
 	BxDFContext Context;
 #if SHADING_PATH_MOBILE
-	InitMobile(Context, N, V, L, NoL);
+	InitMobile(Context, N, V, AreaLight.SpecularL, AreaLight.NoL);
 #else
-	Init( Context, N, V, L );
+	Init( Context, N, V, AreaLight.SpecularL );
 #endif
 	SphereMaxNoH( Context, AreaLight.SphereSinAlpha, false );
 	Context.NoV = saturate( abs( Context.NoV ) + 1e-5 );
@@ -1068,16 +1005,16 @@ FDirectLighting EyeBxDF( FGBufferData GBuffer, half3 N, half3 V, half3 L, float
 		float a2 = Pow4( GBuffer.Roughness );
 		float Energy = EnergyNormalization( a2, Context.VoH, AreaLight );
 		// Generalized microfacet specular
-		float Vis = Vis_SmithJointApprox(a2, Context.NoV, NoL);
+		float Vis = Vis_SmithJointApprox(a2, Context.NoV, AreaLight.NoL);
 		float D = D_GGX(a2, Context.NoH) * Energy;
-		Lighting.Specular = AreaLight.FalloffColor * (Falloff * NoL) * D * Vis * F;
+		Lighting.Specular = AreaLight.FalloffColor * (AreaLight.Falloff * AreaLight.NoL) * D * Vis * F;
 	}
 
-	float IrisNoL = saturate( dot( IrisNormal, L ) );
+	float IrisNoL = saturate( dot( IrisNormal, AreaLight.SpecularL ) );
 	float Power = lerp( 12, 1, IrisNoL );
-	float Caustic = 0.8 + 0.2 * ( Power + 1 ) * pow( saturate( dot( CausticNormal, L ) ), Power );
+	float Caustic = 0.8 + 0.2 * ( Power + 1 ) * pow( saturate( dot( CausticNormal, AreaLight.SpecularL ) ), Power );
 	float Iris = IrisNoL * Caustic;
-	float Sclera = NoL;
+	float Sclera = AreaLight.NoL;
 	
 	Lighting.Specular *= ComputeEnergyConservation(EnergyTerms);
 
@@ -1089,67 +1026,67 @@ FDirectLighting EyeBxDF( FGBufferData GBuffer, half3 N, half3 V, half3 L, float
 #endif
 
 	Lighting.Diffuse = 0;
-	Lighting.Transmission = AreaLight.FalloffColor * ( Falloff * lerp( Sclera, Iris, IrisMask ) * EnergyPreservation ) * Diffuse_Lambert( GBuffer.DiffuseColor );
+	Lighting.Transmission = AreaLight.FalloffColor * ( AreaLight.Falloff * lerp( Sclera, Iris, IrisMask ) * EnergyPreservation ) * Diffuse_Lambert( GBuffer.DiffuseColor );
 #if SHADING_PATH_MOBILE
 	uint SubsurfaceProfileId = ExtractSubsurfaceProfileInt(GBuffer);
 	half Curvature = GBuffer.Curvature;
 	half ShadowFactor = 1.0f - sqrt(Shadow.SurfaceShadow);
 	
 	// Rotate the world normal based on the shadow value, it's just a experimental value 
-	half UnClampedNoL = dot(GBuffer.WorldNormal, L);
+	half UnClampedNoL = dot(GBuffer.WorldNormal, AreaLight.DiffuseL);
 	half UnClampedRotatedNoL = max(UnClampedNoL - max(2.0f * UnClampedNoL, 0.4f) * ShadowFactor, -1.0f);
 	half4 BurleyDiffuse = GetSSProfilePreIntegratedValue(SubsurfaceProfileId, UnClampedRotatedNoL, Curvature);
 
 	// Rotate the world normal based on the shadow value, it's just a experimental value 
-	half UnClampedIrisNoL = dot(IrisNormal, L);
+	half UnClampedIrisNoL = dot(IrisNormal, AreaLight.SpecularL);
 	half UnClampedRotatedIrisNoL = max(UnClampedIrisNoL - max(2.0f * UnClampedIrisNoL, 0.4f) * ShadowFactor, -1.0f);
 	half4 IrisBurleyDiffuse = GetSSProfilePreIntegratedValue(SubsurfaceProfileId, UnClampedRotatedIrisNoL * Caustic, Curvature);
 	// asset specific color
 	half3 Tint = GetSubsurfaceProfileTexture(SSSS_TINT_SCALE_OFFSET, SubsurfaceProfileId).rgb;
 
 	// Needs to apply shadow at here, since the preintegrated value has already counted it, and we need to skip applying shadow outside.
-	Lighting.Transmission = lerp(Lighting.Transmission * Shadow.SurfaceShadow, Falloff * lerp(BurleyDiffuse.rgb, IrisBurleyDiffuse.rgb, IrisMask) * Diffuse_Lambert(GBuffer.DiffuseColor), Tint);
+	Lighting.Transmission = lerp(Lighting.Transmission * Shadow.SurfaceShadow, AreaLight.Falloff * lerp(BurleyDiffuse.rgb, IrisBurleyDiffuse.rgb, IrisMask) * Diffuse_Lambert(GBuffer.DiffuseColor), Tint);
 #endif
 	return Lighting;
 }
 
-FDirectLighting PreintegratedSkinBxDF( FGBufferData GBuffer, half3 N, half3 V, half3 L, float Falloff, half NoL, FAreaLight AreaLight, FShadowTerms Shadow )
+FDirectLighting PreintegratedSkinBxDF( FGBufferData GBuffer, half3 N, half3 V, FAreaLight AreaLight, FShadowTerms Shadow )
 {
-	FDirectLighting Lighting = DefaultLitBxDF( GBuffer, N, V, L, Falloff, NoL, AreaLight, Shadow );
+	FDirectLighting Lighting = DefaultLitBxDF( GBuffer, N, V, AreaLight, Shadow );
 	
 	half3 SubsurfaceColor = ExtractSubsurfaceColor(GBuffer);
 	half Opacity = GBuffer.CustomData.a;
 
-	half3 PreintegratedBRDF = Texture2DSampleLevel(View.PreIntegratedBRDF, View.PreIntegratedBRDFSampler, float2(saturate(dot(N, L) * .5 + .5), 1 - Opacity), 0).rgb;
-	Lighting.Transmission = AreaLight.FalloffColor * Falloff * PreintegratedBRDF * SubsurfaceColor;
+	half3 PreintegratedBRDF = Texture2DSampleLevel(View.PreIntegratedBRDF, View.PreIntegratedBRDFSampler, float2(saturate(dot(N, AreaLight.DiffuseL) * .5 + .5), 1 - Opacity), 0).rgb;
+	Lighting.Transmission = AreaLight.FalloffColor * AreaLight.Falloff * PreintegratedBRDF * SubsurfaceColor;
 
 	return Lighting;
 }
 
-FDirectLighting IntegrateBxDF( FGBufferData GBuffer, half3 N, half3 V, half3 L, float Falloff, half NoL, FAreaLight AreaLight, FShadowTerms Shadow )
+FDirectLighting IntegrateBxDF( FGBufferData GBuffer, half3 N, half3 V, FAreaLight AreaLight, FShadowTerms Shadow )
 {
 	switch( GBuffer.ShadingModelID )
 	{
 		case SHADINGMODELID_DEFAULT_LIT:
 		case SHADINGMODELID_SINGLELAYERWATER:
 		case SHADINGMODELID_THIN_TRANSLUCENT:
-			return DefaultLitBxDF( GBuffer, N, V, L, Falloff, NoL, AreaLight, Shadow );
+			return DefaultLitBxDF( GBuffer, N, V, AreaLight, Shadow );
 		case SHADINGMODELID_SUBSURFACE:
-			return SubsurfaceBxDF( GBuffer, N, V, L, Falloff, NoL, AreaLight, Shadow );
+			return SubsurfaceBxDF( GBuffer, N, V, AreaLight, Shadow );
 		case SHADINGMODELID_PREINTEGRATED_SKIN:
-			return PreintegratedSkinBxDF( GBuffer, N, V, L, Falloff, NoL, AreaLight, Shadow );
+			return PreintegratedSkinBxDF( GBuffer, N, V, AreaLight, Shadow );
 		case SHADINGMODELID_CLEAR_COAT:
-			return ClearCoatBxDF( GBuffer, N, V, L, Falloff, NoL, AreaLight, Shadow );
+			return ClearCoatBxDF( GBuffer, N, V, AreaLight, Shadow );
 		case SHADINGMODELID_SUBSURFACE_PROFILE:
-			return SubsurfaceProfileBxDF( GBuffer, N, V, L, Falloff, NoL, AreaLight, Shadow );
+			return SubsurfaceProfileBxDF( GBuffer, N, V, AreaLight, Shadow );
 		case SHADINGMODELID_TWOSIDED_FOLIAGE:
-			return TwoSidedBxDF( GBuffer, N, V, L, Falloff, NoL, AreaLight, Shadow );
+			return TwoSidedBxDF( GBuffer, N, V, AreaLight, Shadow );
 		case SHADINGMODELID_HAIR:
-			return HairBxDF( GBuffer, N, V, L, Falloff, NoL, AreaLight, Shadow );
+			return HairBxDF( GBuffer, N, V, AreaLight, Shadow );
 		case SHADINGMODELID_CLOTH:
-			return ClothBxDF( GBuffer, N, V, L, Falloff, NoL, AreaLight, Shadow );
+			return ClothBxDF( GBuffer, N, V, AreaLight, Shadow );
 		case SHADINGMODELID_EYE:
-			return EyeBxDF( GBuffer, N, V, L, Falloff, NoL, AreaLight, Shadow );
+			return EyeBxDF( GBuffer, N, V, AreaLight, Shadow );
 		default:
 			return (FDirectLighting)0;
 	}
@@ -1157,14 +1094,10 @@ FDirectLighting IntegrateBxDF( FGBufferData GBuffer, half3 N, half3 V, half3 L,
 
 FDirectLighting EvaluateBxDF( FGBufferData GBuffer, half3 N, half3 V, half3 L, float NoL, FShadowTerms Shadow )
 {
-	FAreaLight AreaLight;
-	AreaLight.SphereSinAlpha = 0;
-	AreaLight.SphereSinAlphaSoft = 0;
-	AreaLight.LineCosSubtended = 1;
-	AreaLight.FalloffColor = 1;
-	AreaLight.Rect = (FRect)0;
-	AreaLight.IsRectAndDiffuseMicroReflWeight = 0;
-	AreaLight.Texture = InitRectTexture();
-
-	return IntegrateBxDF( GBuffer, N, V, L, 1, NoL, AreaLight, Shadow );
+	FAreaLight AreaLight = InitAreaLight();
+	AreaLight.DiffuseL = L;
+	AreaLight.SpecularL = L;
+	AreaLight.NoL = NoL;
+
+	return IntegrateBxDF( GBuffer, N, V, AreaLight, Shadow );
 }
diff --git a/Engine/Shaders/Private/Substrate/SubstrateDeferredLighting.ush b/Engine/Shaders/Private/Substrate/SubstrateDeferredLighting.ush
--- a/Engine/Shaders/Private/Substrate/SubstrateDeferredLighting.ush
+++ b/Engine/Shaders/Private/Substrate/SubstrateDeferredLighting.ush
@@ -137,7 +137,7 @@ FSubstrateDeferredLighting SubstrateDeferredLighting(
 
 		float Roughness = SubstrateGetBSDFRoughness(BSDFContext.BSDF);
 
-		FAreaLightIntegrateContext AreaLightContext = InitAreaLightIntegrateContext();
+		FAreaLight AreaLight = InitAreaLight();
 		FSubstrateEvaluateResult BSDFEvaluate = (FSubstrateEvaluateResult)0;
 		if (LightData.bRectLight)
 		{
@@ -146,25 +146,25 @@ FSubstrateDeferredLighting SubstrateDeferredLighting(
 			{
 				return GetInitialisedSubstrateDeferredLighting(); // Rect light can be non visible due to barn door occlusion
 			}
-			AreaLightContext = CreateRectIntegrateContext(Roughness, BSDFContext.N, BSDFContext.V, Rect, RectTexture);
+			AreaLight = CreateAreaLight(Roughness, BSDFContext.N, BSDFContext.V, Rect, RectTexture);
 
 			// We must have the SubstrateIntegrateBSDF inside the if due to the rectlight texture: it must be non ambiguous which texture is going to be used.
 			// After the compilation, a local resource must map to a unique global resource (the default or the actual rect light texture).
-			BSDFEvaluate = SubstrateEvaluateBSDFCommon(BSDFContext, BSDFShadowTerms, AreaLightContext, Settings, INTEGRATION_AREA_LIGHT_RECT);
+			BSDFEvaluate = SubstrateEvaluateBSDFCommon(BSDFContext, BSDFShadowTerms, AreaLight, Settings, INTEGRATION_AREA_LIGHT_RECT);
 		}
 		else
 		{
 			FCapsuleLight Capsule = GetCapsule(ToLight, LightData);
-			AreaLightContext = CreateCapsuleIntegrateContext(Roughness, BSDFContext.N, BSDFContext.V, Capsule, LightData.bInverseSquared);
+			AreaLight = CreateAreaLight(Roughness, BSDFContext.N, BSDFContext.V, Capsule, LightData.bInverseSquared);
 
 			BRANCH
-			if(IsAreaLight(AreaLightContext.AreaLight))
+			if(IsAreaLight(AreaLight))
 			{
-				BSDFEvaluate = SubstrateEvaluateBSDFCommon(BSDFContext, BSDFShadowTerms, AreaLightContext, Settings, INTEGRATION_AREA_LIGHT_CAPSULE);
+				BSDFEvaluate = SubstrateEvaluateBSDFCommon(BSDFContext, BSDFShadowTerms, AreaLight, Settings, INTEGRATION_AREA_LIGHT_CAPSULE);
 			}
 			else
 			{
-				BSDFEvaluate = SubstrateEvaluateBSDFCommon(BSDFContext, BSDFShadowTerms, AreaLightContext, Settings, INTEGRATION_PUNCTUAL_LIGHT);
+				BSDFEvaluate = SubstrateEvaluateBSDFCommon(BSDFContext, BSDFShadowTerms, AreaLight, Settings, INTEGRATION_PUNCTUAL_LIGHT);
 			}
 		}
 
diff --git a/Engine/Shaders/Private/Substrate/SubstrateEvaluation.ush b/Engine/Shaders/Private/Substrate/SubstrateEvaluation.ush
--- a/Engine/Shaders/Private/Substrate/SubstrateEvaluation.ush
+++ b/Engine/Shaders/Private/Substrate/SubstrateEvaluation.ush
@@ -448,7 +448,7 @@ void SubstrateSetupBSDFContext(inout FSubstrateBSDFContext BSDFContext, in bool
 	BSDFContext.SatNoL = saturate(BSDFContext.Context.NoL);
 }
 
-FSubstrateEvaluateResult SubstrateEvaluateBSDFCommon(FSubstrateBSDFContext BSDFContext, FShadowTerms ShadowTerms, FAreaLightIntegrateContext AreaLightContext, FSubstrateIntegrationSettings Settings, int IntegrationType)
+FSubstrateEvaluateResult SubstrateEvaluateBSDFCommon(FSubstrateBSDFContext BSDFContext, FShadowTerms ShadowTerms, FAreaLight AreaLight, FSubstrateIntegrationSettings Settings, int IntegrationType)
 {
 	FSubstrateEvaluateResult Sample = (FSubstrateEvaluateResult)0;
 
@@ -510,7 +510,7 @@ FSubstrateEvaluateResult SubstrateEvaluateBSDFCommon(FSubstrateBSDFContext BSDFC
 			float Alpha2Spec = Pow4(SafeRoughness);
 
 			float NoV, VoH, NoH;
-			SubstrateSetupBSDFContext(BSDFContext, bHasAnisotropy, NoV, VoH, NoH, AreaLightContext.L, AreaLightContext.AreaLight.SphereSinAlpha);
+			SubstrateSetupBSDFContext(BSDFContext, bHasAnisotropy, NoV, VoH, NoH, AreaLight.SpecularL, AreaLight.SphereSinAlpha);
 
 			////
 			//// Evaluate the diffuse component.
@@ -522,14 +522,14 @@ FSubstrateEvaluateResult SubstrateEvaluateBSDFCommon(FSubstrateBSDFContext BSDFC
 				// * The Chan model bakes transmittance specular directional albedo assuming F=0.04. In previous code we reapplied 
 				//   this transmittance, as the energy preservation code we remove it as well. However the visual effect was small 
 				//   and the cost was rather large (~8%). This is why we removed it in recent iteration
-				Sample.DiffusePathValue = Diffuse_Chan(DiffuseColor, Alpha2Spec, NoV, AreaLightContext.NoL, VoH, NoH, GetAreaLightDiffuseMicroReflWeight(AreaLightContext.AreaLight));
+				Sample.DiffusePathValue = Diffuse_Chan(DiffuseColor, Alpha2Spec, NoV, AreaLight.NoL, VoH, NoH, GetAreaLightDiffuseMicroReflWeight(AreaLight));
 			}
 			else
 #endif
 			{
 				Sample.DiffusePathValue = Diffuse_Lambert(DiffuseColor);
 			}
-			Sample.IntegratedDiffuseValue += (ShadowTerms.SurfaceShadow * AreaLightContext.NoL * AreaLightContext.Falloff) * Sample.DiffusePathValue * AreaLightContext.AreaLight.FalloffColor;
+			Sample.IntegratedDiffuseValue += (ShadowTerms.SurfaceShadow * AreaLight.NoL * AreaLight.Falloff) * Sample.DiffusePathValue * AreaLight.FalloffColor;
 			Sample.DiffuseColor		= DiffuseColor;
 			Sample.DiffusePDF		= BSDFContext.SatNoL / PI;
 			Sample.bPostProcessSubsurface = false;
@@ -595,7 +595,7 @@ FSubstrateEvaluateResult SubstrateEvaluateBSDFCommon(FSubstrateBSDFContext BSDFC
 
 							// GGX lobe
 							{
-								RectLightSpec = RectGGXApproxLTC(SafeRoughness, F0, F90, BSDFContext.N, BSDFContext.V, AreaLightContext.AreaLight.Rect, AreaLightContext.AreaLight.Texture, MeanLightWorldDirection);
+								RectLightSpec = RectGGXApproxLTC(SafeRoughness, F0, F90, BSDFContext.N, BSDFContext.V, AreaLight.Rect, AreaLight.Texture, MeanLightWorldDirection);
 							}
 							
 							// Glint shading only in the complex path
@@ -640,12 +640,12 @@ FSubstrateEvaluateResult SubstrateEvaluateBSDFCommon(FSubstrateBSDFContext BSDFC
 
 								// Area light energy compensation
 								{
-									const float Energy = EnergyNormalization(Alpha2Spec, BSDFContext.Context.VoH, AreaLightContext.AreaLight);
+									const float Energy = EnergyNormalization(Alpha2Spec, BSDFContext.Context.VoH, AreaLight);
 									D *= Energy;
 								}
 							}
 
-							Vis = Substrate_Vis_GGX(SafeRoughness, Alpha2Spec, BSDFContext.Context.NoV, AreaLightContext.NoL);
+							Vis = Substrate_Vis_GGX(SafeRoughness, Alpha2Spec, BSDFContext.Context.NoV, AreaLight.NoL);
 							const float H_PDF = VisibleGGXPDF(BSDFContext.TangentV, BSDFContext.TangentH, Alpha2Spec);
 							PDF = RayPDFToReflectionRayPDF(BSDFContext.Context.VoH, H_PDF);
 						}
@@ -673,7 +673,7 @@ FSubstrateEvaluateResult SubstrateEvaluateBSDFCommon(FSubstrateBSDFContext BSDFC
 			{
 				BSDFContext.N = TopNormal;
 				BSDFContext.TangentBasis[2] = TopNormal;
-				SubstrateSetupBSDFContext(BSDFContext, bHazeHasAnisotropy, NoV, VoH, NoH, AreaLightContext.L, AreaLightContext.AreaLight.SphereSinAlpha);
+				SubstrateSetupBSDFContext(BSDFContext, bHazeHasAnisotropy, NoV, VoH, NoH, AreaLight.SpecularL, AreaLight.SphereSinAlpha);
 			}
 		#endif
 
@@ -725,7 +725,7 @@ FSubstrateEvaluateResult SubstrateEvaluateBSDFCommon(FSubstrateBSDFContext BSDFC
 					if (bIsRectLight)
 					{
 						{
-							RectLightSpecHaze = RectGGXApproxLTC(HazeSafeRoughness, HazeF0, HazeF90, BSDFContext.N, BSDFContext.V, AreaLightContext.AreaLight.Rect, AreaLightContext.AreaLight.Texture);
+							RectLightSpecHaze = RectGGXApproxLTC(HazeSafeRoughness, HazeF0, HazeF90, BSDFContext.N, BSDFContext.V, AreaLight.Rect, AreaLight.Texture);
 					
 							const float H_PDF = VisibleGGXPDF(BSDFContext.TangentV, BSDFContext.TangentH, Pow4(HazeSafeRoughness));
 							HazePDF = RayPDFToReflectionRayPDF(BSDFContext.Context.VoH, H_PDF);
@@ -742,10 +742,10 @@ FSubstrateEvaluateResult SubstrateEvaluateBSDFCommon(FSubstrateBSDFContext BSDFC
 							}
 							else
 							{
-								const float Energy = EnergyNormalization(Alpha2SpecHaze, BSDFContext.Context.VoH, AreaLightContext.AreaLight);
+								const float Energy = EnergyNormalization(Alpha2SpecHaze, BSDFContext.Context.VoH, AreaLight);
 								HazeD = Substrate_D_GGX(HazeSafeRoughness, Alpha2SpecHaze, BSDFContext.Context.NoH) * Energy;
 							}
-							HazeVis = Substrate_Vis_GGX(HazeSafeRoughness, Alpha2SpecHaze, BSDFContext.Context.NoV, AreaLightContext.NoL);
+							HazeVis = Substrate_Vis_GGX(HazeSafeRoughness, Alpha2SpecHaze, BSDFContext.Context.NoV, AreaLight.NoL);
 							const float H_PDF = VisibleGGXPDF(BSDFContext.TangentV, BSDFContext.TangentH, Alpha2SpecHaze);
 							HazePDF = RayPDFToReflectionRayPDF(BSDFContext.Context.VoH, H_PDF);
 						}
@@ -836,11 +836,11 @@ FSubstrateEvaluateResult SubstrateEvaluateBSDFCommon(FSubstrateBSDFContext BSDFC
 				float3 CommonTerm = 0.0f;
 				if (bIsRectLight)
 				{
-					CommonTerm = ShadowTerms.SurfaceShadow; /* AreaLightContext.NoL and falloff is part of the LTC integration already */
+					CommonTerm = ShadowTerms.SurfaceShadow; /* AreaLight.NoL and falloff is part of the LTC integration already */
 				}
 				else
 				{
-					CommonTerm = (ShadowTerms.SurfaceShadow * AreaLightContext.NoL * AreaLightContext.Falloff) * AreaLightContext.AreaLight.FalloffColor;
+					CommonTerm = (ShadowTerms.SurfaceShadow * AreaLight.NoL * AreaLight.Falloff) * AreaLight.FalloffColor;
 				}
 				Sample.IntegratedSpecularValue *= CommonTerm;
 			}
@@ -920,7 +920,7 @@ FSubstrateEvaluateResult SubstrateEvaluateBSDFCommon(FSubstrateBSDFContext BSDFC
 
 				Sample.TransmissionPDF = 1.0f / (4.0f * PI); // SUBSTRATE_TODO this currently match the uniform sphere sampling from SubstrateImportanceSampleBSDF
 				Sample.TransmissionPathValue = TransmissionThroughput; // We are not applying DirectionalAlbedo_SpecularTransmission here to match legacy shading model for now.
-				Sample.IntegratedDiffuseValue += (ShadowTerms.TransmissionShadow * AreaLightContext.Falloff) * AreaLightContext.AreaLight.FalloffColor * Sample.TransmissionPathValue;
+				Sample.IntegratedDiffuseValue += (ShadowTerms.TransmissionShadow * AreaLight.Falloff) * AreaLight.FalloffColor * Sample.TransmissionPathValue;
 			}
 			#endif
 
@@ -975,7 +975,7 @@ FSubstrateEvaluateResult SubstrateEvaluateBSDFCommon(FSubstrateBSDFContext BSDFC
 				Sample.TransmissionPathValue	= TransmissionThroughput * PhaseFunction;
 				Sample.TransmissionPDF			= 1.0f / (4.0f * PI);		// SUBSTRATE_TODO this currently match the uniform sphere sampling from SubstrateImportanceSampleBSDF
 
-				Sample.IntegratedDiffuseValue += (ShadowTerms.TransmissionShadow * AreaLightContext.Falloff) * AreaLightContext.AreaLight.FalloffColor * Sample.TransmissionPathValue;
+				Sample.IntegratedDiffuseValue += (ShadowTerms.TransmissionShadow * AreaLight.Falloff) * AreaLight.FalloffColor * Sample.TransmissionPathValue;
 			}
 			#endif
 			#endif
@@ -1031,7 +1031,7 @@ FSubstrateEvaluateResult SubstrateEvaluateBSDFCommon(FSubstrateBSDFContext BSDFC
 
 				Sample.IntegratedDiffuseValue = lerp(
 					Sample.IntegratedDiffuseValue, 
-					(ShadowTerms.SurfaceShadow * MediumNoL * AreaLightContext.Falloff) * AreaLightContext.AreaLight.FalloffColor * SlabDirectionalAlbedo * DirectionalAlbedo_SpecularTransmission,
+					(ShadowTerms.SurfaceShadow * MediumNoL * AreaLight.Falloff) * AreaLight.FalloffColor * SlabDirectionalAlbedo * DirectionalAlbedo_SpecularTransmission,
 					DiffuseToVolumeBlend);
 				Sample.DiffuseColor		= SlabDirectionalAlbedo;
 			}
@@ -1061,11 +1061,11 @@ FSubstrateEvaluateResult SubstrateEvaluateBSDFCommon(FSubstrateBSDFContext BSDFC
 					BRANCH
 					if (bIsRectLight)
 					{
-						ClothSpecularPathValueNoL = FuzzF0 * RectSheenApproxLTC(FuzzRoughness, BSDFContext.N, BSDFContext.V, AreaLightContext.AreaLight.Rect, AreaLightContext.AreaLight.Texture, DirectionalAlbedo);
+						ClothSpecularPathValueNoL = FuzzF0 * RectSheenApproxLTC(FuzzRoughness, BSDFContext.N, BSDFContext.V, AreaLight.Rect, AreaLight.Texture, DirectionalAlbedo);
 					}
 					else
 					{
-						ClothSpecularPathValueNoL = FuzzF0 * SheenLTC_Eval(BSDFContext.V, AreaLightContext.L, BSDFContext.N, BSDFContext.Context.NoV, FuzzRoughness, View.SheenLTCTexture, View.SheenLTCSampler, DirectionalAlbedo);
+						ClothSpecularPathValueNoL = FuzzF0 * SheenLTC_Eval(BSDFContext.V, AreaLight.DiffuseL, BSDFContext.N, BSDFContext.Context.NoV, FuzzRoughness, View.SheenLTCTexture, View.SheenLTCSampler, DirectionalAlbedo);
 						ClothSpecularPathValueNoL *= DirectionalAlbedo;
 					}
 					EnergyTerms.E = DirectionalAlbedo;
@@ -1079,9 +1079,9 @@ FSubstrateEvaluateResult SubstrateEvaluateBSDFCommon(FSubstrateBSDFContext BSDFC
 					EnergyTerms = ComputeClothEnergyTermsA(FuzzRoughness, BSDFContext.Context.NoV);
 
 					float ClothD   = D_Charlie(FuzzRoughness, BSDFContext.Context.NoH);
-					float ClothVis = Vis_Ashikhmin(BSDFContext.Context.NoV, AreaLightContext.NoL);
+					float ClothVis = Vis_Ashikhmin(BSDFContext.Context.NoV, AreaLight.NoL);
 					float3 ClothF  = F_Schlick(FuzzF0, BSDFContext.Context.VoH);
-					ClothSpecularPathValueNoL = (ClothD * ClothVis) * ClothF * ComputeEnergyConservation(EnergyTerms) * AreaLightContext.NoL;
+					ClothSpecularPathValueNoL = (ClothD * ClothVis) * ClothF * ComputeEnergyConservation(EnergyTerms) * AreaLight.NoL;
 				}
 				#endif //SUBSTRATE_SHEEN_QUALITY
 
@@ -1089,7 +1089,7 @@ FSubstrateEvaluateResult SubstrateEvaluateBSDFCommon(FSubstrateBSDFContext BSDFC
 				const float3 Cloth_DirectionalAlbedo_SpecularTransmission = lerp(1.0, ComputeEnergyPreservation(EnergyTerms), FuzzAmount);
 
 				// Area light are not supported by the cloth BRDF
-				float3 ClothIntegratedSpecularValue = (ShadowTerms.SurfaceShadow * AreaLightContext.Falloff * FuzzAmount) * AreaLightContext.AreaLight.FalloffColor * ClothSpecularPathValueNoL;
+				float3 ClothIntegratedSpecularValue = (ShadowTerms.SurfaceShadow * AreaLight.Falloff * FuzzAmount) * AreaLight.FalloffColor * ClothSpecularPathValueNoL;
 				
 				// Apply specular transmittance to diffuse and specular lob from slab medium and interface sitting below the layer of fuzz.
 				Sample.DiffusePathValue			*= Cloth_DirectionalAlbedo_SpecularTransmission;
@@ -1149,7 +1149,7 @@ FSubstrateEvaluateResult SubstrateEvaluateBSDFCommon(FSubstrateBSDFContext BSDFC
 			Sample.SpecularPDF			= 1.0f / (4.0f * PI);		// SUBSTRATE_TODO this currently match the uniform sphere sampling from SubstrateImportanceSampleBSDF
 			Sample.ThroughputV			= OpaqueBSDFThroughput; 
 			Sample.TransmittanceAlongN	= OpaqueBSDFThroughput;
-			Sample.IntegratedSpecularValue = (ShadowTerms.TransmissionShadow * AreaLightContext.Falloff) * AreaLightContext.AreaLight.FalloffColor * Sample.SpecularPathValue;
+			Sample.IntegratedSpecularValue = (ShadowTerms.TransmissionShadow * AreaLight.Falloff) * AreaLight.FalloffColor * Sample.SpecularPathValue;
 		}
 		break;
 		
@@ -1171,8 +1171,8 @@ FSubstrateEvaluateResult SubstrateEvaluateBSDFCommon(FSubstrateBSDFContext BSDFC
 
 			float Alpha2 = Pow4(SafeRoughness);
 
-			Init(BSDFContext.Context, BSDFContext.N, BSDFContext.V, AreaLightContext.L);
-			SphereMaxNoH(BSDFContext.Context, AreaLightContext.AreaLight.SphereSinAlpha, true);
+			Init(BSDFContext.Context, BSDFContext.N, BSDFContext.V, AreaLight.SpecularL);
+			SphereMaxNoH(BSDFContext.Context, AreaLight.SphereSinAlpha, true);
 
 			BSDFContext.Context.NoV = saturate(max(abs(BSDFContext.Context.NoV), SUBSTRATE_EPSILON));
 
@@ -1193,10 +1193,10 @@ FSubstrateEvaluateResult SubstrateEvaluateBSDFCommon(FSubstrateBSDFContext BSDFC
 			const float Caustic = 0.8 + 0.2 * (Power + 1) * pow(saturate(dot(CausticNormal, BSDFContext.L)), Power);
 
 			const float Iris   = IrisNoL * Caustic;
-			const float Sclera = AreaLightContext.NoL;
+			const float Sclera = AreaLight.NoL;
 
 			Sample.DiffusePathValue = Diffuse_Lambert(DiffuseColor) * lerp(Sclera, Iris, EYE_IRISMASK(BSDFContext.BSDF));
-			Sample.IntegratedDiffuseValue += (ShadowTerms.SurfaceShadow * /* We don't AreaLightContext.NoL as it is already applied through the Sclera/Iris value above */ AreaLightContext.Falloff) * AreaLightContext.AreaLight.FalloffColor * Sample.DiffusePathValue;
+			Sample.IntegratedDiffuseValue += (ShadowTerms.SurfaceShadow * /* We don't AreaLight.NoL as it is already applied through the Sclera/Iris value above */ AreaLight.Falloff) * AreaLight.FalloffColor * Sample.DiffusePathValue;
 
 			Sample.DiffuseColor		= DiffuseColor;
 			Sample.DiffusePDF		= BSDFContext.SatNoL / PI;
@@ -1233,7 +1233,7 @@ FSubstrateEvaluateResult SubstrateEvaluateBSDFCommon(FSubstrateBSDFContext BSDFC
 			if (bIsRectLight)
 			{
 				// In this case, we set D to 1 and Vis will contain the area light / GGX integration
-				RectLightSpec = RectGGXApproxLTC(SafeRoughness, F0, F90, BSDFContext.N, BSDFContext.V, AreaLightContext.AreaLight.Rect, AreaLightContext.AreaLight.Texture);
+				RectLightSpec = RectGGXApproxLTC(SafeRoughness, F0, F90, BSDFContext.N, BSDFContext.V, AreaLight.Rect, AreaLight.Texture);
 
 				const float H_PDF = VisibleGGXPDF(BSDFContext.TangentV, BSDFContext.TangentH, Alpha2);
 				PDF = RayPDFToReflectionRayPDF(BSDFContext.Context.VoH, H_PDF);
@@ -1250,10 +1250,10 @@ FSubstrateEvaluateResult SubstrateEvaluateBSDFCommon(FSubstrateBSDFContext BSDFC
 				}
 				else
 				{
-					const float Energy = EnergyNormalization(Alpha2Spec, BSDFContext.Context.VoH, AreaLightContext.AreaLight);
+					const float Energy = EnergyNormalization(Alpha2Spec, BSDFContext.Context.VoH, AreaLight);
 					D = Substrate_D_GGX(SafeRoughness, Alpha2Spec, BSDFContext.Context.NoH) * Energy;
 				}
-				Vis = Substrate_Vis_GGX(SafeRoughness, Alpha2Spec, BSDFContext.Context.NoV, AreaLightContext.NoL);
+				Vis = Substrate_Vis_GGX(SafeRoughness, Alpha2Spec, BSDFContext.Context.NoV, AreaLight.NoL);
 				const float H_PDF = VisibleGGXPDF(BSDFContext.TangentV, BSDFContext.TangentH, Alpha2Spec);
 				PDF = RayPDFToReflectionRayPDF(BSDFContext.Context.VoH, H_PDF);
 			}
@@ -1266,12 +1266,12 @@ FSubstrateEvaluateResult SubstrateEvaluateBSDFCommon(FSubstrateBSDFContext BSDFC
 			if (bIsRectLight)
 			{
 				Sample.SpecularPathValue       = RectLightSpec * MSScale;
-				Sample.IntegratedSpecularValue = (ShadowTerms.SurfaceShadow * /* AreaLightContext.NoL and falloff is part of the LTC integration already */ Sample.SpecularPathProbability) * Sample.SpecularPathValue;
+				Sample.IntegratedSpecularValue = (ShadowTerms.SurfaceShadow * /* AreaLight.NoL and falloff is part of the LTC integration already */ Sample.SpecularPathProbability) * Sample.SpecularPathValue;
 			}
 			else
 			{
 				Sample.SpecularPathValue        = D * Vis * MSScale * FresnelTerm;
-				Sample.IntegratedSpecularValue += (ShadowTerms.SurfaceShadow * AreaLightContext.NoL * AreaLightContext.Falloff) * AreaLightContext.AreaLight.FalloffColor * Sample.SpecularPathValue * Sample.SpecularPathProbability;
+				Sample.IntegratedSpecularValue += (ShadowTerms.SurfaceShadow * AreaLight.NoL * AreaLight.Falloff) * AreaLight.FalloffColor * Sample.SpecularPathValue * Sample.SpecularPathProbability;
 			}
 
 			////
@@ -1303,11 +1303,12 @@ FSubstrateEvaluateResult SubstrateEvaluateBSDFCommon(FSubstrateBSDFContext BSDFC
 FSubstrateEvaluateResult SubstrateEvaluateBSDF(FSubstrateBSDFContext BSDFContext, FSubstrateIntegrationSettings Settings)
 {
 	FShadowTerms IdentityShadow = { 1, 1, 1, InitHairTransmittanceData() };
-	FAreaLightIntegrateContext DummyAreaLightContext = InitAreaLightIntegrateContext();
-	DummyAreaLightContext.L = BSDFContext.L;
-	DummyAreaLightContext.NoL = BSDFContext.Context.NoL;
-	DummyAreaLightContext.Falloff = 1;
-	return SubstrateEvaluateBSDFCommon(BSDFContext, IdentityShadow, DummyAreaLightContext, Settings, INTEGRATION_PUNCTUAL_LIGHT);
+	FAreaLight DummyAreaLight = InitAreaLight();
+	DummyAreaLight.DiffuseL = BSDFContext.L;
+	DummyAreaLight.SpecularL = BSDFContext.L;
+	DummyAreaLight.NoL = BSDFContext.Context.NoL;
+	DummyAreaLight.Falloff = 1;
+	return SubstrateEvaluateBSDFCommon(BSDFContext, IdentityShadow, DummyAreaLight, Settings, INTEGRATION_PUNCTUAL_LIGHT);
 }
 
 ///////////////////////////////////////////////////////////////////////////////////////////////////
@@ -1786,7 +1787,7 @@ FSubstrateEnvLightResult SubstrateEvaluateForEnvLight(FSubstrateBSDFContext BSDF
 
 			#if SUBSTRATE_FASTPATH==0
 			BRANCH
-				if (BSDF_GETSSSTYPE(BSDFContext.BSDF) == SSS_TYPE_WRAP || BSDF_GETSSSTYPE(BSDFContext.BSDF) == SSS_TYPE_TWO_SIDED_WRAP)
+			if (BSDF_GETSSSTYPE(BSDFContext.BSDF) == SSS_TYPE_WRAP || BSDF_GETSSSTYPE(BSDFContext.BSDF) == SSS_TYPE_TWO_SIDED_WRAP)
 			{
 				const bool bTwoSided = BSDF_GETSSSTYPE(BSDFContext.BSDF) == SSS_TYPE_TWO_SIDED_WRAP;
 				const float3 SlabDiffuseColor = bTwoSided ? DiffuseColor : float3(1, 1, 1);
diff --git a/Engine/Shaders/Private/Substrate/SubstrateLightingCommon.ush b/Engine/Shaders/Private/Substrate/SubstrateLightingCommon.ush
--- a/Engine/Shaders/Private/Substrate/SubstrateLightingCommon.ush
+++ b/Engine/Shaders/Private/Substrate/SubstrateLightingCommon.ush
@@ -67,7 +67,7 @@ float3 SubstrateForwardLightingCommon(
 			LightData, TranslatedWorldPosition, ToLight, LightAttenuation, Dither, ShadowTerms);
 
 		float Roughness = SubstrateGetBSDFRoughness(SubstrateBSDFContext.BSDF);
-		FAreaLightIntegrateContext AreaLightContext = InitAreaLightIntegrateContext();
+		FAreaLight AreaLight = InitAreaLight();
 		BRANCH
 		if (ShadowTerms.SurfaceShadow + ShadowTerms.TransmissionShadow > 0)
 		{
@@ -82,27 +82,27 @@ float3 SubstrateForwardLightingCommon(
 				{
 					LightMask = 0.0f; // Rect light can be non visible due to barn door occlusion
 				}
-				AreaLightContext = CreateRectIntegrateContext(Roughness, SubstrateBSDFContext.N, SubstrateBSDFContext.V, Rect, RectTexture);
+				AreaLight = CreateAreaLight(Roughness, SubstrateBSDFContext.N, SubstrateBSDFContext.V, Rect, RectTexture);
 				Lighting = IntegrateLight(Rect);
 
 				// We must have the evaluate inside the if due to the rectlight texture: it must be now be ambiguous which texture is going ot be used.
 				// After te compilation, a local resource must map to a unique global resource (the default or the actual rect light texture).
-				BSDFEvaluate = SubstrateEvaluateBSDFCommon(SubstrateBSDFContext, ShadowTerms, AreaLightContext, Settings, INTEGRATION_AREA_LIGHT_RECT);
+				BSDFEvaluate = SubstrateEvaluateBSDFCommon(SubstrateBSDFContext, ShadowTerms, AreaLight, Settings, INTEGRATION_AREA_LIGHT_RECT);
 			}
 			else
 			{
 				FCapsuleLight Capsule = GetCapsule(ToLight, LightData);
-				AreaLightContext = CreateCapsuleIntegrateContext(Roughness, SubstrateBSDFContext.N, SubstrateBSDFContext.V, Capsule, LightData.bInverseSquared);
+				AreaLight = CreateAreaLight(Roughness, SubstrateBSDFContext.N, SubstrateBSDFContext.V, Capsule, LightData.bInverseSquared);
 				Lighting = IntegrateLight(Capsule, LightData.bInverseSquared);
 
 				BRANCH
-				if(IsAreaLight(AreaLightContext.AreaLight))
+				if(IsAreaLight(AreaLight))
 				{
-					BSDFEvaluate = SubstrateEvaluateBSDFCommon(SubstrateBSDFContext, ShadowTerms, AreaLightContext, Settings, INTEGRATION_AREA_LIGHT_CAPSULE);
+					BSDFEvaluate = SubstrateEvaluateBSDFCommon(SubstrateBSDFContext, ShadowTerms, AreaLight, Settings, INTEGRATION_AREA_LIGHT_CAPSULE);
 				}
 				else
 				{
-					BSDFEvaluate = SubstrateEvaluateBSDFCommon(SubstrateBSDFContext, ShadowTerms, AreaLightContext, Settings, INTEGRATION_PUNCTUAL_LIGHT);
+					BSDFEvaluate = SubstrateEvaluateBSDFCommon(SubstrateBSDFContext, ShadowTerms, AreaLight, Settings, INTEGRATION_PUNCTUAL_LIGHT);
 				}
 			}
 
@@ -121,25 +121,25 @@ float3 SubstrateForwardLightingCommon(
 				{
 					LightMask = 0.0f; // Rect light can be non visible due to barn door occlusion
 				}
-				AreaLightContext = CreateRectIntegrateContext(Roughness, SubstrateBSDFContext.N, SubstrateBSDFContext.V, Rect, RectTexture);
+				AreaLight = CreateAreaLight(Roughness, SubstrateBSDFContext.N, SubstrateBSDFContext.V, Rect, RectTexture);
 
 				// We must have the evaluate inside the if due to the rectlight texture: it must be now be ambiguous which texture is going ot be used.
 				// After te compilation, a local resource must map to a unique global resource (the default or the actual rect light texture).
-				BSDFEvaluate = SubstrateEvaluateBSDFCommon(SubstrateBSDFContext, ShadowTerms, AreaLightContext, Settings, INTEGRATION_AREA_LIGHT_RECT);
+				BSDFEvaluate = SubstrateEvaluateBSDFCommon(SubstrateBSDFContext, ShadowTerms, AreaLight, Settings, INTEGRATION_AREA_LIGHT_RECT);
 			}
 			else
 			{
 				FCapsuleLight Capsule = GetCapsule(ToLight, LightData);
-				AreaLightContext = CreateCapsuleIntegrateContext(Roughness, SubstrateBSDFContext.N, SubstrateBSDFContext.V, Capsule, LightData.bInverseSquared);
+				AreaLight = CreateAreaLight(Roughness, SubstrateBSDFContext.N, SubstrateBSDFContext.V, Capsule, LightData.bInverseSquared);
 				
 				BRANCH
-				if(IsAreaLight(AreaLightContext.AreaLight))
+				if(IsAreaLight(AreaLight))
 				{
-					BSDFEvaluate = SubstrateEvaluateBSDFCommon(SubstrateBSDFContext, ShadowTerms, AreaLightContext, Settings, INTEGRATION_AREA_LIGHT_CAPSULE);
+					BSDFEvaluate = SubstrateEvaluateBSDFCommon(SubstrateBSDFContext, ShadowTerms, AreaLight, Settings, INTEGRATION_AREA_LIGHT_CAPSULE);
 				}
 				else
 				{
-					BSDFEvaluate = SubstrateEvaluateBSDFCommon(SubstrateBSDFContext, ShadowTerms, AreaLightContext, Settings, INTEGRATION_PUNCTUAL_LIGHT);
+					BSDFEvaluate = SubstrateEvaluateBSDFCommon(SubstrateBSDFContext, ShadowTerms, AreaLight, Settings, INTEGRATION_PUNCTUAL_LIGHT);
 				}
 			}
 
