Add minimal support for specular profile with Format=0.

Disable for now until we can enable this per material.

diff --git a/Engine/Plugins/TextureGraph/Shaders/TextureGraphMaterialShader.usf b/Engine/Plugins/TextureGraph/Shaders/TextureGraphMaterialShader.usf
--- a/Engine/Plugins/TextureGraph/Shaders/TextureGraphMaterialShader.usf
+++ b/Engine/Plugins/TextureGraph/Shaders/TextureGraphMaterialShader.usf
@@ -87,10 +87,12 @@ void TextureGraphMaterialShaderPS(
 	FSubstrateIntegrationSettings Settings = InitSubstrateIntegrationSettings();
 
 	const float3 SurfaceWorldNormal = MaterialParameters.TangentToWorld[2].xyz;
+	const float3 V = MaterialParameters.CameraVector;
 	FExportResult Export = SubstrateMaterialExportOut(
 		Settings,
 		SubstratePixelHeader,
 		SubstrateData,
+		V,
 		SurfaceWorldNormal,
 		MaterialParameters.WorldPosition_CamRelative,
 		0.0f // Curvature
diff --git a/Engine/Shaders/Private/AnisotropyPassShader.usf b/Engine/Shaders/Private/AnisotropyPassShader.usf
--- a/Engine/Shaders/Private/AnisotropyPassShader.usf
+++ b/Engine/Shaders/Private/AnisotropyPassShader.usf
@@ -134,10 +134,12 @@ void MainPixelShader(
 	FSubstrateIntegrationSettings Settings = InitSubstrateIntegrationSettings(false /*bForceFullyRough*/, false /*SubstrateStruct.bRoughDiffuse*/, 0/*SubstrateStruct.PeelLayersAboveDepth*/, false/*SubstrateStruct.bRoughnessTracking*/);
 
 	const float3 SurfaceWorldNormal = MaterialParameters.TangentToWorld[2].xyz;
+	const float3 V = MaterialParameters.CameraVector;
 	FExportResult Export = SubstrateMaterialExportOut(
 		Settings,
 		SubstratePixelHeader,
 		SubstrateData,
+		V,
 		SurfaceWorldNormal,
 		MaterialParameters.WorldPosition_CamRelative,
 		0.0f // Curvature
diff --git a/Engine/Shaders/Private/BasePassPixelShader.usf b/Engine/Shaders/Private/BasePassPixelShader.usf
--- a/Engine/Shaders/Private/BasePassPixelShader.usf
+++ b/Engine/Shaders/Private/BasePassPixelShader.usf
@@ -1727,10 +1727,13 @@ void FPixelShaderInOut_MainPS(
 	FSubstrateIntegrationSettings Settings = InitSubstrateIntegrationSettings(false /*bForceFullyRough*/, false /*SubstrateStruct.bRoughDiffuse*/, SubstrateStruct.PeelLayersAboveDepth, SubstrateStruct.bRoughnessTracking);
 
 	const float3 SurfaceWorldNormal = MaterialParameters.TangentToWorld[2].xyz;
+	// Align V with N to not have any view dependent effect when exporting material data.
+	const float3 FakeV = SurfaceWorldNormal;
 	FExportResult Export = SubstrateMaterialExportOut(
 		Settings,
 		SubstratePixelHeader,
 		SubstrateData,
+		FakeV,
 		SurfaceWorldNormal,
 		MaterialParameters.WorldPosition_CamRelative,
 		0.0f /*MobileShadingPathCurvature*/);
@@ -1866,10 +1869,12 @@ void FPixelShaderInOut_MainPS(
 		// GBuffer has been reset above already.
 
 		const float3 SurfaceWorldNormal = MaterialParameters.TangentToWorld[2].xyz;
+		const float3 V = MaterialParameters.CameraVector;
 		FExportResult Export = SubstrateMaterialExportOut(
 			Settings,
 			SubstratePixelHeader,
 			SubstrateData,
+			V,
 			SurfaceWorldNormal,
 			MaterialParameters.WorldPosition_CamRelative,
 			0.0f // Curvature
diff --git a/Engine/Shaders/Private/DeferredDecal.usf b/Engine/Shaders/Private/DeferredDecal.usf
--- a/Engine/Shaders/Private/DeferredDecal.usf
+++ b/Engine/Shaders/Private/DeferredDecal.usf
@@ -202,6 +202,7 @@ void FPixelShaderInOut_MainPS(inout FPixelShaderIn In, inout FPixelShaderOut Out
 			Settings,
 			SubstrateHeader,
 			PixelMaterialInputs.GetFrontSubstrateData(),
+			CameraVector,
 			SurfaceWorldNormal,
 			MaterialParameters.WorldPosition_CamRelative,
 			0.0f /*MobileShadingPathCurvature*/);
diff --git a/Engine/Shaders/Private/MeshDecals.usf b/Engine/Shaders/Private/MeshDecals.usf
--- a/Engine/Shaders/Private/MeshDecals.usf
+++ b/Engine/Shaders/Private/MeshDecals.usf
@@ -158,10 +158,12 @@ void FPixelShaderInOut_MainPS(
 	{
 		const FSubstrateIntegrationSettings Settings = InitSubstrateIntegrationSettings(false /*bForceFullyRough*/, false /*bRoughDiffuseEnabled*/, 0, false/*bRoughnessTracking*/);
 		const float3 SurfaceWorldNormal = MaterialParameters.TangentToWorld[2].xyz;
+		const float3 V = MaterialParameters.CameraVector;
 		FExportResult Export = SubstrateMaterialExportOut(
 			Settings,
 			SubstrateHeader,
 			PixelMaterialInputs.GetFrontSubstrateData(),
+			V,
 			SurfaceWorldNormal,
 			MaterialParameters.WorldPosition_CamRelative,
 			0.0f /*MobileShadingPathCurvature*/);
diff --git a/Engine/Shaders/Private/MobileBasePassPixelShader.usf b/Engine/Shaders/Private/MobileBasePassPixelShader.usf
--- a/Engine/Shaders/Private/MobileBasePassPixelShader.usf
+++ b/Engine/Shaders/Private/MobileBasePassPixelShader.usf
@@ -730,10 +730,12 @@ void Main(
 	#endif
 
 		const float3 SurfaceWorldNormal = MaterialParameters.TangentToWorld[2].xyz;
+		const float3 V = MaterialParameters.CameraVector;
 		FExportResult Export = SubstrateMaterialExportOut(
 			Settings,
 			SubstratePixelHeader,
 			SubstrateData,
+			V,
 			SurfaceWorldNormal,
 			MaterialParameters.WorldPosition_CamRelative,
 		#if SHADING_PATH_MOBILE && MATERIAL_SUBSURFACE_PROFILE_USE_CURVATURE
diff --git a/Engine/Shaders/Private/PathTracing/PathTracingGPULightmassMaterialHitShader.usf b/Engine/Shaders/Private/PathTracing/PathTracingGPULightmassMaterialHitShader.usf
--- a/Engine/Shaders/Private/PathTracing/PathTracingGPULightmassMaterialHitShader.usf
+++ b/Engine/Shaders/Private/PathTracing/PathTracingGPULightmassMaterialHitShader.usf
@@ -90,10 +90,12 @@ RAY_TRACING_ENTRY_CLOSEST_HIT(GPULightmassMaterialCHS,
 	FSubstrateIntegrationSettings Settings = InitSubstrateIntegrationSettings(true /*bForceFullyRough*/, false /*SubstrateStruct.bRoughDiffuse*/, -1/*SubstrateStruct.PeelLayersAboveDepth*/, false/*SubstrateStruct.bRoughnessTracking*/);
 
 	const float3 SurfaceWorldNormal = GeoNormal;
+	const float3 V = MaterialParameters.CameraVector;
 	FExportResult Export = SubstrateMaterialExportOut(
 		Settings,
 		SubstratePixelHeader,
 		SubstrateData,
+		V,
 		SurfaceWorldNormal,
 		MaterialParameters.WorldPosition_CamRelative,
 		0.0f /*MobileShadingPathCurvature*/);
@@ -406,10 +408,12 @@ RAY_TRACING_ENTRY_ANY_HIT(GPULightmassMaterialAHS,
 	FSubstrateIntegrationSettings Settings = InitSubstrateIntegrationSettings(true /*bForceFullyRough*/, false /*SubstrateStruct.bRoughDiffuse*/, -1/*SubstrateStruct.PeelLayersAboveDepth*/, false/*SubstrateStruct.bRoughnessTracking*/);
 
 	const float3 SurfaceWorldNormal = GeoNormal;
+	const float3 V = MaterialParameters.CameraVector;
 	FExportResult Export = SubstrateMaterialExportOut(
 		Settings,
 		SubstratePixelHeader,
 		SubstrateData,
+		V,
 		SurfaceWorldNormal,
 		MaterialParameters.WorldPosition_CamRelative,
 		0.0f /*MobileShadingPathCurvature*/);
diff --git a/Engine/Shaders/Private/RayTracing/RayTracingMaterialHitShaders.usf b/Engine/Shaders/Private/RayTracing/RayTracingMaterialHitShaders.usf
--- a/Engine/Shaders/Private/RayTracing/RayTracingMaterialHitShaders.usf
+++ b/Engine/Shaders/Private/RayTracing/RayTracingMaterialHitShaders.usf
@@ -447,10 +447,12 @@ void GetMaterialPayload(
 	FSubstrateIntegrationSettings Settings = InitSubstrateIntegrationSettings();
 
 	const float3 SurfaceWorldNormal = MaterialParameters.TangentToWorld[2].xyz;
+	const float3 V = MaterialParameters.CameraVector;
 	FExportResult Export = SubstrateMaterialExportOut(
 		Settings,
 		SubstratePixelHeader,
 		SubstrateData,
+		V,
 		SurfaceWorldNormal,
 		MaterialParameters.WorldPosition_CamRelative,
 		0.0f // Curvature
diff --git a/Engine/Shaders/Private/Substrate/SubstrateExport.ush b/Engine/Shaders/Private/Substrate/SubstrateExport.ush
--- a/Engine/Shaders/Private/Substrate/SubstrateExport.ush
+++ b/Engine/Shaders/Private/Substrate/SubstrateExport.ush
@@ -943,6 +943,7 @@ FExportResult SubstrateMaterialExportOut(
 	in FSubstrateIntegrationSettings Settings,
 	in FSubstratePixelHeader SubstratePixelHeader,
 	in FSubstrateData SubstrateData,
+	in float3 V,
 	in float3 SurfaceWorldNormal,
 	in float3 WorldPosition_CamRelative,
 	in float MobileShadingPathCurvature)
@@ -951,19 +952,24 @@ FExportResult SubstrateMaterialExportOut(
 	Out.IndirectIrradianceScale = 1.0;
 	Out.SubsurfaceProfileRadiusScale = 1.0;
 
-	// Align V with N to not have any view dependent effect when exporting material data.
-	const float3 FakeV = SurfaceWorldNormal;
-	
 	if (SubstratePixelHeader.SubstrateTree.BSDFCount > 0)
 	{
 		// Update tree (coverage/transmittance/luminace weights)
-		SubstratePixelHeader.SubstrateUpdateTree(SubstrateData, FakeV, Settings, Out.Coverage, Out.TransmittancePreCoverage);
+		SubstratePixelHeader.SubstrateUpdateTree(SubstrateData, V, Settings, Out.Coverage, Out.TransmittancePreCoverage);
 
 		const int BSDFIdx = 0;
 		if(SubstratePixelHeader.SubstrateTree.BSDFCount > 0)
 		{
 		#define CurrentBSDF SubstratePixelHeader.SubstrateTree.BSDFs[BSDFIdx]
 
+			// Bake the Specular Profiles into F0/F90 of BSDFs to blend in order to retain some of the original look when using parameter blending or fully simplified materials.
+			// Disable for now until we have mean to enable this per material
+		#if 0 //SUBSTRATE_SPECPROFILE_ENABLED
+			const float NoV = saturate(dot(SurfaceWorldNormal, V));
+			CurrentBSDF.SubstrateBakeSpecularProfileToReflectivity(NoV);
+			BSDF_SETHASSPECPROFILE(CurrentBSDF, 0);
+		#endif
+
 			if (SubstrateIsBSDFVisible(CurrentBSDF))
 			{
 				const uint BSDFType = BSDF_GETTYPE(CurrentBSDF);
