Substrate - fixed missing anisotropy in some cases due to bad define default value.

#tests pc editor dx12
#rnx

diff --git a/Engine/Shaders/Private/Substrate/Substrate.ush b/Engine/Shaders/Private/Substrate/Substrate.ush
--- a/Engine/Shaders/Private/Substrate/Substrate.ush
+++ b/Engine/Shaders/Private/Substrate/Substrate.ush
@@ -143,6 +143,12 @@
 #define SUBSTRATE_OPAQUE_MATERIAL 0
 #endif
 
+// This is set for material writting to GBuffer for instance based on input connection from the material graph.
+// However, when InternalGetSubstrateSlabBSDF is called form SubstrateReadGBufferBSDF (from many global shaders) we have to assume that we can read and setup anisotropic material so we use the project setting only.
+#ifndef MATERIAL_USES_ANISOTROPY
+#define MATERIAL_USES_ANISOTROPY SUPPORTS_ANISOTROPIC_MATERIALS
+#endif
+
 // Behind the scene, we do all the simple volumetric lighting & transmitance computations considering a slab of 100 centimeters = 1 meter.
 // If the user specifies a different thickness to simulate varying appearance, we rescale the mean free path in order to not have to store thickness as part of the BSDF description.
 #define SUBSTRATE_SIMPLEVOLUME_THICKNESS_CM	100.0f
@@ -1604,7 +1610,7 @@ FSubstrateData InternalGetSubstrateSlabBSDF(
 	const uint bHasAnisotropy = Anisotropy != 0 ? 1 : 0;
 #else
 	Anisotropy = 0.0f;
-	const uint bHasAnisotropy = false;
+	const uint bHasAnisotropy = 0;
 #endif
 	if (!SUBSTRATE_FASTPATH && bHasAnisotropy)
 	{
