Fix hair strands shading with Substrate Format=0.

#test editor dx12, substrate=0/1, format=0/1

diff --git a/Engine/Shaders/Private/HairStrands/HairStrandsMaterialCommon.ush b/Engine/Shaders/Private/HairStrands/HairStrandsMaterialCommon.ush
--- a/Engine/Shaders/Private/HairStrands/HairStrandsMaterialCommon.ush
+++ b/Engine/Shaders/Private/HairStrands/HairStrandsMaterialCommon.ush
@@ -210,7 +210,7 @@ FEvaluateOutput Evaluate(
 		float3 TranslatedWorldPosition = SvPositionToResolvedTranslatedWorld(SvPosition);
 		CalcMaterialParametersEx(MaterialParameters, PixelMaterialInputs, SvPosition, LocalScreenPosition, bIsFrontFace, TranslatedWorldPosition, TranslatedWorldPosition);
 
-#if SUBSTRATE_GBUFFER_FORMAT==1 && MATERIAL_IS_SUBSTRATE && MATERIAL_SHADINGMODEL_HAIR
+#if MATERIAL_IS_SUBSTRATE && MATERIAL_SHADINGMODEL_HAIR
 		FSubstratePixelHeader SubstratePixelHeader = MaterialParameters.GetFrontSubstrateHeader();
 		SubstratePixelHeader.SetCastContactShadow(GetPrimitiveData(MaterialParameters).Flags & PRIMITIVE_SCENE_DATA_FLAG_CAST_SHADOWS);
 		SubstratePixelHeader.SetDynamicIndirectShadowCasterRepresentation(GetPrimitiveData(MaterialParameters).Flags & PRIMITIVE_SCENE_DATA_FLAG_HAS_CAPSULE_REPRESENTATION);
