Fix iris shading with Substrate when using material attribute with CustomTangent/CustomBottomNormal output nodes at the same time.

diff --git a/Engine/Shaders/Private/Substrate/SubstrateExport.ush b/Engine/Shaders/Private/Substrate/SubstrateExport.ush
--- a/Engine/Shaders/Private/Substrate/SubstrateExport.ush
+++ b/Engine/Shaders/Private/Substrate/SubstrateExport.ush
@@ -1048,7 +1048,7 @@ FExportResult SubstrateMaterialExportOut(
 					Out.Metallic = IrisDistance;
 
 					#if NUM_MATERIAL_OUTPUTS_GETTANGENTOUTPUT > 0
-						float3 Tangent = T /*GetTangentOutput0(MaterialParameters)*/;
+						float3 Tangent = IrisNormal;
 						Out.CustomData.yz = UnitVectorToOctahedron( normalize(Tangent) ) * 0.5 + 0.5;
 					#endif
 				#endif
diff --git a/Engine/Shaders/Private/Substrate/SubstrateRead.ush b/Engine/Shaders/Private/Substrate/SubstrateRead.ush
--- a/Engine/Shaders/Private/Substrate/SubstrateRead.ush
+++ b/Engine/Shaders/Private/Substrate/SubstrateRead.ush
@@ -99,6 +99,7 @@ FSubstrateGBufferBSDF SubstrateReadGBufferBSDF(in FScreenSpaceData ScreenSpaceDa
 			const float2 IrisNormalDelta = float2(GBuffer.CustomData.y, GBuffer.CustomData.z) * 2 - (256.0 / 255.0);
 			const float2 WorldNormalOct = UnitVectorToOctahedron(GBuffer.WorldNormal);
 			IrisNormal = OctahedronToUnitVector(WorldNormalOct + IrisNormalDelta);
+			IrisDistance = 1.0; // Set at 1 in order to let's the mask drive the interpolation. See SubstrateEvaluation.ush.
 		#else
 			IrisNormal = OctahedronToUnitVector(GBuffer.CustomData.yz * 2 - 1);
 			IrisDistance = GBuffer.StoredMetallic;
diff --git a/Engine/Shaders/Private/Substrate/SubstrateVisualizeCommon.ush b/Engine/Shaders/Private/Substrate/SubstrateVisualizeCommon.ush
--- a/Engine/Shaders/Private/Substrate/SubstrateVisualizeCommon.ush
+++ b/Engine/Shaders/Private/Substrate/SubstrateVisualizeCommon.ush
@@ -11,23 +11,20 @@
 #error SUBSTRATE_MATERIALCONTAINER_IS_VIEWRESOURCE needs to be defined
 #endif
 
-void DrawReferentialTWS(float3 P /*In Translated World Space*/, float3 X, float3 Y, float3 N, float3 InColor)
+void AddTipLineTWS(inout FShaderPrintContext Ctx, float3 TWSPos, float3 In, float4 InTipColor, float4 InCoreColor=ColorYellow, float Scale=1.0f)
+{
+	const float CoreSize = 0.8f * Scale;
+	const float TipSize  = 1.0f * Scale;
+	AddLineTWS(Ctx, TWSPos,                 TWSPos + In * CoreSize, InCoreColor, InCoreColor);
+	AddLineTWS(Ctx, TWSPos + In * CoreSize, TWSPos + In * TipSize,  InTipColor,  InTipColor);
+}
+
+void DrawReferentialTWS(inout FShaderPrintContext Ctx, float3 P /*In Translated World Space*/, float3 X, float3 Y, float3 N, float4 InColor)
 {
 	const float Size = 10.f;
-	const float SizeWithTip = 12.0f;
-	const float4 Color = float4(InColor, 1);
-	const float4 ColorX = float4(1, 0, 0, 1);
-	const float4 ColorY = float4(0, 1, 0, 1);
-	const float4 ColorZ = float4(0, 0, 1, 1);
-
-	// Core
-	AddLineTWS(P, P + X * Size, Color, Color);
-	AddLineTWS(P, P + Y * Size, Color, Color);
-	AddLineTWS(P, P + N * Size, Color, Color);
-	// Tips
-	AddLineTWS(P + X * Size, P + X * SizeWithTip, ColorX, ColorX);
-	AddLineTWS(P + Y * Size, P + Y * SizeWithTip, ColorY, ColorY);
-	AddLineTWS(P + N * Size, P + N * SizeWithTip, ColorZ, ColorZ);
+	AddTipLineTWS(Ctx, P, X, ColorRed, InColor, Size);
+	AddTipLineTWS(Ctx, P, Y, ColorBlue, InColor, Size);
+	AddTipLineTWS(Ctx, P, N, ColorGreen, InColor, Size);
 }
 
 void AddDrawPixelFootprint(float3 P, float3 dPdx, float3 dPdy, float2 Scale, bool bNormalize, float4 Color)
@@ -128,7 +125,10 @@ void SubstratePrintBSDF(
 	const float3 WorldNormal = BSDFContext.N;
 
 	// Draw Referential
-	DrawReferentialTWS(WorldPosition, BSDFContext.X, BSDFContext.Y, BSDFContext.N, float3(1, 1, 0));
+	if (BSDF_GETTYPE(BSDFContext.BSDF) != SUBSTRATE_BSDF_TYPE_EYE)
+	{
+		DrawReferentialTWS(Context, WorldPosition, BSDFContext.X, BSDFContext.Y, BSDFContext.N, ColorYellow);
+	}
 
 	switch (BSDF_GETTYPE(BSDFContext.BSDF))
 	{
@@ -144,8 +144,7 @@ void SubstratePrintBSDF(
 				if (Haziness.HasBottomNormal())
 				{
 					const float3 BottomNormal = SubstrateGetSimpleCoatBottomNormal(Haziness.BottomNormalOct, WorldNormal);
-					AddLineTWS(WorldPosition, WorldPosition + BottomNormal * 10.0f, ColorYellow, ColorYellow);							// Core
-					AddLineTWS(WorldPosition + BottomNormal * 10.f, WorldPosition + BottomNormal * 12.0f, ColorPurple, ColorPurple);	// Tip				
+					AddTipLineTWS(Context, WorldPosition, BottomNormal, ColorPurple, ColorYellow, 10.f);
 				}
 			#endif
 			}
@@ -160,14 +159,10 @@ void SubstratePrintBSDF(
 		case SUBSTRATE_BSDF_TYPE_EYE:
 		{
 			// Tangent basis for the iris
-			const float3x3 IrisTangentBasis = GetTangentBasis(EYE_IRISNORMAL(BSDFContext.BSDF));
-			//DrawReferentialTWS(WorldPosition, IrisTangentBasis[0], IrisTangentBasis[1], IrisTangentBasis[2], float3(0, 1, 1));
-			AddLineTWS(WorldPosition, WorldPosition + EYE_IRISNORMAL(BSDFContext.BSDF) * LineSize, float4(0, 1, 1, 1));
-
-			// Tangent basis for the iris plane
-			const float3x3 IrisPlaneTangentBasis = GetTangentBasis(EYE_IRISPLANENORMAL(BSDFContext.BSDF));
-			//DrawReferentialTWS(WorldPosition, IrisPlaneTangentBasis[0], IrisPlaneTangentBasis[1], IrisPlaneTangentBasis[2], float3(1, 0, 1));
-			AddLineTWS(WorldPosition, WorldPosition + EYE_IRISPLANENORMAL(BSDFContext.BSDF) * LineSize, float4(1, 0, 1, 1));
+			const float EyeLineSize = 1.f;
+			AddTipLineTWS(Context, WorldPosition, BSDFContext.N, ColorBlue, ColorYellow, EyeLineSize);
+			AddTipLineTWS(Context, WorldPosition, EYE_IRISNORMAL(BSDFContext.BSDF), ColorRed, ColorYellow, EyeLineSize);
+			AddTipLineTWS(Context, WorldPosition, EYE_IRISPLANENORMAL(BSDFContext.BSDF), ColorGreen, ColorYellow, EyeLineSize);
 
 		}
 		break;
diff --git a/Engine/Source/Runtime/Engine/Private/Materials/HLSLMaterialTranslator.cpp b/Engine/Source/Runtime/Engine/Private/Materials/HLSLMaterialTranslator.cpp
--- a/Engine/Source/Runtime/Engine/Private/Materials/HLSLMaterialTranslator.cpp
+++ b/Engine/Source/Runtime/Engine/Private/Materials/HLSLMaterialTranslator.cpp
@@ -1818,7 +1818,8 @@ void FHLSLMaterialTranslator::TranslateMaterial()
 					return UMaterialExpressionSubstrateConvertMaterialAttributes::CompileCommon(this, 0,
 						FMaterialAttributeDefinitionMap::GetConnectedMaterialAttributesBitmask(BaseMaterial->GetExpressions()), EditorOnlyData->MaterialAttributes, MaterialInterface->GetShadingModels().GetFirstShadingModel(),
 						*WaterScatteringCoefficients, *WaterAbsorptionCoefficients, *WaterPhaseG, *ColorScaleBehindWater,
-						bHasSSS, bHasSSS ? MaterialInterface->SubsurfaceProfile : nullptr);
+						bHasSSS, bHasSSS ? MaterialInterface->SubsurfaceProfile : nullptr,
+						*ClearCoatNormal, *CustomTangent);
 
 				}
 				else if ((Material->GetMaterialDomain() == MD_Surface || Material->GetMaterialDomain() == MD_RuntimeVirtualTexture || Material->GetMaterialDomain() == MD_DeferredDecal) && EditorOnlyData)
diff --git a/Engine/Source/Runtime/Engine/Private/Materials/MaterialExpressionSubstrate.cpp b/Engine/Source/Runtime/Engine/Private/Materials/MaterialExpressionSubstrate.cpp
--- a/Engine/Source/Runtime/Engine/Private/Materials/MaterialExpressionSubstrate.cpp
+++ b/Engine/Source/Runtime/Engine/Private/Materials/MaterialExpressionSubstrate.cpp
@@ -3825,7 +3825,8 @@ FExpressionInput* UMaterialExpressionSubstrateConvertMaterialAttributes::GetInpu
 int32 UMaterialExpressionSubstrateConvertMaterialAttributes::CompileCommon(class FMaterialCompiler* Compiler, int32 OutputIndex,
 	const uint64 CachedConnectedMaterialAttributesInputs, FMaterialAttributesInput& MaterialAttributes, enum EMaterialShadingModel ShadingModelOverride,
 	FExpressionInput& WaterScatteringCoefficients, FExpressionInput& WaterAbsorptionCoefficients, FExpressionInput& WaterPhaseG, FExpressionInput& ColorScaleBehindWater,
-	const bool bHasSSS, USubsurfaceProfile* SSSProfile)
+	const bool bHasSSS, USubsurfaceProfile* SSSProfile,
+	FExpressionInput& ClearCoatNormal, FExpressionInput& CustomTangent)
 {
 	static const FGuid ClearCoatBottomNormalGuid = FMaterialAttributeDefinitionMap::GetCustomAttributeID(TEXT("ClearCoatBottomNormal"));
 	static const FGuid CustomEyeTangentGuid = FMaterialAttributeDefinitionMap::GetCustomAttributeID(TEXT("CustomEyeTangent"));
@@ -3876,6 +3877,13 @@ int32 UMaterialExpressionSubstrateConvertMaterialAttributes::CompileCommon(class
 		ClearCoat_NewRegisteredSharedLocalBasis = SubstrateCompilationInfoCreateSharedLocalBasis(Compiler, ClearCoat_NormalCodeChunk, ClearCoat_TangentCodeChunk);
 		ClearCoat_BasisIndexMacro = Compiler->GetSubstrateSharedLocalBasisIndexMacro(ClearCoat_NewRegisteredSharedLocalBasis);
 	}
+	else if (ClearCoatNormal.IsConnected())
+	{
+		ClearCoat_NormalCodeChunk = CompileWithDefaultNormalWS(Compiler, ClearCoatNormal);
+		ClearCoat_TangentCodeChunk = TangentCodeChunk;
+		ClearCoat_NewRegisteredSharedLocalBasis = SubstrateCompilationInfoCreateSharedLocalBasis(Compiler, ClearCoat_NormalCodeChunk, ClearCoat_TangentCodeChunk);
+		ClearCoat_BasisIndexMacro = Compiler->GetSubstrateSharedLocalBasisIndexMacro(ClearCoat_NewRegisteredSharedLocalBasis);
+	}
 	else
 	{
 		ClearCoat_NormalCodeChunk = NormalCodeChunk;
@@ -3897,6 +3905,11 @@ int32 UMaterialExpressionSubstrateConvertMaterialAttributes::CompileCommon(class
 			CustomTangent_TangentCodeChunk = Compiler->VertexNormal();
 		}
 	}
+	else if (CustomTangent.IsConnected())
+	{
+		// Legacy code doesn't do tangent <-> world basis conversion on tangent output, when provided.
+		CustomTangent_TangentCodeChunk = CompileWithDefaultNormalWS(Compiler, CustomTangent, false /*bConvertToRequestedSpace*/);
+	}
 	else
 	{
 		CustomTangent_TangentCodeChunk = NormalCodeChunk;
@@ -4004,9 +4017,11 @@ int32 UMaterialExpressionSubstrateConvertMaterialAttributes::CompileCommon(class
 
 int32 UMaterialExpressionSubstrateConvertMaterialAttributes::Compile(class FMaterialCompiler* Compiler, int32 OutputIndex)
 {
+	FExpressionInput DummyInput;
 	return CompileCommon(Compiler, OutputIndex, GetConnectedMaterialAttributesInputs(Material), MaterialAttributes, ShadingModelOverride,
 		WaterScatteringCoefficients, WaterAbsorptionCoefficients, WaterPhaseG, ColorScaleBehindWater,
-		HasSSS(), SubsurfaceProfile);
+		HasSSS(), SubsurfaceProfile,
+		DummyInput, DummyInput);
 }
 
 void UMaterialExpressionSubstrateConvertMaterialAttributes::PostEditChangeProperty(FPropertyChangedEvent& PropertyChangedEvent)
diff --git a/Engine/Source/Runtime/Engine/Public/Materials/MaterialExpressionSubstrate.h b/Engine/Source/Runtime/Engine/Public/Materials/MaterialExpressionSubstrate.h
--- a/Engine/Source/Runtime/Engine/Public/Materials/MaterialExpressionSubstrate.h
+++ b/Engine/Source/Runtime/Engine/Public/Materials/MaterialExpressionSubstrate.h
@@ -988,7 +988,8 @@ class UMaterialExpressionSubstrateConvertMaterialAttributes : public UMaterialEx
 	static int32 CompileCommon(class FMaterialCompiler* Compiler, int32 OutputIndex,
 		const uint64 CachedConnectedMaterialAttributesInputs, FMaterialAttributesInput& MaterialAttributes, enum EMaterialShadingModel ShadingModelOverride,
 		FExpressionInput& WaterScatteringCoefficients, FExpressionInput& WaterAbsorptionCoefficients, FExpressionInput& WaterPhaseG, FExpressionInput& ColorScaleBehindWater,
-		const bool bHasSSS, USubsurfaceProfile* SSSProfile);
+		const bool bHasSSS, USubsurfaceProfile* SSSProfile,
+		FExpressionInput& ClearCoatNormal, FExpressionInput& CustomTangent);
 
 #endif // WITH_EDITOR
 	//~ End UMaterialExpression Interface
