Substrate - Fixed many shader compilation warnings and unecessary float/integer conversions.

#tests pc editor and consol cook
#rnx

diff --git a/Engine/Shaders/Private/ShaderPrintCommon.ush b/Engine/Shaders/Private/ShaderPrintCommon.ush
--- a/Engine/Shaders/Private/ShaderPrintCommon.ush
+++ b/Engine/Shaders/Private/ShaderPrintCommon.ush
@@ -269,7 +269,7 @@ struct FPackedShaderPrintItem
 
 FPackedShaderPrintItem PackShaderPrintItem(FShaderPrintItem In)
 {
-	const uint3 Color8bits = saturate(In.Color) * 0xFF;
+	const uint3 Color8bits = uint3(saturate(In.Color) * 255.0f);
 	const uint2 ScreenPos16bit = f32tof16(In.ScreenPos);
 
 	FPackedShaderPrintItem Out;
diff --git a/Engine/Shaders/Private/SpecularProfileCommon.ush b/Engine/Shaders/Private/SpecularProfileCommon.ush
--- a/Engine/Shaders/Private/SpecularProfileCommon.ush
+++ b/Engine/Shaders/Private/SpecularProfileCommon.ush
@@ -4,7 +4,7 @@
 
 uint GetSpecularProfileId(float In)
 {
-	return abs(In);
+	return uint(abs(In));
 }
 
 // Return 0: View/Light angles, 1: Half angles
diff --git a/Engine/Shaders/Private/Substrate/Substrate.ush b/Engine/Shaders/Private/Substrate/Substrate.ush
--- a/Engine/Shaders/Private/Substrate/Substrate.ush
+++ b/Engine/Shaders/Private/Substrate/Substrate.ush
@@ -1004,7 +1004,7 @@ uint PackSSSProfile(float InProfileId, float InRadius, half InThickness)
 
 void UnpackSSSProfile(uint In, inout float OutProfileId, inout float OutRadius, inout half OutThickness)
 {
-	OutProfileId = BitFieldExtractU32(In, 8, 0);
+	OutProfileId = float(BitFieldExtractU32(In, 8, 0));
 	OutRadius = UnpackG8(In);
 	OutThickness = Unpack10F((In >> 16) & 0x3FF);
 }
@@ -1681,7 +1681,7 @@ FSubstrateData InternalGetSubstrateSlabBSDF(
 
 	if (SSSType == SSS_TYPE_DIFFUSION_PROFILE)
 	{
-		SLAB_SSSPROFILEID	(SubstrateData.InlinedBSDF)			= SSSProfileID;
+		SLAB_SSSPROFILEID	(SubstrateData.InlinedBSDF)			= float(SSSProfileID);
 		SLAB_SSSPROFILERADIUSSCALE(SubstrateData.InlinedBSDF)	= SSSMFPScale;
 		SLAB_SSSPROFILETHICKNESSCM(SubstrateData.InlinedBSDF)	= ThicknessCm;
 	}
@@ -1892,7 +1892,7 @@ FSubstrateData GetSubstrateEyeBSDF(float3 DiffuseAlbedo, float3 F0, float Roughn
 	EYE_IRISMASK					(SubstrateData.InlinedBSDF) = IrisMask;
 	EYE_IRISNORMAL					(SubstrateData.InlinedBSDF) = IrisNormal;
 	EYE_IRISPLANENORMAL				(SubstrateData.InlinedBSDF) = IrisPlaneNormal;
-	EYE_SSSPROFILEID				(SubstrateData.InlinedBSDF) = SSSProfileId;
+	EYE_SSSPROFILEID				(SubstrateData.InlinedBSDF) = float(SSSProfileId);
 	EYE_F0							(SubstrateData.InlinedBSDF) = F0;
 	EYE_F90							(SubstrateData.InlinedBSDF) = saturate(F0 * 50.f);
 
@@ -3492,7 +3492,7 @@ void FSubstrateBSDF::SubstrateBakeSpecularProfileToReflectivity(in float NoV)
 	if (BSDF_GETHASSPECPROFILE(this))
 	{
 		// Evaluate the spec lut
-		const uint SpecProfileID = SLAB_SPECPROFILEID(this);
+		const float SpecProfileID = SLAB_SPECPROFILEID(this);
 		const float3 SpecLUTSample = EvaluateSpecularProfile(SpecProfileID, NoV, 1.0f/*NoL*/, 1.0f/*VoH*/, 1.0f/*NoH*/);
 		SLAB_F0(this) *= SpecLUTSample;
 		SLAB_F90(this) *= SpecLUTSample;
@@ -4582,7 +4582,7 @@ FSubstrateBSDF UnpackSubstrateBSDFIn(FSubstrateMaterialContainer SubstrateBuffer
 				SLAB_F0(OutBSDF)			= OutF0;
 				SLAB_ROUGHNESS(OutBSDF)		= SubstrateUnpackTopLayerData(SubstrateHeader.PackedTopLayerData).Roughness;
 
-				SLAB_SSSPROFILEID(OutBSDF)			= BitFieldExtractU32(Data2, 8, 0);
+				SLAB_SSSPROFILEID(OutBSDF)			= float(BitFieldExtractU32(Data2, 8, 0));
 				SLAB_SSSPROFILERADIUSSCALE(OutBSDF)	= UnpackG8(Data2);
 				SLAB_HAZINESS(OutBSDF)				= BitFieldExtractU32(Data2, 16, 16);
 
diff --git a/Engine/Shaders/Private/Substrate/SubstrateVisualize.usf b/Engine/Shaders/Private/Substrate/SubstrateVisualize.usf
--- a/Engine/Shaders/Private/Substrate/SubstrateVisualize.usf
+++ b/Engine/Shaders/Private/Substrate/SubstrateVisualize.usf
@@ -27,7 +27,7 @@
 uint bOverrideCursorPosition;
 int2 GetCursorPosition()
 {
-	const int2 OverrideCursorPosition = (View.ViewSizeAndInvSize.xy + View.ViewRectMin.xy) / 2;
+	const int2 OverrideCursorPosition = int2((View.ViewSizeAndInvSize.xy + View.ViewRectMin.xy) / 2.0f);
 	return bOverrideCursorPosition > 0 ? OverrideCursorPosition : View.CursorPosition;	
 }
 
@@ -49,7 +49,7 @@ float3 GetWorldPositionFromPixelPos(uint2 PixelPos)
 [numthreads(1, 1, 1)]
 void MaterialPrintInfoCS(uint3 DispatchThreadId : SV_DispatchThreadID)
 {
-	const uint2 PixelPos = (float2(GetCursorPosition()) * View.ViewResolutionFraction);
+	const uint2 PixelPos = uint2(float2(GetCursorPosition()) * View.ViewResolutionFraction);
 	const float3 WorldPosition = GetWorldPositionFromPixelPos(PixelPos);
 	const float3 V = -GetCameraVectorFromTranslatedWorldPosition(WorldPosition);
 
diff --git a/Engine/Shaders/Shared/SubstrateVisualizeDefinitions.h b/Engine/Shaders/Shared/SubstrateVisualizeDefinitions.h
--- a/Engine/Shaders/Shared/SubstrateVisualizeDefinitions.h
+++ b/Engine/Shaders/Shared/SubstrateVisualizeDefinitions.h
@@ -598,7 +598,7 @@ FSubstratePixelDebugData ConvertToSerializableSubstratePixelDebugData(uint2 InCo
 				Data.Closures[i].GlintValue = 1.0;	// No glint
 			}
 
-			Data.Closures[i].SpecProfileID = -1.0f;
+			Data.Closures[i].SpecProfileID = -1;
 			if (BSDF_GETHASSPECPROFILE(BSDF))
 			{
 				Data.Closures[i].SpecProfileID = GetSpecularProfileId(SLAB_SPECPROFILEID(BSDF));
@@ -682,10 +682,6 @@ FSubstratePixelDebugData ConvertToSerializableSubstratePixelDebugData(uint2 InCo
 		const uint TopLayerDataBytes = SUBSTRATE_GBUFFER_FORMAT == 1 ? 4 : 0;
 		const uint SSSDataBytes = 8;
 
-		const bool bSubstrateMaterial = Header.ClosureCount > 0;
-		const bool bIsSimpleMaterial  = Header.IsSimpleMaterial() || Header.ClosureCount == 0;
-		const bool bIsSingleMaterial  = !Header.IsSimpleMaterial() && Header.IsSingleMaterial();
-
 		const uint HeaderSize = FootPrint_PostHeader - FootPrint_Start;
 		const uint BSDFsSize  = FootPrint_PostBSDFs  - FootPrint_PostHeader;
 		const uint TotalSize  = (FootPrint_PostBSDFs - FootPrint_Start) + (bHasSSSData ? SSSDataBytes : 0);
