Fix missing/incorrect env. / SSR contribution when using Format=0 with Substrate.

Energy conservation is disabled by default when using Format=0. This caused the specular energy to be too weak.

diff --git a/Engine/Shaders/Private/ShadingEnergyConservation.ush b/Engine/Shaders/Private/ShadingEnergyConservation.ush
--- a/Engine/Shaders/Private/ShadingEnergyConservation.ush
+++ b/Engine/Shaders/Private/ShadingEnergyConservation.ush
@@ -52,6 +52,11 @@
 ///////////////////////////////////////////////////////////////////////////////////////////////////
 // LUT Lookups
 
+float2 GGXEnergyLookupEnv(float Roughness, float NoV)
+{
+	return View.ShadingEnergyGGXSpecTexture.SampleLevel(SharedShadingEnergySampler, float2(NoV, Roughness), 0);
+}
+
 float2 GGXEnergyLookup(float Roughness, float NoV)
 {
 #if USE_ENERGY_CONSERVATION == 1
diff --git a/Engine/Shaders/Private/Substrate/SubstrateEvaluation.ush b/Engine/Shaders/Private/Substrate/SubstrateEvaluation.ush
--- a/Engine/Shaders/Private/Substrate/SubstrateEvaluation.ush
+++ b/Engine/Shaders/Private/Substrate/SubstrateEvaluation.ush
@@ -1532,6 +1532,17 @@ FSubstrateImportanceSampleResult SubstrateImportanceSampleBSDF(FSubstrateBSDFCon
 	return Out;
 }
 
+// Energy conservation used for env. evaluation. 
+// Format=0 disables energy conservation for performance reason. However In order to keep env. evalution correct, we force simple energy lookup).
+FBxDFEnergyTermsRGB ComputeGGXSpecEnergyTermsRGB_Env(float SafeRoughness, float NoV, float3 F0, float3 F90)
+{
+	#if USE_ENERGY_CONSERVATION > 0
+	return ComputeGGXSpecEnergyTermsRGB(SafeRoughness, NoV, F0, F90);
+	#else
+	return ComputeFresnelEnergyTermsRGB(GGXEnergyLookupEnv(SafeRoughness, NoV), F0, F90);
+	#endif
+}
+
 struct FSubstrateEnvLightResult
 {
 	float3 DiffuseNormal;
@@ -1592,7 +1603,7 @@ FSubstrateEnvLightResult SubstrateEvaluateForEnvLight(FSubstrateBSDFContext BSDF
 			float3 EvalEnvBRDF = 0;
 			float DirectionalAlbedo_SpecularTransmission = 0;
 			{
-				FBxDFEnergyTermsRGB SpecularEnergyTerms = ComputeGGXSpecEnergyTermsRGB(SafeRoughness, BSDFContext.Context.NoV, F0, F90);
+				FBxDFEnergyTermsRGB SpecularEnergyTerms = ComputeGGXSpecEnergyTermsRGB_Env(SafeRoughness, BSDFContext.Context.NoV, F0, F90);
 				FBxDFEnergyTermsRGB DiffuseEnergyTerms  = ComputeDiffuseEnergyTermsRGB(SafeRoughness, BSDFContext.Context.NoV);
 				EvalEnvBRDF								= SpecularEnergyTerms.E;
 				DirectionalAlbedo_SpecularTransmission	= ComputeEnergyPreservation(SpecularEnergyTerms);
@@ -1726,7 +1737,7 @@ FSubstrateEnvLightResult SubstrateEvaluateForEnvLight(FSubstrateBSDFContext BSDF
 						const float3 HazeClearCoatTransmittance = SimpleClearCoatTransmittance(ClearCoatBottomContextNoL, ClearCoatBottomContextNoV, SubstrateGetBSDFMetallic(BSDFContext.BSDF), SubstrateGetBSDFBaseColor(BSDFContext.BSDF));
 						const float HazeF0 = 0.04f;
 						const float HazeF90= 1.0f;
-						FBxDFEnergyTermsRGB EnergyTerms = ComputeGGXSpecEnergyTermsRGB(HazeSafeRoughness, ClearCoatTopContextNoV, HazeF0, HazeF90);
+						FBxDFEnergyTermsRGB EnergyTerms = ComputeGGXSpecEnergyTermsRGB_Env(HazeSafeRoughness, ClearCoatTopContextNoV, HazeF0, HazeF90);
 						const float3 EvalEnvBRDFHaze = EnergyTerms.E;
 						const float3 TopLayerSpecularTransmittionApprox = saturate(1.0f - EvalEnvBRDFHaze); // Simple and do not use extra LUT for now
 
@@ -1757,7 +1768,7 @@ FSubstrateEnvLightResult SubstrateEvaluateForEnvLight(FSubstrateBSDFContext BSDF
 						SubstrateEnvLightResult.SSRSpecularWeight *= BlendFactor;
 
 						// Compute the directional albedo for hazy specular
-						FBxDFEnergyTermsRGB EnergyTerms = ComputeGGXSpecEnergyTermsRGB(HazeSafeRoughness, BSDFContext.Context.NoV, F0, F90);
+						FBxDFEnergyTermsRGB EnergyTerms = ComputeGGXSpecEnergyTermsRGB_Env(HazeSafeRoughness, BSDFContext.Context.NoV, F0, F90);
 						const float3 EvalEnvBRDFHaze = EnergyTerms.E;
 						SubstrateEnvLightResult.SpecularHazeWeight = lerp(0.0f, HazeWeight, HazinessFadeIn) * EvalEnvBRDFHaze;
 						SubstrateEnvLightResult.SpecularHazeSafeRoughness = HazeSafeRoughness;
@@ -1909,7 +1920,7 @@ FSubstrateEnvLightResult SubstrateEvaluateForEnvLight(FSubstrateBSDFContext BSDF
 			float3 EvalEnvBRDF = 0;
 			float3 DirectionalAlbedo_SpecularTransmission = 0;
 			{
-				FBxDFEnergyTermsRGB SpecularEnergyTerms = ComputeGGXSpecEnergyTermsRGB(SafeRoughness, BSDFContext.Context.NoV, F0, F90);
+				FBxDFEnergyTermsRGB SpecularEnergyTerms = ComputeGGXSpecEnergyTermsRGB_Env(SafeRoughness, BSDFContext.Context.NoV, F0, F90);
 				const float DiffuseDirectionalAlbedo = 1.0f;
 				EvalEnvBRDF									= SpecularEnergyTerms.E;
 				DirectionalAlbedo_SpecularTransmission		= ComputeEnergyPreservation(SpecularEnergyTerms);
