Substrate - Fixed  mobile renderer deferred shader compilation.

#rnx
#tests PC editor and console cook

diff --git a/Engine/Shaders/Private/MobileBasePassPixelShader.usf b/Engine/Shaders/Private/MobileBasePassPixelShader.usf
--- a/Engine/Shaders/Private/MobileBasePassPixelShader.usf
+++ b/Engine/Shaders/Private/MobileBasePassPixelShader.usf
@@ -767,7 +767,7 @@ void Main(
 			Export.Roughness,
 			Export.Anisotropy,
 			Export.SubsurfaceColor,
-			Export.SubsurfaceProfileID,
+			Export.SubsurfaceProfileIDFloat,
 			0.0f,					// Dither
 			Export.ShadingModelID
 		);
