Ensure SpecularProfile setup & evalution code is only added when a specular profile is set on a material (in deferred and forward).

diff --git a/Engine/Shaders/Private/Substrate/Substrate.ush b/Engine/Shaders/Private/Substrate/Substrate.ush
--- a/Engine/Shaders/Private/Substrate/Substrate.ush
+++ b/Engine/Shaders/Private/Substrate/Substrate.ush
@@ -133,7 +133,13 @@
 
 #define SUBSTRATE_GLINTS_ENABLED (SUBSTRATE_COMPLEXSPECIALPATH && PLATFORM_ENABLES_SUBSTRATE_GLINTS)
 
+// If MATERIAL_USES_SPECULAR_PROFILE is defined (i.e.: deferred: base pass, forward: base pass + lighting), use the definition to enable/disable specular profile. 
+// Otherwise enable specular evaluation based on project settings to ensure correct evaluation in global shader
+#ifdef MATERIAL_USES_SPECULAR_PROFILE
+#define SUBSTRATE_SPECPROFILE_ENABLED (MATERIAL_USES_SPECULAR_PROFILE && PLATFORM_ENABLES_SUBSTRATE_SPECULAR_PROFILE)
+#else
 #define SUBSTRATE_SPECPROFILE_ENABLED (SUBSTRATE_COMPLEXPATH && PLATFORM_ENABLES_SUBSTRATE_SPECULAR_PROFILE)
+#endif
 
 #define SUBSTRATE_LEGACY_MATERIAL_APPLIES_FINAL_WEIGHT (!MATERIALBLENDING_MASKED && !MATERIALBLENDING_SOLID && SUBSTRATE_LEGACY_PREMULT_ALPHA_OVERRIDE==0 || SUBSTRATE_MATERIAL_EXPORT_FROM_TRANSLUCENT)
 
@@ -149,7 +155,7 @@
 #endif
 
 // This is set for material writting to GBuffer for instance based on input connection from the material graph.
-// However, when InternalGetSubstrateSlabBSDF is called form SubstrateReadGBufferBSDF (from many global shaders) we have to assume that we can read and setup anisotropic material so we use the project setting only.
+// However, when InternalGetSubstrateSlabBSDF is called from SubstrateReadGBufferBSDF (from many global shaders) we have to assume that we can read and setup anisotropic material so we use the project setting only.
 #ifndef MATERIAL_USES_ANISOTROPY
 #define MATERIAL_USES_ANISOTROPY SUPPORTS_ANISOTROPIC_MATERIALS
 #endif
diff --git a/Engine/Shaders/Private/Substrate/SubstrateEvaluation.ush b/Engine/Shaders/Private/Substrate/SubstrateEvaluation.ush
--- a/Engine/Shaders/Private/Substrate/SubstrateEvaluation.ush
+++ b/Engine/Shaders/Private/Substrate/SubstrateEvaluation.ush
@@ -824,7 +824,7 @@ FSubstrateEvaluateResult SubstrateEvaluateBSDFCommon(FSubstrateBSDFContext BSDFC
 			}
 
 			// Lighting LUT
-			#if SUBSTRATE_FASTPATH==0 && SUBSTRATE_SPECPROFILE_ENABLED
+			#if SUBSTRATE_SPECPROFILE_ENABLED
 			BRANCH
 			if (BSDF_GETHASSPECPROFILE(BSDFContext.BSDF))
 			{
