Fix crash on AmbientCubemapComposition with Substrate and Format=0 due to mismatch PS/VS.

Substrate permutation use tile VS, while legacy use fullscreen quad VS.

diff --git a/Engine/Shaders/Private/AmbientCubemapComposite.usf b/Engine/Shaders/Private/AmbientCubemapComposite.usf
--- a/Engine/Shaders/Private/AmbientCubemapComposite.usf
+++ b/Engine/Shaders/Private/AmbientCubemapComposite.usf
@@ -19,7 +19,7 @@
 Texture2D AmbientOcclusionTexture;
 SamplerState AmbientOcclusionSampler;
 
-#if SUBSTRATE_GBUFFER_FORMAT==1
+#if SUBSTRATE_ENABLED
 
 // Override the default lighting input used for the env. lighting to sample the 'custom composite env. lighting'
 float3 GetEnvDiffuseLighting(float3 InBentNormal)
