ShaderCompiler: Workaround for FXC error "X3524: duplicate attribute branch" with newer SPIRV-Tools

#rnx

diff --git a/Engine/Shaders/Private/DeferredLightingCommon.ush b/Engine/Shaders/Private/DeferredLightingCommon.ush
--- a/Engine/Shaders/Private/DeferredLightingCommon.ush
+++ b/Engine/Shaders/Private/DeferredLightingCommon.ush
@@ -396,7 +396,10 @@ FLightAccumulator AccumulateDynamicLighting(
 		LightAccumulator_AddSplit( LightAccumulator, Lighting.Diffuse, Lighting.Specular, Lighting.Diffuse, MaskedLightColor * Shadow.SurfaceShadow, bNeedsSeparateSubsurfaceLightAccumulation );
 		LightAccumulator_AddSplit( LightAccumulator, Lighting.Transmission, 0.0f, Lighting.Transmission, MaskedLightColor * Shadow.TransmissionShadow, bNeedsSeparateSubsurfaceLightAccumulation );
 #else // SHADING_PATH_MOBILE
+
+	#if !COMPILER_FXC // Work around "error X3524: duplicate attribute branch" in cases where this branch is constant
 		BRANCH
+	#endif
 		if( Shadow.SurfaceShadow + Shadow.TransmissionShadow > 0 )
 		{
 			const bool bNeedsSeparateSubsurfaceLightAccumulation = UseSubsurfaceProfile(GBuffer.ShadingModelID);
