Fix incorrect albedo on Nanite objects with Substrate when using Format=0

diff --git a/Engine/Shaders/Private/ComputeShaderOutputCommon.ush b/Engine/Shaders/Private/ComputeShaderOutputCommon.ush
--- a/Engine/Shaders/Private/ComputeShaderOutputCommon.ush
+++ b/Engine/Shaders/Private/ComputeShaderOutputCommon.ush
@@ -88,7 +88,7 @@ FPixelShaderOut ShadePixel(const float2 SVPositionXY, uint QuadIndex, uint QuadP
 	FPixelShaderInOut_MainPS(Interpolants, BasePassInterpolants, PixelShaderIn, PixelShaderOut, EyeIndex, QuadPixelWriteMask);
 #endif
 
-#if !SUBSTRATE_ENABLED
+#if !SUBSTRATE_ENABLED || (SUBSTRATE_ENABLED && SUBSTRATE_GBUFFER_FORMAT == 0)
 	if (!bHighPrecision)
 	{
 		PixelShaderOut.MRT[3] = float4(LinearToSrgb(PixelShaderOut.MRT[3].rgb), PixelShaderOut.MRT[3].a); // BaseColor is sRGB
