Fix eye shading mismatch with Substrate when a eye material does not provide a SSS.

Eye shading always run with SSS even when no profile is provided.

diff --git a/Engine/Shaders/Private/Substrate/Substrate.ush b/Engine/Shaders/Private/Substrate/Substrate.ush
--- a/Engine/Shaders/Private/Substrate/Substrate.ush
+++ b/Engine/Shaders/Private/Substrate/Substrate.ush
@@ -1902,7 +1902,8 @@ FSubstrateData GetSubstrateEyeBSDF(float3 DiffuseAlbedo, float3 F0, float Roughn
 	EYE_F0							(SubstrateData.InlinedBSDF) = F0;
 	EYE_F90							(SubstrateData.InlinedBSDF) = saturate(F0 * 50.f);
 
-	const bool bHasSSS = SSSProfileId != 0;
+	// Legacy eye shading always enables SSS on eye material. If no SSS profile is provided, Profile=0 will be used.
+	const bool bHasSSS = true;
 	BSDF_SETSSSTYPE(SubstrateData.InlinedBSDF, bHasSSS ? SSS_TYPE_DIFFUSION_PROFILE : SSS_TYPE_NONE);
 #if SUBSTRATE_INLINE_SHADING
 	SubstrateData.InlinedBSDF.Coverage = 1.0f;
