Substrate: Conditionally compile in support for only the enabled legacy shading models

Reduces code generation complexity a tiny bit. Also fix one field left uninitialized that is only used by the clearcoat case.

#rnx

diff --git a/Engine/Shaders/Private/Substrate/SubstrateLegacyConversion.ush b/Engine/Shaders/Private/Substrate/SubstrateLegacyConversion.ush
--- a/Engine/Shaders/Private/Substrate/SubstrateLegacyConversion.ush
+++ b/Engine/Shaders/Private/Substrate/SubstrateLegacyConversion.ush
@@ -84,6 +84,7 @@ FSubstrateLegacyParameters InitSubstrateLegacyParameters(FSubstratePixelFootprin
 	Out.SecondRoughness = 0.0f;
 	Out.SecondRoughnessWeight = 0.0f;
 	Out.SecondRoughnessAsSimpleClearCoat = false;
+	Out.ClearCoatUseBottomNormal = false;
 	Out.FuzzAmount = 0.f;
 	Out.FuzzColor = Zeros;
 	Out.Thickness = SUBSTRATE_LAYER_DEFAULT_THICKNESS_CM;
@@ -137,7 +138,7 @@ FSubstrateData CreateLegacySlab(
 		0.0f,												// Specular profile
 		In.Thickness,
 		In.bIsThinSurface,
-		true,												// bIsAtTheBottomOfTopology, legacy always has a single BSDF being at the bottom and top at teh same time.
+		true,												// bIsAtTheBottomOfTopology, legacy always has a single BSDF being at the bottom and top at the same time.
 		In.SharedLocalBasisIndex, SharedLocalBasisTypes);
 }
 
@@ -235,7 +236,7 @@ FSubstrateData SubstrateConvertLegacyMaterialDynamic(
 	LegacySlab.RawNormal = RawNormal;
 
 	float FinalWeight = 1.0f;
-
+	#if MATERIAL_SHADINGMODEL_UNLIT
 	if (ShadingModel == SHADINGMODELID_UNLIT)
 	{
 		// Unlit is handle with a emissive slab
@@ -251,7 +252,11 @@ FSubstrateData SubstrateConvertLegacyMaterialDynamic(
 		FinalWeight = Opacity;
 #endif
 	}
-	else if (ShadingModel == SHADINGMODELID_DEFAULT_LIT)
+	else
+	#endif
+
+	#if MATERIAL_SHADINGMODEL_DEFAULT_LIT
+	if (ShadingModel == SHADINGMODELID_DEFAULT_LIT)
 	{
 		LegacySlab.BaseColor = BaseColor;
 		LegacySlab.Specular = Specular;
@@ -265,7 +270,11 @@ FSubstrateData SubstrateConvertLegacyMaterialDynamic(
 		FinalWeight = Opacity;
 #endif
 	}
-	else if (ShadingModel == SHADINGMODELID_SUBSURFACE)
+	else
+	#endif
+
+	#if MATERIAL_SHADINGMODEL_SUBSURFACE
+	if (ShadingModel == SHADINGMODELID_SUBSURFACE)
 	{
 		const float ThicknessInCm = SUBSTRATE_SIMPLEVOLUME_THICKNESS_CM;
 		const float3 MFPInCm = TransmittanceToMeanFreePath(SubSurfaceColor, ThicknessInCm * CENTIMETER_TO_METER) * METER_TO_CENTIMETER;
@@ -297,7 +306,11 @@ FSubstrateData SubstrateConvertLegacyMaterialDynamic(
 		FinalWeight = Opacity;
 #endif
 	}
-	else if (ShadingModel == SHADINGMODELID_PREINTEGRATED_SKIN)
+	else
+	#endif
+
+	#if MATERIAL_SHADINGMODEL_PREINTEGRATED_SKIN
+	if (ShadingModel == SHADINGMODELID_PREINTEGRATED_SKIN)
 	{
 		// Use default profile MFP (e.g., human skin) as MFP value for converting hardcoded pre-integrated SSS texture for deferred material.
 		const float3 MFPInCm = float3(1.0f, 0.088964f, 0.072095f) * 2.6748f * 0.1f;
@@ -321,7 +334,11 @@ FSubstrateData SubstrateConvertLegacyMaterialDynamic(
 		FinalWeight = Opacity;
 #endif
 	}
-	else if (ShadingModel == SHADINGMODELID_SUBSURFACE_PROFILE)
+	else
+	#endif
+
+	#if MATERIAL_SHADINGMODEL_SUBSURFACE_PROFILE
+	if (ShadingModel == SHADINGMODELID_SUBSURFACE_PROFILE)
 	{
 		LegacySlab.BaseColor = BaseColor;
 		LegacySlab.Specular = Specular;
@@ -344,7 +361,11 @@ FSubstrateData SubstrateConvertLegacyMaterialDynamic(
 		FinalWeight = Opacity;
 #endif
 	}
-	else if (ShadingModel == SHADINGMODELID_CLEAR_COAT)
+	else
+	#endif
+
+	#if MATERIAL_SHADINGMODEL_CLEAR_COAT
+	if (ShadingModel == SHADINGMODELID_CLEAR_COAT)
 	{
 		// When using SUBSTRATE_GBUFFER_FORMAT, data are stored into fixed GBuffer layout which does not contain bit for bUseBottomNormal
 		// In this case we always force bUseBottomNormal if CLEAR_COAT_BOTTOM_NORMAL is enabled
@@ -375,7 +396,11 @@ FSubstrateData SubstrateConvertLegacyMaterialDynamic(
 		FinalWeight = Opacity;
 #endif
 	}
-	else if (ShadingModel == SHADINGMODELID_TWOSIDED_FOLIAGE)
+	else
+	#endif
+
+	#if MATERIAL_SHADINGMODEL_TWOSIDED_FOLIAGE
+	if (ShadingModel == SHADINGMODELID_TWOSIDED_FOLIAGE)
 	{
 		// Set a thickness that will enabled the thin lighting model (corresponding to the legacy two-sided lighting model)
 		const float Thickness = SUBSTRATE_LAYER_DEFAULT_THICKNESS_CM;
@@ -399,7 +424,11 @@ FSubstrateData SubstrateConvertLegacyMaterialDynamic(
 		FinalWeight = Opacity;
 #endif
 	}
-	else if (ShadingModel == SHADINGMODELID_CLOTH)
+	else
+	#endif
+	
+	#if MATERIAL_SHADINGMODEL_CLOTH
+	if (ShadingModel == SHADINGMODELID_CLOTH)
 	{
 		LegacySlab.BaseColor = BaseColor;
 		LegacySlab.Specular = Specular;
@@ -418,7 +447,11 @@ FSubstrateData SubstrateConvertLegacyMaterialDynamic(
 		FinalWeight = Opacity;
 #endif
 	}
-	else if (ShadingModel == SHADINGMODELID_THIN_TRANSLUCENT)	
+	else
+	#endif
+	
+	#if MATERIAL_SHADINGMODEL_THIN_TRANSLUCENT
+	if (ShadingModel == SHADINGMODELID_THIN_TRANSLUCENT)	
 	{
 		const float Thickness = DefaultThickness;
 
@@ -450,6 +483,11 @@ FSubstrateData SubstrateConvertLegacyMaterialDynamic(
 		// Thin translucency model have a special surface coverage input.
 		FinalWeight = ThinTranslucentSurfaceCoverage;
 	}
+	else
+	#endif
+	{
+		// base case of else-if chain above
+	}
 
 	// Fixed layer structure for helping compiler to unroll and optimize shader
 	FSubstrateData Slab = CreateLegacySlab(LegacySlab, SharedLocalBasisTypes);
