Path Tracer: Implement additional visualization modes

These cover specific features and specific fields related to Substrate support

#rnx

diff --git a/Engine/Shaders/Private/PathTracing/PathTracingCommon.ush b/Engine/Shaders/Private/PathTracing/PathTracingCommon.ush
--- a/Engine/Shaders/Private/PathTracing/PathTracingCommon.ush
+++ b/Engine/Shaders/Private/PathTracing/PathTracingCommon.ush
@@ -687,7 +687,7 @@ struct FPathTracingPayload : FMinimalPayload
 	float3 WorldSmoothNormal;        // smooth normal before normal/bump is applied
 #endif
 
-	float  BSDFOpacity;              // how much should we weigh down the bsdf by? (for glass, controls the blend between solid and glass lobes)
+	float  BSDFOpacity;              // how much should we weigh down the bsdf by? (this also accounts for the probability of picking a particular lobe in multi-lobed materials)
 	float3 TransparencyColor;        // how much do we see straight through the surface?
 	uint   ShadingModelID;
 	uint   Flags;
diff --git a/Engine/Shaders/Private/PathTracing/PathTracingDebug.usf b/Engine/Shaders/Private/PathTracing/PathTracingDebug.usf
--- a/Engine/Shaders/Private/PathTracing/PathTracingDebug.usf
+++ b/Engine/Shaders/Private/PathTracing/PathTracingDebug.usf
@@ -992,7 +992,8 @@ FPathTracingResult PathTracingDebugCore(FRayDesc Ray, uint2 LaunchIndex)
 			FLightLoopCount LightLoopCount = LightGridLookup(TranslatedWorldPos);
 
 			float FakeShading = abs(dot(WorldNormal, Ray.Direction)); // cheap shading, just to give some definition to otherwise "solid color" modes
-			float3 Color = 0, ExtraAlpha = 0;
+			float3 Color = 0;
+			float ExtraAlpha = 0;
 			switch (DebugMode)
 			{
 				case PATH_TRACER_DEBUG_VIZ_RADIANCE: 					Color = HitPayload.Radiance * View.PreExposure; break;
@@ -1019,7 +1020,7 @@ FPathTracingResult PathTracingDebugCore(FRayDesc Ray, uint2 LaunchIndex)
 				case PATH_TRACER_DEBUG_VIZ_CUSTOM_DATA0:				Color = HitPayload.CustomData0.rgb * View.PreExposure; ExtraAlpha = HitPayload.CustomData0.a; break;
 				case PATH_TRACER_DEBUG_VIZ_CUSTOM_DATA1:				Color = HitPayload.CustomData1.rgb * View.PreExposure; ExtraAlpha = HitPayload.CustomData1.a; break;
 #endif
-				case PATH_TRACER_DEBUG_VIZ_OPACITY:						Color = HitPayload.TransparencyColor; break;
+				case PATH_TRACER_DEBUG_VIZ_TRANSPARENCY_COLOR:			Color = HitPayload.TransparencyColor; break;
 				case PATH_TRACER_DEBUG_VIZ_IOR:							Color = HitPayload.Ior; break;
 				case PATH_TRACER_DEBUG_VIZ_SHADING_MODEL:				Color = GetShadingModelColor(HitPayload.ShadingModelID); break;
 				case PATH_TRACER_DEBUG_VIZ_LIGHTING_CHANNEL_MASK:		Color = GetLightingChannelMaskDebugColor(HitPayload.PrimitiveLightingChannelMask); break;
@@ -1082,10 +1083,16 @@ FPathTracingResult PathTracingDebugCore(FRayDesc Ray, uint2 LaunchIndex)
 					}
 					break;
 				}
+				case PATH_TRACER_DEBUG_VIZ_BSDF_OPACITY:    Color = HitPayload.BSDFOpacity; break;
 #if SUBSTRATE_ENABLED
 				case PATH_TRACER_DEBUG_VIZ_FUZZ_COLOR:     	Color = (HitPayload.ShadingModelID == SHADINGMODELID_SUBSTRATE) ? HitPayload.FuzzColor     : 0.0; break;
 				case PATH_TRACER_DEBUG_VIZ_FUZZ_ROUGHNESS:	Color = (HitPayload.ShadingModelID == SHADINGMODELID_SUBSTRATE) ? HitPayload.FuzzRoughness : 0.0; break;
 				case PATH_TRACER_DEBUG_VIZ_FUZZ_AMOUNT:		Color = (HitPayload.ShadingModelID == SHADINGMODELID_SUBSTRATE) ? HitPayload.FuzzAmount    : 0.0; break;
+				case PATH_TRACER_DEBUG_VIZ_SUBSTRATE_WEIGHT_V: 			Color = HitPayload.WeightV; 			break;
+				case PATH_TRACER_DEBUG_VIZ_SUBSTRATE_TRANSMITTANCE_N: 	Color = HitPayload.TransmittanceN; 		break;
+				case PATH_TRACER_DEBUG_VIZ_SUBSTRATE_COVERAGE_ABOVE_N: 	Color = HitPayload.CoverageAboveAlongN; break;
+				case PATH_TRACER_DEBUG_VIZ_SUBSTRATE_F90:				Color = HitPayload.SpecularEdgeColor;   break;
+
 #else
 				case PATH_TRACER_DEBUG_VIZ_FUZZ_COLOR:     	Color = (HitPayload.ShadingModelID == SHADINGMODELID_CLOTH) ? HitPayload.GetClothColor()  : 0.0; break;
 				case PATH_TRACER_DEBUG_VIZ_FUZZ_ROUGHNESS: 	Color = (HitPayload.ShadingModelID == SHADINGMODELID_CLOTH) ? HitPayload.Roughness        : 0.0; break;
@@ -1098,6 +1105,7 @@ FPathTracingResult PathTracingDebugCore(FRayDesc Ray, uint2 LaunchIndex)
 			}
 
 			Result.Radiance.rgb += PathThroughput * Color;
+			Result.Radiance.a += Luminance(PathThroughput) * ExtraAlpha;
 		}
 
 		if (PackedPayload.ShouldOutputDepth() && Result.Z == 0.0)
diff --git a/Engine/Shaders/Shared/PathTracingDefinitions.h b/Engine/Shaders/Shared/PathTracingDefinitions.h
--- a/Engine/Shaders/Shared/PathTracingDefinitions.h
+++ b/Engine/Shaders/Shared/PathTracingDefinitions.h
@@ -86,7 +86,7 @@
 #define PATH_TRACER_DEBUG_VIZ_BASE_COLOR				4 
 #define PATH_TRACER_DEBUG_VIZ_DIFFUSE_COLOR				5 
 #define PATH_TRACER_DEBUG_VIZ_SPECULAR_COLOR			6 
-#define PATH_TRACER_DEBUG_VIZ_OPACITY					7 
+#define PATH_TRACER_DEBUG_VIZ_TRANSPARENCY_COLOR		7
 #define PATH_TRACER_DEBUG_VIZ_METALLIC					8 
 #define PATH_TRACER_DEBUG_VIZ_SPECULAR					9 
 #define PATH_TRACER_DEBUG_VIZ_ROUGHNESS					10 
@@ -113,3 +113,8 @@
 #define PATH_TRACER_DEBUG_VIZ_FUZZ_COLOR                31
 #define PATH_TRACER_DEBUG_VIZ_FUZZ_ROUGHNESS            32
 #define PATH_TRACER_DEBUG_VIZ_FUZZ_AMOUNT               33
+#define PATH_TRACER_DEBUG_VIZ_BSDF_OPACITY              34
+#define PATH_TRACER_DEBUG_VIZ_SUBSTRATE_WEIGHT_V		35
+#define PATH_TRACER_DEBUG_VIZ_SUBSTRATE_TRANSMITTANCE_N 36
+#define PATH_TRACER_DEBUG_VIZ_SUBSTRATE_COVERAGE_ABOVE_N 37
+#define PATH_TRACER_DEBUG_VIZ_SUBSTRATE_F90 			38
diff --git a/Engine/Source/Runtime/Renderer/Private/PathTracing.cpp b/Engine/Source/Runtime/Renderer/Private/PathTracing.cpp
--- a/Engine/Source/Runtime/Renderer/Private/PathTracing.cpp
+++ b/Engine/Source/Runtime/Renderer/Private/PathTracing.cpp
@@ -68,7 +68,7 @@ struct FPathTracingDebugModeInfo
 		REGISTER_VIZ_TYPE(TEXT("BaseColor") 	    	, PATH_TRACER_DEBUG_VIZ_BASE_COLOR				);
 		REGISTER_VIZ_TYPE(TEXT("DiffuseColor") 	   		, PATH_TRACER_DEBUG_VIZ_DIFFUSE_COLOR			);
 		REGISTER_VIZ_TYPE(TEXT("SpecularColor") 	   	, PATH_TRACER_DEBUG_VIZ_SPECULAR_COLOR			);
-		REGISTER_VIZ_TYPE(TEXT("Opacity") 	     		, PATH_TRACER_DEBUG_VIZ_OPACITY					);
+		REGISTER_VIZ_TYPE(TEXT("TransparencyColor")		, PATH_TRACER_DEBUG_VIZ_TRANSPARENCY_COLOR		);
 		REGISTER_VIZ_TYPE(TEXT("Metallic") 	     		, PATH_TRACER_DEBUG_VIZ_METALLIC				);
 		REGISTER_VIZ_TYPE(TEXT("Specular") 	     		, PATH_TRACER_DEBUG_VIZ_SPECULAR				);
 		REGISTER_VIZ_TYPE(TEXT("Roughness") 	     	, PATH_TRACER_DEBUG_VIZ_ROUGHNESS				);
@@ -95,6 +95,11 @@ struct FPathTracingDebugModeInfo
 		REGISTER_VIZ_TYPE(TEXT("FuzzColor")   			, PATH_TRACER_DEBUG_VIZ_FUZZ_COLOR				);
 		REGISTER_VIZ_TYPE(TEXT("FuzzRoughness") 		, PATH_TRACER_DEBUG_VIZ_FUZZ_ROUGHNESS			);
 		REGISTER_VIZ_TYPE(TEXT("FuzzAmount")   			, PATH_TRACER_DEBUG_VIZ_FUZZ_AMOUNT				);
+		REGISTER_VIZ_TYPE(TEXT("BSDFOpacity")   		, PATH_TRACER_DEBUG_VIZ_BSDF_OPACITY			);
+		REGISTER_VIZ_TYPE(TEXT("SubstrateWeightV") 		, PATH_TRACER_DEBUG_VIZ_SUBSTRATE_WEIGHT_V		);
+		REGISTER_VIZ_TYPE(TEXT("SubstrateTransmittanceN"), PATH_TRACER_DEBUG_VIZ_SUBSTRATE_TRANSMITTANCE_N);
+		REGISTER_VIZ_TYPE(TEXT("SubstrateCoverageAboveN"), PATH_TRACER_DEBUG_VIZ_SUBSTRATE_COVERAGE_ABOVE_N);
+		REGISTER_VIZ_TYPE(TEXT("SubstrateF90")           , PATH_TRACER_DEBUG_VIZ_SUBSTRATE_F90			);
 
 #undef REGISTER_VIZ_TYPE
 	}
