Add pixel center to ShaderPrint zoom to ease visualizing current pixel.

diff --git a/Engine/Shaders/Private/ShaderPrintDraw.usf b/Engine/Shaders/Private/ShaderPrintDraw.usf
--- a/Engine/Shaders/Private/ShaderPrintDraw.usf
+++ b/Engine/Shaders/Private/ShaderPrintDraw.usf
@@ -416,6 +416,13 @@ bool IsValid(int2 InCoord)
 	return all(InCoord >= 0) && all(InCoord < Resolution);
 }
 
+bool IsOnBorder(int2 In, int2 MinB, int2 MaxB)
+{
+	return all(In <= MaxB) && all(In >= MinB) && 
+		(In.x == MinB.x || In.x == MaxB.x ||
+		 In.y == MinB.y || In.y == MaxB.y);
+}
+
 [numthreads(8, 8, 1)]
 void DrawZoomCS(int2 DispatchThreadId : SV_DispatchThreadID)
 {	
@@ -437,6 +444,10 @@ void DrawZoomCS(int2 DispatchThreadId : SV_DispatchThreadID)
 	case 3: OutOffset = int2(Offset.x, Resolution.y - Offset.y - OutMaxPixel.y); break;
 	};
 
+	const int2 CenterCoord = OutOffset + PixelExtent * ZoomFactor;
+	const int2 MinCenter = CenterCoord;
+	const int2 MaxCenter = CenterCoord + 1 * ZoomFactor;
+
 	if (IsValid(ShaderPrintData.CursorCoord) && IsValid(InCoord))
 	{
 		const float4 Color = InTexture.Load(uint3(InCoord, 0));
@@ -444,10 +455,15 @@ void DrawZoomCS(int2 DispatchThreadId : SV_DispatchThreadID)
 		for (uint X=0; X<ZoomFactor; ++X)
 		{
 			const int2 OutCoord = DispatchThreadId * ZoomFactor + int2(X, Y) + OutOffset;
-			const bool bIsBorder = any(OutCoord == OutOffset) || any(OutCoord == OutOffset + OutMaxPixel - 1);
+
+			// External border
+			const bool bIsBorder0 = any(OutCoord == OutOffset) || any(OutCoord == OutOffset + OutMaxPixel - 1);
+			// Pixel center border
+			const bool bIsBorder1 = IsOnBorder(OutCoord, MinCenter, MaxCenter);
+
 			if (IsValid(OutCoord))
 			{
-				OutTexture[OutCoord] = bIsBorder ? float4(1,1,0,1) : Color;
+				OutTexture[OutCoord] = bIsBorder0 || bIsBorder1 ? float4(1,1,0,1) : Color;
 			}
 
 		}
