Fix VRS support with Substrate.

Pixel offset were incorrectly computed when the top-left pixel of a quad was not shaded. Now all offset are always computed based on the correct "quad's base-pixel".

diff --git a/Engine/Shaders/Private/BasePassPixelShader.usf b/Engine/Shaders/Private/BasePassPixelShader.usf
--- a/Engine/Shaders/Private/BasePassPixelShader.usf
+++ b/Engine/Shaders/Private/BasePassPixelShader.usf
@@ -889,7 +889,9 @@ void FPixelShaderInOut_MainPS(
 	in FPixelShaderIn In,
 	inout FPixelShaderOut Out,
 	const uint EyeIndex,
-	uint QuadPixelWriteMask=1)
+	uint QuadPixelWriteMask=1,
+	uint2 QuadPixelPos=0,
+	bool bQuadShadingEnabled=false)
 {
 	// Velocity
 	float4 OutVelocity = 0;
@@ -1984,9 +1986,14 @@ void FPixelShaderInOut_MainPS(
 
 #else // SUBSTRATE_GBUFFER_FORMAT==0
 
+		if (!bQuadShadingEnabled)
+		{
+			QuadPixelPos = PixelPos;
+		}
+
 		// Generate the Substrate material data to write out
 		FSubstrateAddressing SubstrateAddressing = GetSubstratePixelDataByteOffset(PixelPos, uint2(ResolvedView.BufferSizeAndInvSize.xy), SubstrateStruct.MaxBytesPerPixel);
-		FRWSubstrateMaterialContainer RWSubstrateMaterialContainer = InitialiseRWSubstrateMaterialContainer(SubstrateStruct.MaterialTextureArrayUAVWithoutRTs, QuadPixelWriteMask);
+		FRWSubstrateMaterialContainer RWSubstrateMaterialContainer = InitialiseRWSubstrateMaterialContainer(SubstrateStruct.MaterialTextureArrayUAVWithoutRTs, QuadPixelWriteMask, QuadPixelPos);
 		PackSubstrateOut(
 			RWSubstrateMaterialContainer,
 			SubstrateStruct.MaterialTextureArrayUAVWithoutRTs,
@@ -2024,7 +2031,7 @@ void FPixelShaderInOut_MainPS(
 		if (QuadPixelWriteMask>0)
 		{
 			const float3 OpaqueRefractionData = SubstratePackOpaqueRoughRefractionData(OpaqueRoughRefractionData);
-			WriteDataToBuffer(SubstrateStruct.OpaqueRoughRefractionTextureUAV, OpaqueRefractionData, PixelPos, QuadPixelWriteMask);
+			WriteDataToBuffer(SubstrateStruct.OpaqueRoughRefractionTextureUAV, OpaqueRefractionData, PixelPos, QuadPixelWriteMask, QuadPixelPos);
 		}
 	#endif
 
@@ -2034,12 +2041,12 @@ void FPixelShaderInOut_MainPS(
 		{
 			// Subsurface header
 			{
-				SubstrateStoreSubsurfaceHeader(SubstrateStruct.MaterialTextureArrayUAVWithoutRTs, SubstrateStruct.FirstSliceStoringSubstrateSSSDataWithoutMRT, PixelPos, SSSData.Header.Bytes, QuadPixelWriteMask);
+				SubstrateStoreSubsurfaceHeader(SubstrateStruct.MaterialTextureArrayUAVWithoutRTs, SubstrateStruct.FirstSliceStoringSubstrateSSSDataWithoutMRT, PixelPos, SSSData.Header.Bytes, QuadPixelWriteMask, QuadPixelPos);
 			}
 			// Subsurface data
 			if (SubstrateSubSurfaceHeaderHasExtras(SSSData.Header))
 			{
-				SubstrateStoreSubsurfaceExtras(SubstrateStruct.MaterialTextureArrayUAVWithoutRTs, SubstrateStruct.FirstSliceStoringSubstrateSSSDataWithoutMRT, PixelPos, SSSData.Extras.Bytes, QuadPixelWriteMask);
+				SubstrateStoreSubsurfaceExtras(SubstrateStruct.MaterialTextureArrayUAVWithoutRTs, SubstrateStruct.FirstSliceStoringSubstrateSSSDataWithoutMRT, PixelPos, SSSData.Extras.Bytes, QuadPixelWriteMask, QuadPixelPos);
 			}
 		}
 
diff --git a/Engine/Shaders/Private/ComputeShaderOutputCommon.ush b/Engine/Shaders/Private/ComputeShaderOutputCommon.ush
--- a/Engine/Shaders/Private/ComputeShaderOutputCommon.ush
+++ b/Engine/Shaders/Private/ComputeShaderOutputCommon.ush
@@ -43,7 +43,7 @@ uint GetDataByteOffset()	{ return GetRootConstant3(); }
 	#endif
 #endif
 
-FPixelShaderOut ShadePixel(const float2 SVPositionXY, uint QuadIndex, uint QuadPixelWriteMask)
+FPixelShaderOut ShadePixel(const float2 SVPositionXY, uint QuadIndex, uint QuadPixelWriteMask, uint2 QuadPixelPos)
 {
 	// Note: Driven by bHighPrecisionGBuffers in GBufferInfo.cpp
 	const bool bHighPrecision = GetHighPrecision() != 0u;
@@ -85,7 +85,7 @@ FPixelShaderOut ShadePixel(const float2 SVPositionXY, uint QuadIndex, uint QuadP
 
 #if PIXELSHADEROUTPUT_BASEPASS
 	FBasePassInterpolantsVSToPS BasePassInterpolants = (FBasePassInterpolantsVSToPS)0;
-	FPixelShaderInOut_MainPS(Interpolants, BasePassInterpolants, PixelShaderIn, PixelShaderOut, EyeIndex, QuadPixelWriteMask);
+	FPixelShaderInOut_MainPS(Interpolants, BasePassInterpolants, PixelShaderIn, PixelShaderOut, EyeIndex, QuadPixelWriteMask, QuadPixelPos, true/*QuadShadingEnabled*/);
 #endif
 
 #if !SUBSTRATE_ENABLED || (SUBSTRATE_ENABLED && SUBSTRATE_GBUFFER_FORMAT == 0)
@@ -186,7 +186,7 @@ void ExportPixel(const uint2 PixelPos, FPixelShaderOut ShadedPixel)
 void ProcessPixel(uint ShadingBin, const uint2 PixelPos, const float2 SVPositionXY, uint QuadIndex, uint DispatchIndex, uint PixelWriteMask, uint HelperLaneCount)
 {
 	// All lanes shade (regardless of export - so ddx/ddy are valid)
-	FPixelShaderOut ShadedPixel = ShadePixel(SVPositionXY, QuadIndex, PixelWriteMask);
+	FPixelShaderOut ShadedPixel = ShadePixel(SVPositionXY, QuadIndex, PixelWriteMask, PixelPos /*QuadPixelPos*/);
 
 #if VIS_HELPER_LANES
 	ShadedPixel.MRT[3].rgb = ColorMapTurbo(float(HelperLaneCount) / 4.0f);
diff --git a/Engine/Shaders/Private/DebugProbes.usf b/Engine/Shaders/Private/DebugProbes.usf
--- a/Engine/Shaders/Private/DebugProbes.usf
+++ b/Engine/Shaders/Private/DebugProbes.usf
@@ -164,7 +164,7 @@ void WriteOutMaterial(
 	const FSubstrateIntegrationSettings Settings = InitSubstrateIntegrationSettings(false /*bForceFullyRough*/, bRoughDiffuse, 0/*PeelLayersAboveDepth*/, false/*bRoughnessTracking*/);
 
 	FSubstrateAddressing SubstrateAddressing = GetSubstratePixelDataByteOffset(PixelCoord, uint2(ResolvedView.BufferSizeAndInvSize.xy), MaxBytesPerPixel);
-	FRWSubstrateMaterialContainer RWSubstrateMaterialContainer = InitialiseRWSubstrateMaterialContainer(MaterialTextureArrayUAV, true /*bAllowWrite*/);
+	FRWSubstrateMaterialContainer RWSubstrateMaterialContainer = InitialiseRWSubstrateMaterialContainer(MaterialTextureArrayUAV, 1/*QuadPixelWriteMask*/, PixelCoord);
 	PackSubstrateOut(
 		RWSubstrateMaterialContainer,
 		MaterialTextureArrayUAV,
diff --git a/Engine/Shaders/Private/Substrate/Substrate.ush b/Engine/Shaders/Private/Substrate/Substrate.ush
--- a/Engine/Shaders/Private/Substrate/Substrate.ush
+++ b/Engine/Shaders/Private/Substrate/Substrate.ush
@@ -2015,15 +2015,16 @@ FSubstrateAddressing GetSubstratePixelDataByteOffset(uint2 PixelPos, uint2 ViewB
 
 struct FRWSubstrateMaterialContainerStruct
 {
-	uint MaterialRenderTargets[SUBSTRATE_BASE_PASS_MRT_OUTPUT_COUNT];
-	uint QuadPixelWriteMask;
+	uint  MaterialRenderTargets[SUBSTRATE_BASE_PASS_MRT_OUTPUT_COUNT];
+	uint2 QuadPixelPos;
+	uint  QuadPixelWriteMask;
 //	RWTexture2DArray<uint> ExtraMaterialDataUAV; // If this is used, we cannot run shaders on debug...
 };
 
 #if SUBSTRATE_MATERIALCONTAINER_IS_VIEWRESOURCE
 
 #define FRWSubstrateMaterialContainer	FRWSubstrateMaterialContainerStruct
-FRWSubstrateMaterialContainer InitialiseRWSubstrateMaterialContainer(RWTexture2DArray<uint> ExtraMaterialDataUAV, uint QuadPixelWriteMask)
+FRWSubstrateMaterialContainer InitialiseRWSubstrateMaterialContainer(RWTexture2DArray<uint> ExtraMaterialDataUAV, uint QuadPixelWriteMask, uint2 QuadPixelPos)
 {
 	FRWSubstrateMaterialContainer Out;
 	UNROLL
@@ -2031,6 +2032,7 @@ FRWSubstrateMaterialContainer InitialiseRWSubstrateMaterialContainer(RWTexture2D
 	{
 		Out.MaterialRenderTargets[i] = 0;
 	}
+	Out.QuadPixelPos = QuadPixelPos;
 	Out.QuadPixelWriteMask = QuadPixelWriteMask;
 //	Out.ExtraMaterialDataUAV = ExtraMaterialDataUAV;
 	return Out;
@@ -2051,7 +2053,7 @@ void SubstrateStoreUint1(inout FRWSubstrateMaterialContainer SubstrateBuffer, RW
 	}
 	else if (SubstrateBuffer.QuadPixelWriteMask>0) // Only write data to UAV if pixel is allowed to write (i.e., it is not an helper lane/pixel)
 	{
-		WriteDataToBuffer(ExtraMaterialDataUAV, Data, uint3(SubstrateAddressing.PixelCoords, SubstrateAddressing.CurrentIndex - SUBSTRATE_BASE_PASS_MRT_OUTPUT_COUNT), SubstrateBuffer.QuadPixelWriteMask);
+		WriteDataToBuffer(ExtraMaterialDataUAV, Data, uint3(SubstrateAddressing.PixelCoords, SubstrateAddressing.CurrentIndex - SUBSTRATE_BASE_PASS_MRT_OUTPUT_COUNT), SubstrateBuffer.QuadPixelWriteMask, SubstrateBuffer.QuadPixelPos);
 	}
 	SubstrateAddressing.CurrentIndex++;
   #endif // PLATFORM_SUPPORTS_SUBSTRATE_UINT1
diff --git a/Engine/Shaders/Private/Substrate/SubstrateExportCommon.ush b/Engine/Shaders/Private/Substrate/SubstrateExportCommon.ush
--- a/Engine/Shaders/Private/Substrate/SubstrateExportCommon.ush
+++ b/Engine/Shaders/Private/Substrate/SubstrateExportCommon.ush
@@ -9,7 +9,10 @@
 #define SUBSTRATE_VRS_WRITE 0
 #endif
 
-void WriteDataToBuffer(RWTexture2DArray<uint> OutBuffer, uint InData, uint3 InCoord, uint InQuadPixelWriteMask)
+// * InQuadPixelWriteMask: mask describing which pixels needs to be written
+// * InQuadPixelPos      : pixel coord of top left quad pixel. All mask's offset are expressed relative to this quad pixel
+
+void WriteDataToBuffer(RWTexture2DArray<uint> OutBuffer, uint InData, uint3 InCoord, uint InQuadPixelWriteMask, uint2 InQuadPixelPos)
 {
 	if (InQuadPixelWriteMask > 0u)
 	{
@@ -18,29 +21,29 @@ void WriteDataToBuffer(RWTexture2DArray<uint> OutBuffer, uint InData, uint3 InCo
 		if (InQuadPixelWriteMask & 1u)
 	#endif
 		{
-			OutBuffer[InCoord] = InData;
+			OutBuffer[uint3(InQuadPixelPos, InCoord.z)] = InData;
 		}
 	#if SUBSTRATE_VRS_WRITE
 		BRANCH
 		if (InQuadPixelWriteMask & 2u) // Copy H
 		{
-			OutBuffer[uint3(InCoord.xy + uint2(1, 0), InCoord.z)] = InData;
+			OutBuffer[uint3(InQuadPixelPos.xy + uint2(1, 0), InCoord.z)] = InData;
 		}
 		BRANCH
 		if (InQuadPixelWriteMask & 4u) // Copy V
 		{
-			OutBuffer[uint3(InCoord.xy + uint2(0, 1), InCoord.z)] = InData;
+			OutBuffer[uint3(InQuadPixelPos.xy + uint2(0, 1), InCoord.z)] = InData;
 		}
 		BRANCH
 		if (InQuadPixelWriteMask & 8u) // Copy D
 		{
-			OutBuffer[uint3(InCoord.xy + uint2(1, 1), InCoord.z)] = InData;
+			OutBuffer[uint3(InQuadPixelPos.xy + uint2(1, 1), InCoord.z)] = InData;
 		}
 	#endif
 	}
 }
 
-void WriteDataToBuffer(RWTexture2D<float3> OutBuffer, float3 InData, uint2 InCoord, uint InQuadPixelWriteMask)
+void WriteDataToBuffer(RWTexture2D<float3> OutBuffer, float3 InData, uint2 InCoord, uint InQuadPixelWriteMask, uint2 InQuadPixelPos)
 {
 	if (InQuadPixelWriteMask > 0u)
 	{
@@ -49,23 +52,23 @@ void WriteDataToBuffer(RWTexture2D<float3> OutBuffer, float3 InData, uint2 InCoo
 		if (InQuadPixelWriteMask & 1u)
 	#endif
 		{
-			OutBuffer[InCoord] = InData;
+			OutBuffer[InQuadPixelPos] = InData;
 		}
 	#if SUBSTRATE_VRS_WRITE
 		BRANCH
 		if (InQuadPixelWriteMask & 2u) /* Copy H */
 		{
-			OutBuffer[InCoord + uint2(1, 0)] = InData;
+			OutBuffer[InQuadPixelPos + uint2(1, 0)] = InData;
 		}
 		BRANCH
 		if (InQuadPixelWriteMask & 4u) /* Copy V */
 		{
-			OutBuffer[InCoord + uint2(0, 1)] = InData;
+			OutBuffer[InQuadPixelPos + uint2(0, 1)] = InData;
 		}
 		BRANCH
 		if (InQuadPixelWriteMask & 8u) /* Copy D */
 		{
-			OutBuffer[InCoord + uint2(1, 1)] = InData;
+			OutBuffer[InQuadPixelPos + uint2(1, 1)] = InData;
 		}
 	#endif
 	}
diff --git a/Engine/Shaders/Private/Substrate/SubstrateSubsurface.ush b/Engine/Shaders/Private/Substrate/SubstrateSubsurface.ush
--- a/Engine/Shaders/Private/Substrate/SubstrateSubsurface.ush
+++ b/Engine/Shaders/Private/Substrate/SubstrateSubsurface.ush
@@ -199,14 +199,24 @@ FSubstrateSubsurfaceData SubstrateLoadSubsurfaceData(Texture2DArray<uint> Substr
 	return SSSData;
 }
 
-void SubstrateStoreSubsurfaceHeader(RWTexture2DArray<uint> SubstrateBuffer, uint FirstSliceStoringSubstrateSSSData, uint2 PixelPos, uint HeaderBytes, uint QuadPixelWriteMask=1)
+void SubstrateStoreSubsurfaceHeader(RWTexture2DArray<uint> SubstrateBuffer, uint FirstSliceStoringSubstrateSSSData, uint2 PixelPos, uint HeaderBytes, uint QuadPixelWriteMask, uint2 QuadPixelPos)
 {
-	WriteDataToBuffer(SubstrateBuffer, HeaderBytes, uint3(PixelPos, FirstSliceStoringSubstrateSSSData + 0), QuadPixelWriteMask);
+	WriteDataToBuffer(SubstrateBuffer, HeaderBytes, uint3(PixelPos, FirstSliceStoringSubstrateSSSData + 0), QuadPixelWriteMask, QuadPixelPos);
 }
 
-void SubstrateStoreSubsurfaceExtras(RWTexture2DArray<uint> SubstrateBuffer, uint FirstSliceStoringSubstrateSSSData, uint2 PixelPos, uint ExtraBytes, uint QuadPixelWriteMask=1)
+void SubstrateStoreSubsurfaceHeader(RWTexture2DArray<uint> SubstrateBuffer, uint FirstSliceStoringSubstrateSSSData, uint2 PixelPos, uint HeaderBytes)
 {
-	WriteDataToBuffer(SubstrateBuffer, ExtraBytes, uint3(PixelPos, FirstSliceStoringSubstrateSSSData + 1), QuadPixelWriteMask);
+	SubstrateStoreSubsurfaceHeader(SubstrateBuffer, FirstSliceStoringSubstrateSSSData, PixelPos, HeaderBytes, 1 /*QuadPixelWriteMask*/, PixelPos);
+}
+
+void SubstrateStoreSubsurfaceExtras(RWTexture2DArray<uint> SubstrateBuffer, uint FirstSliceStoringSubstrateSSSData, uint2 PixelPos, uint ExtraBytes, uint QuadPixelWriteMask, uint2 QuadPixelPos)
+{
+	WriteDataToBuffer(SubstrateBuffer, ExtraBytes, uint3(PixelPos, FirstSliceStoringSubstrateSSSData + 1), QuadPixelWriteMask, QuadPixelPos);
+}
+
+void SubstrateStoreSubsurfaceExtras(RWTexture2DArray<uint> SubstrateBuffer, uint FirstSliceStoringSubstrateSSSData, uint2 PixelPos, uint ExtraBytes)
+{
+	SubstrateStoreSubsurfaceExtras(SubstrateBuffer, FirstSliceStoringSubstrateSSSData, PixelPos, ExtraBytes, 1/*QuadPixelWriteMask*/, PixelPos);
 }
 
 ///////////////////////////////////////////////////////////////////////////////
