Path Tracer: Refactor slab selection logic to use a single loop

This keeps the computation of the slab albedo and the population of the payload structure closer together, making it easier to keep them in sync.

Account for two sided foliage and thin slabs better, fixing some inconsistencies.

Fix coverage of unlit and legacy glass materials.

#rnx

diff --git a/Engine/Shaders/Private/ParticipatingMediaCommon.ush b/Engine/Shaders/Private/ParticipatingMediaCommon.ush
--- a/Engine/Shaders/Private/ParticipatingMediaCommon.ush
+++ b/Engine/Shaders/Private/ParticipatingMediaCommon.ush
@@ -184,6 +184,12 @@ float3 ExtinctionToTransmittance(in float3 Extinction, in float ThicknessMeters)
 	return exp(-Extinction * ThicknessMeters);
 }
 
+float3 MeanFreePathToTransmittance(in float3 MeanFreePath, in float ThicknessMeters)
+{
+	return exp(-ThicknessMeters / max(MeanFreePath, PARTICIPATING_MEDIA_MIN_MFP_METER));
+}
+
+
 //---------------------------------------------
 // Lighting evaluation function for an Isotropic participating media a slab of 1 meter
 //---------------------------------------------
diff --git a/Engine/Shaders/Private/PathTracing/PathTracingMaterialHitShader.usf b/Engine/Shaders/Private/PathTracing/PathTracingMaterialHitShader.usf
--- a/Engine/Shaders/Private/PathTracing/PathTracingMaterialHitShader.usf
+++ b/Engine/Shaders/Private/PathTracing/PathTracingMaterialHitShader.usf
@@ -28,6 +28,7 @@
 #include "/Engine/Private/RayTracing/RayTracingCommon.ush"
 #include "/Engine/Private/RayTracing/RayTracingHitGroupCommon.ush"
 #include "/Engine/Private/PathTracing/PathTracingShaderUtils.ush"
+#include "/Engine/Private/PathTracing/Utilities/PathTracingRIS.ush"
 #include "/Engine/Generated/Material.ush"   
 #include "/Engine/Generated/VertexFactory.ush"
 
@@ -370,7 +371,7 @@ RAY_TRACING_ENTRY_CLOSEST_HIT(PathTracingMaterialCHS,
 	Payload.TransparencyColor = 0;
 #else
 	Payload.BSDFOpacity = UnlitSurfaceCoverage;
-	Payload.TransparencyColor = UnlitSurfaceTransmittancePreCoverage;
+	Payload.TransparencyColor = lerp(1.0, UnlitSurfaceTransmittancePreCoverage, UnlitSurfaceCoverage);
 #endif
 
 #else // SUBSTRATE_OPTIMIZED_UNLIT
@@ -411,331 +412,280 @@ RAY_TRACING_ENTRY_CLOSEST_HIT(PathTracingMaterialCHS,
 			Payload.Radiance += BSDF.LuminanceWeightV * BSDF_GETEMISSIVE(BSDF);
 		}
 
-
 #if SUBSTRATE_CLAMPED_CLOSURE_COUNT > 1
 		// The material might be using more than one slab -- insert code to stochastically pick one 
-		float SlabPdf = 1.0;
-		if (BSDFCount > 1)
+		FRISContext SlabSampler = InitRISContext(PackedPayload.GetStochasticSlabRand());
+		for (BSDFIdx = 0; BSDFIdx < BSDFCount; ++BSDFIdx)
+#else
+		// Just one slab, don't bother inserting logic for lobe selection, we just have one
+		BSDFIdx = 0;
+#endif
 		{
-			float SlabCDF[SUBSTRATE_CLAMPED_CLOSURE_COUNT];
-			float SlabCDFSum = 0.0;
-			for (BSDFIdx = 0; BSDFIdx < BSDFCount; ++BSDFIdx)
+			switch (BSDF_GETTYPE(BSDF))
 			{
-				float Pdf = 0.0; // The goal is to predict the weight of each BSDF
-				switch (BSDF_GETTYPE(BSDF))
+				case SUBSTRATE_BSDF_TYPE_SLAB:
 				{
-					case SUBSTRATE_BSDF_TYPE_SLAB:
-					{
-						const float3 WeightV = BSDF.LuminanceWeightV;
-						const float3 TransmittanceN = BSDF.TransmittanceAboveAlongN;
-						const float CoverageAboveAlongN = BSDF.CoverageAboveAlongN;
-						const float3 MaxLobeWeight = WeightV * lerp(1.0f, TransmittanceN, CoverageAboveAlongN); // largest value LobeWeight() could return
+					const float3 WeightV = BSDF.LuminanceWeightV;
+					const float3 TransmittanceN = BSDF.TransmittanceAboveAlongN;
+					const float CoverageAboveAlongN = BSDF.CoverageAboveAlongN;
+					const float3 MaxLobeWeight = WeightV * lerp(1.0f, TransmittanceN, CoverageAboveAlongN); // largest value LobeWeight() could return
 
-						// Is the slab even visible at all?
-						if (any(MaxLobeWeight > 0.0))
+					// Is the slab even visible at all?
+					if (any(MaxLobeWeight > 0.0))
+					{
+						const float3 N = SubstratePixelHeader.SharedLocalBases.Normals[BSDF_GETSHAREDLOCALBASISID(BSDF)];
+						const float NoV = saturate(dot(N, V_World));
+
+						float3 DiffuseColor = SLAB_DIFFUSEALBEDO(BSDF);
+						float3 F0 = SLAB_F0(BSDF);
+						float3 F90 = SLAB_F90(BSDF) * F0RGBToMicroOcclusion(F0);
+
+						float3 RoughnessData = 0;
+						RoughnessData.xy = AdjustMaterialRoughness(SLAB_ROUGHNESS(BSDF), PackedPayload.GetPathRoughness());
+						// SUBSTRATE_TODO: Support Fresnel82?
+						float3 Spec0E = ComputeGGXSpecEnergyTermsRGB(RoughnessData.x, NoV, F0, F90).E;
+						float3 Spec1E = Spec0E;
+						bool bUseClearCoat = false;
+						if (BSDF_GETHASHAZINESS(BSDF))
 						{
-							const float3 N = SubstratePixelHeader.SharedLocalBases.Normals[BSDF_GETSHAREDLOCALBASISID(BSDF)];
-							const float NoV = saturate(dot(N, V_World));
-
-							float3 DiffuseColor = SLAB_DIFFUSEALBEDO(BSDF);
-							float3 F0 = SLAB_F0(BSDF);
-							float3 F90 = SLAB_F90(BSDF) * F0RGBToMicroOcclusion(F0);
-
-							float Roughness0 = MakeRoughnessSafe(SLAB_ROUGHNESS(BSDF), PackedPayload.GetPathRoughness());
-							// SUBSTRATE_TODO: Support Fresnel82?
-							float3 Spec0E = ComputeGGXSpecEnergyTermsRGB(Roughness0, NoV, F0, F90).E;
-							float3 Spec1E = Spec0E;
-							float Roughness1 = Roughness0;
-							float SpecBlend = 0.0;
-							bool bUseClearCoat = false;
-							if (BSDF_GETHASHAZINESS(BSDF))
-							{
-								const FHaziness Haziness = UnpackHaziness(SLAB_HAZINESS(BSDF));
-								SpecBlend = Haziness.Weight;
-								Roughness1 = MakeRoughnessSafe(Haziness.Roughness, PackedPayload.GetPathRoughness());
-								if (Haziness.bSimpleClearCoat)
-								{
-									Spec1E = ComputeGGXSpecEnergyTermsRGB(Roughness1, NoV, CLEAR_COAT_F0).E;
-									bUseClearCoat = true;
-								}
-								else
-								{
-									Spec1E = ComputeGGXSpecEnergyTermsRGB(Roughness1, NoV, F0, F90).E;
-								}
-							}
-							if (BSDF_GETSSSTYPE(BSDF) == SSS_TYPE_SIMPLEVOLUME)
-							{
-								DiffuseColor = GetSimpleVolumeDiffuseColor(DiffuseColor, SLAB_SSSMFP(BSDF));
-							}
-
-							float FuzzAmount = 0.0;
-							float FuzzRoughness = 1.0;
-							float3 FuzzColor = 0.0;
-							if (BSDF_GETHASFUZZ(BSDF))
-							{
-								FuzzAmount = SLAB_FUZZ_AMOUNT(BSDF);
-								FuzzColor = SLAB_FUZZ_COLOR(BSDF);
-								FuzzRoughness = SLAB_FUZZ_ROUGHNESS(BSDF);
-								FuzzRoughness = MakeRoughnessSafe(FuzzRoughness, SUBSTRATE_MIN_FUZZ_ROUGHNESS);
-							}
-							
-							const float FuzzE = ComputeSheenEnergyTerms(FuzzRoughness, NoV).E;
-							const float FuzzAttenuation = 1.0 - FuzzAmount * FuzzE;
-							float DiffuseWeight = FuzzAttenuation;
-							float3 Spec0Albedo = FuzzAttenuation * Spec0E;
-							float3 Spec1Albedo = FuzzAttenuation * Spec1E;
-							if (bUseClearCoat)
+							const FHaziness Haziness = UnpackHaziness(SLAB_HAZINESS(BSDF));
+							RoughnessData.z = Haziness.Weight;
+							RoughnessData.y = AdjustMaterialRoughness(Haziness.Roughness, PackedPayload.GetPathRoughness());
+							if (Haziness.bSimpleClearCoat)
 							{
-								// clearcoat slab
-								float CoatAttenuation = 1.0 - SpecBlend * Spec1E.x;
-								DiffuseWeight *= CoatAttenuation * (1.0 - Luminance(Spec0E));
-								Spec0Albedo *= CoatAttenuation;
-								Spec1Albedo *= SpecBlend;
+								Spec1E = ComputeGGXSpecEnergyTermsRGB(RoughnessData.y, NoV, CLEAR_COAT_F0).E;
+								bUseClearCoat = true;
 							}
 							else
 							{
-								// dual-specular slab
-								DiffuseWeight *= 1.0 - Luminance(lerp(Spec0E, Spec1E, SpecBlend));
-								Spec0Albedo *= 1.0 - SpecBlend;
-								Spec1Albedo *=       SpecBlend;
+								Spec1E = ComputeGGXSpecEnergyTermsRGB(RoughnessData.y, NoV, F0, F90).E;
 							}
-							// SUBSTRATE_TODO: Need to account for glass case here for lobe selection ....
-							const float3 FuzzAlbedo = FuzzAmount * FuzzE * FuzzColor;
-							const float3 SlabAlbedo = MaxLobeWeight * (DiffuseWeight * DiffuseColor + Spec0Albedo + Spec1Albedo + FuzzAlbedo);
-							Pdf = LobeColorToWeight(SlabAlbedo);
 						}
-						break;
-					}
-					default:
-					{
-						// In theory none of the other slab types can be layered
-						break;
-					}
-				}
-				SlabCDFSum += Pdf;
-				SlabCDF[BSDFIdx] = SlabCDFSum;
-			}
-			if (SlabCDFSum > 0.0)
-			{
-				// linear search
-				float PreviousCdfValue = 0;
-				float RandSample = SlabCDFSum * PackedPayload.GetStochasticSlabRand();
-				for (BSDFIdx = 0; BSDFIdx < BSDFCount - 1; ++BSDFIdx)
-				{
-					if (RandSample < SlabCDF[BSDFIdx])
-					{
-						break;
-					}
-					PreviousCdfValue = SlabCDF[BSDFIdx];
-				}
-				SlabPdf = (SlabCDF[BSDFIdx] - PreviousCdfValue) / SlabCDFSum;
-			}
-			else
-			{
-				BSDFIdx = 0;
-				SlabPdf = 0.0;
-			}
-		}
-		if (SlabPdf > 0)
-		{
-#elif SUBSTRATE_CLAMPED_CLOSURE_COUNT == 1
-		if (BSDFCount > 0)
-		{
-			// There is only one slab active
-			BSDFIdx = 0;
-#else
-// shouldn't be possible to reach this point
-#error Unexpected BSDF Count
-#endif // SUBSTRATE_CLAMPED_CLOSURE_COUNT 
-
-
-#if SUBSTRATE_CLAMPED_CLOSURE_COUNT > 0
-			Payload.WorldNormal = SubstratePixelHeader.SharedLocalBases.Normals[BSDF_GETSHAREDLOCALBASISID(BSDF)];
-			if (SubstrateGetSharedLocalBasisType(SubstratePixelHeader.SharedLocalBases.Types, BSDFIdx) == SUBSTRATE_BASIS_TYPE_TANGENT)
-			{
-				Payload.WorldTangent = SubstratePixelHeader.SharedLocalBases.Tangents[BSDF_GETSHAREDLOCALBASISID(BSDF)];
-			}
-
-			switch (BSDF_GETTYPE(BSDF))
-			{
-				case SUBSTRATE_BSDF_TYPE_SLAB:
-				{
-					Payload.ShadingModelID    = SHADINGMODELID_SUBSTRATE;
-					Payload.DiffuseColor      = SLAB_DIFFUSEALBEDO(BSDF);
-					Payload.SpecularColor     = SLAB_F0(BSDF);
-					Payload.SpecularEdgeColor = SLAB_F90(BSDF);
-					Payload.RoughnessData.x   = MakeRoughnessSafe(SLAB_ROUGHNESS(BSDF), PackedPayload.GetPathRoughness());
-					Payload.RoughnessData.y = Payload.RoughnessData.x;
-					Payload.RoughnessData.z = 0;
-					if (BSDF_GETHASHAZINESS(BSDF))
-					{
-						const FHaziness Haziness = UnpackHaziness(SLAB_HAZINESS(BSDF));
-						if (Haziness.bSimpleClearCoat)
+						float3 LegacySubsurfaceColor = 0;
+						float LegacyOpacity = 1;
+						if (BSDF_GETSSSTYPE(BSDF) == SSS_TYPE_SIMPLEVOLUME)
 						{
-							Payload.ShadingModelID = SHADINGMODELID_CLEAR_COAT;
+							// TODO: This should really be represented by a different BSDF, but for now we are approximating it with a diffuse lobe
+							DiffuseColor = GetSimpleVolumeDiffuseColor(DiffuseColor, SLAB_SSSMFP(BSDF));
 						}
-						Payload.RoughnessData.y = MakeRoughnessSafe(Haziness.Roughness, PackedPayload.GetPathRoughness());
-						Payload.RoughnessData.z = Haziness.Weight;
-					#if CLEAR_COAT_BOTTOM_NORMAL
-						Payload.BottomNormalOct16bits = SLAB_HAZINESS(BSDF) >> 16;
-					#else
-						Payload.BottomNormalOct16bits = 0;
-					#endif
-					}
-					Payload.Anisotropy = BSDF_GETHASANISOTROPY(BSDF) ? SLAB_ANISOTROPY(BSDF) : 0.0;
-					if (BSDF_GETHASFUZZ(BSDF))
-					{
-						Payload.FuzzAmount    = SLAB_FUZZ_AMOUNT(BSDF);
-						Payload.FuzzColor     = SLAB_FUZZ_COLOR(BSDF);
-						Payload.FuzzRoughness = MakeRoughnessSafe(SLAB_FUZZ_ROUGHNESS(BSDF), SUBSTRATE_MIN_FUZZ_ROUGHNESS);
-					}
-					else
-					{
-						Payload.FuzzAmount    = 0.0;
-						Payload.FuzzColor     = 0.0;
-						Payload.FuzzRoughness = 0.0;
-					}
-					if (BSDF_GETSSSTYPE(BSDF) == SSS_TYPE_SIMPLEVOLUME)
-					{
-						// TODO: this case should be represented by a different BSDF
-						Payload.DiffuseColor = GetSimpleVolumeDiffuseColor(Payload.DiffuseColor, SLAB_SSSMFP(BSDF));
-					}
-					else if (BSDF_GETSSSTYPE(BSDF) == SSS_TYPE_DIFFUSION)
-					{
-						Payload.MeanFreePath = SLAB_SSSMFP(BSDF);
-						Payload.PhaseG       = SLAB_SSSPHASEANISOTROPY(BSDF);
-					}
-					else if (BSDF_GETSSSTYPE(BSDF) == SSS_TYPE_DIFFUSION_PROFILE)
-					{
-						const uint ProfileId = SLAB_SSSPROFILEID(BSDF);
-						const float Scale = SLAB_SSSPROFILERADIUSSCALE(BSDF);
-						Payload.MeanFreePath = DecodeSSSProfileRadius(ProfileId, Payload.DiffuseColor, Scale);
-#if USE_SSS_PROFILE_ANISOTROPY
-						Payload.PhaseG = DecodeSSSProfileScatteringDistribution(ProfileId);
-						Payload.MeanFreePath *= 1 - G; // Compensate the extra translucency by scaling down the MFP
-#else
-						Payload.PhaseG = 0.0;
-#endif
-					}
-					else if (BSDF_GETSSSTYPE(BSDF) == SSS_TYPE_WRAP || BSDF_GETSSSTYPE(BSDF) == SSS_TYPE_TWO_SIDED_WRAP)
-					{
-						const bool bFoliage = BSDF_GETSSSTYPE(BSDF) == SSS_TYPE_TWO_SIDED_WRAP;
-						const float TransmittanceNoL = 1.0f;
-						const float3 SlabDiffuseColor = bFoliage ? Payload.DiffuseColor : float3(1, 1, 1);
-						const FParticipatingMedia PM = SubstrateSlabCreateParticipatingMedia(Payload.DiffuseColor, SLAB_SSSMFP(BSDF));
-						float3 SubsurfaceColor = IsotropicMediumSlabTransmittance(PM, SUBSTRATE_SIMPLEVOLUME_THICKNESS_M, TransmittanceNoL);
-						const float Opacity = 1.f - abs(SLAB_SSSPHASEANISOTROPY(BSDF));
-
-						float3 Front = Payload.DiffuseColor;
-						float3 Back = SubsurfaceColor;
-						float3 Sum = Front + Back;
-						float S = max(Sum.x, max(Sum.y, Sum.z));
-						if (S > 1)
+						else if (BSDF_GETSSSTYPE(BSDF) == SSS_TYPE_WRAP || BSDF_GETSSSTYPE(BSDF) == SSS_TYPE_TWO_SIDED_WRAP)
+						{
+							// For the old MATERIAL_SHADINGMODEL_SUBSURFACE and SHADINGMODELID_TWOSIDED_FOLIAGE models, we previously had logic to normalize the DiffuseColor and SubsurfaceColor to ensure the sum is < 1
+							const float TransmittanceNoL = 1.0f;
+							// Reverse the MFP mapping to get back to the original SubsurfaceColor that was provided to SubstrateConvertLegacyMaterialDynamic
+							const float3 MFPInCm = SLAB_SSSMFP(BSDF);
+							LegacySubsurfaceColor = MeanFreePathToTransmittance(MFPInCm * CENTIMETER_TO_METER, SUBSTRATE_SIMPLEVOLUME_THICKNESS_CM * CENTIMETER_TO_METER);
+							LegacyOpacity = 1.f - abs(SLAB_SSSPHASEANISOTROPY(BSDF));
+
+							float3 Front = DiffuseColor;
+							float3 Back = LegacySubsurfaceColor;
+							float3 Sum = Front + Back;
+							float S = max3(Sum.x, Sum.y, Sum.z);
+							if (S > 1)
+							{
+								// nudge the material back to something physically plausible
+								// NOTE: we ignore spec here since it should be accounted for by the brdf model itself
+								DiffuseColor = Front / S;
+								LegacySubsurfaceColor = Back / S;
+							}
+						}
+						// SUBSTRATE_TODO: Other SSSTYPE cases may result in some portion of a "glass" BSDF
+						float FuzzAmount = 0.0;
+						float FuzzRoughness = 1.0;
+						float3 FuzzColor = 0.0;
+						if (BSDF_GETHASFUZZ(BSDF))
 						{
-							// nudge the material back to something physically plausible
-							// NOTE: we ignore spec here since it should be accounted for by the brdf model itself
-							Payload.DiffuseColor = Front / S;
-							SubsurfaceColor = Back / S;
+							FuzzAmount = SLAB_FUZZ_AMOUNT(BSDF);
+							FuzzColor = SLAB_FUZZ_COLOR(BSDF);
+							FuzzRoughness = SLAB_FUZZ_ROUGHNESS(BSDF);
+							FuzzRoughness = MakeRoughnessSafe(FuzzRoughness, SUBSTRATE_MIN_FUZZ_ROUGHNESS);
 						}
-						if (bFoliage)
+						const float FuzzE = ComputeSheenEnergyTerms(FuzzRoughness, NoV).E;
+						const float FuzzAttenuation = 1.0 - FuzzAmount * FuzzE;
+						float DiffuseWeight = FuzzAttenuation;
+						float3 Spec0Albedo = FuzzAttenuation * Spec0E;
+						float3 Spec1Albedo = FuzzAttenuation * Spec1E;
+						if (bUseClearCoat)
 						{
-							// foliage case, only the color matters
-							Payload.ShadingModelID = SHADINGMODELID_TWOSIDED_FOLIAGE;
-							Payload.SetFoliageTransmissionColor(SubsurfaceColor);
+							// clearcoat slab
+							float CoatAttenuation = 1.0 - RoughnessData.z * Spec1E.x;
+							DiffuseWeight *= CoatAttenuation * (1.0 - Luminance(Spec0E));
+							Spec0Albedo *= CoatAttenuation;
+							Spec1Albedo *= RoughnessData.z;
 						}
 						else
 						{
-
-							// Match formula from non-Substrate handling of wrap SSS
-							const float SSSRadius = (1 - Opacity) * 10.0; // simple formula, up to 10cm radius
-							// The legacy behavior has both diffuse+sss responses separately, but we can combine them to avoid needing a custom lobe
-							// In practice the diffuse and SSS responses are similar, so the difference should be minimal (and is an approximation of the raster behavior anyway)
-							Payload.DiffuseColor += SubsurfaceColor;
-							Payload.MeanFreePath = SSSRadius;
-							Payload.PhaseG = 0;
+							// dual-specular slab
+							DiffuseWeight *= 1.0 - Luminance(lerp(Spec0E, Spec1E, RoughnessData.z));
+							Spec0Albedo *= 1.0 - RoughnessData.z;
+							Spec1Albedo *= RoughnessData.z;
 						}
-					}
-					else
-					{
-						// TODO: Are there any other SSS types to consider?
-						Payload.MeanFreePath = 0;
-						Payload.PhaseG       = 0;
-					}
-					Payload.SpecularProfileId = 0;
-					Payload.GlintValue = 1.f;
+						// SUBSTRATE_TODO: Need to account for glass case here for lobe selection ....
+						const float3 FuzzAlbedo = FuzzAmount * FuzzE * FuzzColor;
+						const float3 SlabAlbedo = MaxLobeWeight * (DiffuseWeight * (DiffuseColor + LegacySubsurfaceColor) + Spec0Albedo + Spec1Albedo + FuzzAlbedo);
+						#if SUBSTRATE_CLAMPED_CLOSURE_COUNT > 1
+						// If we are dealing witha multi-BSDF material, stochastically decide if we need to write to the payload
+						// This means estimating the overal albedo of the slab and selectively writing to the payload
+						// NOTE: Technically, some of the albedo calculation are not needed if we aren't picking between slabs
+						// We rely on the compiler to get rid of dead code here to avoid diverging the logic too far between single and multi-lobe cases
+						float Contrib = LobeColorToWeight(SlabAlbedo);
+						if (SlabSampler.Accept(Contrib))
+						#endif
+						{
+							// We've selected this particular slab - now write it into the payload
+							// NOTE: This could happen more than once, so be careful to write to all fields, even the unused ones!
+							#if SUBSTRATE_CLAMPED_CLOSURE_COUNT > 1
+							Payload.BSDFOpacity = 1.0 / Contrib;
+							#endif
+							Payload.WorldNormal = N;
+							if (SubstrateGetSharedLocalBasisType(SubstratePixelHeader.SharedLocalBases.Types, BSDFIdx) == SUBSTRATE_BASIS_TYPE_TANGENT)
+							{
+								Payload.WorldTangent = SubstratePixelHeader.SharedLocalBases.Tangents[BSDF_GETSHAREDLOCALBASISID(BSDF)];
+							}
+							Payload.WeightV               = BSDF.LuminanceWeightV;
+							Payload.TransmittanceN        = BSDF.TransmittanceAboveAlongN;
+							Payload.CoverageAboveAlongN   = BSDF.CoverageAboveAlongN;
+							Payload.ShadingModelID    	  = SHADINGMODELID_SUBSTRATE;
+							Payload.Ior                   = 0.0;
+							Payload.DiffuseColor          = DiffuseColor;
+							Payload.SpecularColor         = F0;
+							Payload.SpecularEdgeColor     = F90;
+							Payload.RoughnessData         = RoughnessData;
+							Payload.BottomNormalOct16bits = 0;
+							if (bUseClearCoat)
+							{
+								Payload.ShadingModelID = SHADINGMODELID_CLEAR_COAT;
+							#if CLEAR_COAT_BOTTOM_NORMAL
+								Payload.BottomNormalOct16bits = SLAB_HAZINESS(BSDF) >> 16;
+							#endif
+							}
+							Payload.Anisotropy = BSDF_GETHASANISOTROPY(BSDF) ? SLAB_ANISOTROPY(BSDF) : 0.0;
+							if (BSDF_GETHASFUZZ(BSDF))
+							{
+								Payload.FuzzAmount    = FuzzAmount;
+								Payload.FuzzColor     = FuzzColor;
+								Payload.FuzzRoughness = FuzzRoughness;
+							}
+							else
+							{
+								Payload.FuzzAmount    = 0.0;
+								Payload.FuzzColor     = 0.0;
+								Payload.FuzzRoughness = 0.0;
+							}
+							Payload.MeanFreePath = 0;
+							Payload.PhaseG       = 0;
+							if (BSDF_GETSSSTYPE(BSDF) == SSS_TYPE_SIMPLEVOLUME)
+							{
+								// TODO: this case should be represented by a different BSDF
+								// For now, we already adjusted DiffuseColor according to MFP above, so nothing extra to do
+							}
+							else if (BSDF_GETSSSTYPE(BSDF) == SSS_TYPE_DIFFUSION)
+							{
+								Payload.MeanFreePath = SLAB_SSSMFP(BSDF);
+								Payload.PhaseG       = SLAB_SSSPHASEANISOTROPY(BSDF);
+							}
+							else if (BSDF_GETSSSTYPE(BSDF) == SSS_TYPE_DIFFUSION_PROFILE)
+							{
+								const uint ProfileId = SLAB_SSSPROFILEID(BSDF);
+								const float Scale = SLAB_SSSPROFILERADIUSSCALE(BSDF);
+								Payload.MeanFreePath = DecodeSSSProfileRadius(ProfileId, Payload.DiffuseColor, Scale);
+								#if USE_SSS_PROFILE_ANISOTROPY
+								const float G = DecodeSSSProfileScatteringDistribution(ProfileId);
+								Payload.PhaseG = G;
+								Payload.MeanFreePath *= 1 - G; // Adjust MFP to preserve appearance
+								#endif
+							}
+							else if (BSDF_GETSSSTYPE(BSDF) == SSS_TYPE_WRAP)
+							{
+								// Legacy support for SHADINGMODELID_SUBSURFACE
+								const float SSSRadius = (1 - LegacyOpacity) * 10.0; // simple formula, up to 10cm radius
+								// The legacy behavior has both diffuse+sss responses separately, but we can combine them to avoid needing a custom lobe
+								// In practice the diffuse and SSS responses are similar, so the difference should be minimal (and is an approximation of the raster behavior anyway)
+								// SUBSTRATE_TODO: If we implement stochastic lobe picking _within_ slabs, we could fine tune this to get back to the original behavior exactly
+								Payload.DiffuseColor += LegacySubsurfaceColor;
+								Payload.MeanFreePath = SSSRadius;
+								Payload.PhaseG = 0;
+							}
+							else if (BSDF_GETSSSTYPE(BSDF) == SSS_TYPE_TWO_SIDED_WRAP)
+							{
+								// Legacy support for SHADINGMODELID_TWOSIDED_FOLIAGE
+								Payload.ShadingModelID = SHADINGMODELID_TWOSIDED_FOLIAGE;
+								Payload.SetFoliageTransmissionColor(LegacySubsurfaceColor);
+							}
 
-#if SUBSTRATE_COMPLEXPATH
-					if (BSDF_GETHASSPECPROFILE(BSDF))
-					{
-						Payload.SpecularProfileId = SLAB_SPECPROFILEID(BSDF);
-					}
-#endif
-#if SUBSTRATE_COMPLEXSPECIALPATH
-					if (BSDF_GETHASGLINT(BSDF))
-					{
-						Payload.GlintValue = SLAB_GLINT_VALUE(BSDF);
-						Payload.GlintUV    = SLAB_GLINT_UV(BSDF);
-						Payload.GlintUVdx  = SLAB_GLINT_UVDDX(BSDF);
-						Payload.GlintUVdy  = SLAB_GLINT_UVDDY(BSDF);
-					}
-#endif
-					Payload.WeightV             = BSDF.LuminanceWeightV;
-					Payload.TransmittanceN      = BSDF.TransmittanceAboveAlongN;
-					Payload.CoverageAboveAlongN = BSDF.CoverageAboveAlongN;
-#if MATERIAL_ISTHINSURFACE && MATERIALBLENDING_ANY_TRANSLUCENT
-					Payload.ShadingModelID      = SHADINGMODELID_THIN_TRANSLUCENT;
-#endif
+							if (BSDF.bIsBottom)
+							{
+								#if REFRACTION_USE_INDEX_OF_REFRACTION
+								float Ior = GetRefractionIor(PixelMaterialInputs, F0);
+								#endif
 
-#if MATERIALBLENDING_ANY_TRANSLUCENT && REFRACTION_USE_INDEX_OF_REFRACTION
-					// NOTE: only the bottom most layer can refract, coatings should never include a refraction lobe
-					Payload.Ior					= BSDF.bIsBottom ? GetRefractionIor(PixelMaterialInputs, Payload.SpecularColor) : 0.0;
-					Payload.MeanFreePath        = BSDF.bIsBottom ? SLAB_SSSMFP(BSDF) : Payload.MeanFreePath;
-					if (Payload.Ior > 0 && Payload.ShadingModelID == SHADINGMODELID_SUBSTRATE)
-					{
-						Payload.ShadingModelID = SHADINGMODELID_SOLID_GLASS;
-					}
-#else
-					Payload.Ior                 = 0.0;
-#endif
+								#if 1 && MATERIALBLENDING_TRANSLUCENT && (SUBSTRATE_BLENDING_TRANSLUCENT_GREYTRANSMITTANCE || SUBSTRATE_BLENDING_TRANSLUCENT_COLOREDTRANSMITTANCE) && REFRACTION_USE_INDEX_OF_REFRACTION
+								// Disable refraction handling for now -- need to come up with a better way to know the amount of "glass" we have separate from the concept of opacity
+								// In legacy slabs, the old Opacity plug should blend between DefaultLit behavior and Glass, while in new Substrate slabs this blend should depend on the amount of albedo and MFP.
 
-#if MATERIAL_ISTHINSURFACE && MATERIALBLENDING_ANY_TRANSLUCENT
-					// Material is configured as a thin surface, and refraction is disabled
-					if (BSDF.bIsBottom && !(Payload.Ior > 0.0))
-					{
-						Payload.MeanFreePath = SLAB_SSSMFP(BSDF);
-						float3 Transmission = Payload.GetTransmittanceColor();
-						const float NoV = saturate(dot(Payload.WorldNormal, V_World));
-						float3 Tr = ComputeThinSlabWeights(Transmission, NoV, Payload.Ior, F0RGBToF0(Payload.SpecularColor)).Transmitted;
-						Payload.TransparencyColor += TotalCoverage * Tr;
-					}
-#endif
+								// NOTE: only the bottom most layer can refract, coating slabs should never include a refraction lobe
+								// We group this with the SSS conditions because it involves MFP
 
-					break;
-				}
-				case SUBSTRATE_BSDF_TYPE_HAIR:
-				{
-					// clamp the roughness to whatever came along the path
-					Payload.SetBaseColor(HAIR_BASECOLOR(BSDF));
-					Payload.SetHairLongitudinalRoughness(max(HAIR_ROUGHNESS(BSDF), PackedPayload.GetPathRoughness()));
-					Payload.SetHairAzimuthalRoughness(HAIR_SCATTER(BSDF)); // SUBSTRATE_TODO: Map from scatter param to azimuthal roughness here instead of in the raygen
-					Payload.SetHairSpecular(HAIR_SPECULAR(BSDF));
+								float3 MFP = SLAB_SSSMFP(BSDF);
+								if (Ior > 0 && any(MFP > 0.0) && Payload.ShadingModelID == SHADINGMODELID_SUBSTRATE)
+								{
+									// SUBSTRATE_TODO: This isn't quite correct - we want to incorporate a proper mixing for the diffuse/sss part of the material here, not just make pure glass
+									#if MATERIAL_ISTHINSURFACE || MATERIAL_SHADINGMODEL_THIN_TRANSLUCENT // SUBSTRATE_TODO: Should THIN_TRANSLUCENT imply ISTHIN_SURFACE?
+									Payload.ShadingModelID = SHADINGMODELID_THIN_TRANSLUCENT;
+									#else
+									Payload.ShadingModelID = SHADINGMODELID_SOLID_GLASS;
+									#endif
+									Payload.MeanFreePath   = MFP;
+									Payload.Ior            = Ior;
+								}
+								#endif
 
-	#if HAIR_STRAND_MESH_FACTORY
-					Payload.WorldSmoothNormal = Payload.WorldNormal = Payload.WorldGeoNormal;
-					Payload.WorldTangent      = MaterialParameters.TangentToWorld[2];
-					Payload.SetHairPrimitiveUV(MaterialParameters.HairPrimitiveUV);
-	#else
-					Payload.SetHairPrimitiveUV(0.5);
-	#endif
-					Payload.ShadingModelID = SHADINGMODELID_HAIR;
+								#if (MATERIAL_ISTHINSURFACE || MATERIAL_SHADINGMODEL_THIN_TRANSLUCENT) && MATERIALBLENDING_TRANSLUCENT && (SUBSTRATE_BLENDING_TRANSLUCENT_GREYTRANSMITTANCE || SUBSTRATE_BLENDING_TRANSLUCENT_COLOREDTRANSMITTANCE)
+								// Material is configured as a thin surface, and refraction is disabled
+								if (!(Payload.Ior > 0.0))
+								{
+									// SUBSTRATE_TODO: Do we really want to keep this? Should maybe always use refraction path?
+									Payload.MeanFreePath = SLAB_SSSMFP(BSDF);
+									float3 Transmission = Payload.GetTransmittanceColor();
+									const float NoV = saturate(dot(Payload.WorldNormal, V_World));
+									float3 Tr = ComputeThinSlabWeights(Transmission, NoV, Payload.Ior, F0RGBToF0(Payload.SpecularColor)).Transmitted;
+									Payload.TransparencyColor += TotalCoverage * Tr;
+								}
+								#endif
+
+							}
+
+							// Handle special-case specular features
+							Payload.SpecularProfileId = 0;
+							Payload.GlintValue = 1.f;
+
+							#if SUBSTRATE_COMPLEXPATH
+							if (BSDF_GETHASSPECPROFILE(BSDF))
+							{
+								Payload.SpecularProfileId = SLAB_SPECPROFILEID(BSDF);
+							}
+							#endif
+							#if SUBSTRATE_COMPLEXSPECIALPATH
+							if (BSDF_GETHASGLINT(BSDF))
+							{
+								Payload.GlintValue = SLAB_GLINT_VALUE(BSDF);
+								Payload.GlintUV    = SLAB_GLINT_UV(BSDF);
+								Payload.GlintUVdx  = SLAB_GLINT_UVDDX(BSDF);
+								Payload.GlintUVdy  = SLAB_GLINT_UVDDY(BSDF);
+							}
+							#endif
+						}
+					}
 					break;
 				}
+				// NOTE: The cases below are for back-compatibility and cannot be layered
 				case SUBSTRATE_BSDF_TYPE_EYE:
 				{
+					Payload.WorldNormal = SubstratePixelHeader.SharedLocalBases.Normals[BSDF_GETSHAREDLOCALBASISID(BSDF)];
 					Payload.DiffuseColor = EYE_DIFFUSEALBEDO(BSDF);
 					Payload.SpecularColor = EYE_F0(BSDF);
 					// NOTE: EYE_F90 is always 1.0 (does not exist in legacy model), so don't bother with it
-					Payload.SetEyeRoughness(max(EYE_ROUGHNESS(BSDF), PackedPayload.GetPathRoughness()));
+					Payload.SetEyeRoughness(AdjustMaterialRoughness(EYE_ROUGHNESS(BSDF), PackedPayload.GetPathRoughness()));
 
 					if (BSDF_GETSSSTYPE(BSDF) == SSS_TYPE_DIFFUSION_PROFILE)
 					{
@@ -757,31 +707,49 @@ RAY_TRACING_ENTRY_CLOSEST_HIT(PathTracingMaterialCHS,
 					Payload.ShadingModelID = SHADINGMODELID_EYE;
 					break;
 				}
+				case SUBSTRATE_BSDF_TYPE_HAIR:
+				{
+					// clamp the roughness to whatever came along the path
+					Payload.SetBaseColor(HAIR_BASECOLOR(BSDF));
+					Payload.SetHairLongitudinalRoughness(AdjustMaterialRoughness(HAIR_ROUGHNESS(BSDF), PackedPayload.GetPathRoughness()));
+					Payload.SetHairAzimuthalRoughness(HAIR_SCATTER(BSDF)); // SUBSTRATE_TODO: Map from scatter param to azimuthal roughness here instead of in the raygen
+					Payload.SetHairSpecular(HAIR_SPECULAR(BSDF));
+
+					#if HAIR_STRAND_MESH_FACTORY
+					Payload.WorldSmoothNormal = Payload.WorldNormal = Payload.WorldGeoNormal;
+					Payload.WorldTangent      = MaterialParameters.TangentToWorld[2];
+					Payload.SetHairPrimitiveUV(MaterialParameters.HairPrimitiveUV);
+					#else
+					Payload.SetHairPrimitiveUV(0.5);
+					#endif
+					Payload.ShadingModelID = SHADINGMODELID_HAIR;
+					break;
+				}
 				case SUBSTRATE_BSDF_TYPE_SINGLELAYERWATER:
 				{
+					Payload.WorldNormal = SubstratePixelHeader.SharedLocalBases.Normals[BSDF_GETSHAREDLOCALBASISID(BSDF)];
+					Payload.BSDFOpacity = 1;
 					Payload.ShadingModelID = SHADINGMODELID_SOLID_GLASS;
 					float3 BaseColor = SLW_BASECOLOR(BSDF);
 					float Metallic = SLW_METALLIC(BSDF);
 					float Specular = SLW_SPECULAR(BSDF);
 					float Roughness = SLW_ROUGHNESS(BSDF);
 					float3 Extinction = 0.0;
-#if SUBSTRATE_INLINE_SINGLELAYERWATER
+					#if SUBSTRATE_INLINE_SINGLELAYERWATER
 
 					//float Opacity = SLW_TOPMATERIALOPACITY(BSDF); // TODO: Account for this with an extra random decision?
 					//float3 WaterAlbedo = SLW_WATERALBEDO(BSDF); // TODO: Account for scattering in refractive mediums
 					Extinction = SLW_WATEREXTINCTION(BSDF);
 					//float PhaseG = SLW_WATERPHASEG(BSDF);
 					//float3 ColorBehindWater = SLW_COLORSCALEBEHINDWATER(BSDF);
-#else
-					
-#endif
+					#endif
 					Payload.WeightV             = BSDF.LuminanceWeightV;
 					Payload.TransmittanceN      = BSDF.TransmittanceAboveAlongN;
 					Payload.CoverageAboveAlongN = BSDF.CoverageAboveAlongN;
 
 					Payload.DiffuseColor = ComputeDiffuseAlbedo(BaseColor, Metallic);
 					Payload.SpecularColor = ComputeF0(Specular, BaseColor, Metallic);
-					Payload.RoughnessData.x = MakeRoughnessSafe(Roughness, PackedPayload.GetPathRoughness());
+					Payload.RoughnessData.x = AdjustMaterialRoughness(Roughness, PackedPayload.GetPathRoughness());
 					Payload.RoughnessData.yz = 0;
 					Payload.Ior = DielectricF0ToIor(DielectricSpecularToF0(Specular));
 					Payload.MeanFreePath = rcp(max(Extinction, 1e-5));
@@ -791,21 +759,16 @@ RAY_TRACING_ENTRY_CLOSEST_HIT(PathTracingMaterialCHS,
 					Payload.GlintValue = 1;
 					break;
 				}
-				default:
-				{
-					// SUBSTRATE_TODO: handle all other cases
-					break;
-				}
-			}
-#if SUBSTRATE_CLAMPED_CLOSURE_COUNT > 1
-			// account for the probability of having chosen this slab
-			Payload.BSDFOpacity /= SlabPdf;
-#endif
-		} // close brace opened by SUBSTRATE_CLAMPED_CLOSURE_COUNT > 1 or SUBSTRATE_CLAMPED_CLOSURE_COUNT == 1 ifdefs above
-#endif // SUBSTRATE_CLAMPED_CLOSURE_COUNT > 0
+			}  // switch on slab tyype
+		} // for loop over BSDFs
+
+		#if SUBSTRATE_CLAMPED_CLOSURE_COUNT > 1
+		// normalize the weight for picking a given slab
+		Payload.BSDFOpacity *= SlabSampler.GetNormalization();
+		#endif
 
 #undef  BSDF
-	} // if BSDFCount > 0
+	} // if (BSDFCount > 0)
 
 
 #endif // SUBSTRATE_OPTIMIZED_UNLIT
@@ -1012,6 +975,7 @@ RAY_TRACING_ENTRY_CLOSEST_HIT(PathTracingMaterialCHS,
 		float ScatteringDistribution = clamp(GetSubsurfaceMediumMaterialOutput1(MaterialParameters), -0.99f, 0.99f);
 		Payload.SetSubsurfacePhaseFunction(ScatteringDistribution);
 #elif USE_SSS_PROFILE_ANISOTROPY
+		// NOTE: this is not currently used
 		float G = DecodeSSSProfileScatteringDistribution(ProfileId);
 		Payload.SetSubsurfacePhaseFunction(G);
 		Payload.SetSubsurfaceRadius(Payload.GetSubsurfaceRadius() * (1 - G));
@@ -1356,8 +1320,6 @@ RAY_TRACING_ENTRY_ANY_HIT(PathTracingMaterialAHS,
 
 #if SUBSTRATE_ENABLED
 
-		float3 Transparency = 0.0;
-
 #if MATERIAL_IS_SUBSTRATE 
 		FSubstratePixelHeader SubstratePixelHeader = MaterialParameters.GetFrontSubstrateHeader();
 		FSubstrateData SubstrateData = PixelMaterialInputs.GetFrontSubstrateData();
@@ -1412,10 +1374,14 @@ RAY_TRACING_ENTRY_ANY_HIT(PathTracingMaterialAHS,
 							RoughnessData.z = Haziness.Weight;
 						}
 
-						LocalSigmaT = rcp(max(SLAB_SSSMFP(BSDF), 0.0001));
-#if REFRACTION_USE_INDEX_OF_REFRACTION
-						Ior = GetRefractionIor(PixelMaterialInputs, F0);
+						if (BSDF_GETSSSTYPE(BSDF) == SSS_TYPE_NONE)
+						{
+							// only non-SSS slabs 
+							LocalSigmaT = rcp(max(SLAB_SSSMFP(BSDF), 0.0001));
+#if MATERIALBLENDING_TRANSLUCENT && (SUBSTRATE_BLENDING_TRANSLUCENT_GREYTRANSMITTANCE || SUBSTRATE_BLENDING_TRANSLUCENT_COLOREDTRANSMITTANCE) && REFRACTION_USE_INDEX_OF_REFRACTION
+							Ior = GetRefractionIor(PixelMaterialInputs, F0);
 #endif
+						}
 						break;
 					}
 					case SUBSTRATE_BSDF_TYPE_SINGLELAYERWATER:
@@ -1447,7 +1413,7 @@ RAY_TRACING_ENTRY_ANY_HIT(PathTracingMaterialAHS,
 		// Solid Glass Fake Caustics
 		float Tr = FakeCaustics(F0, Ior, NoV);
 		// Account for transmission 
-		PackedPayload.SetTau(PackedPayload.GetTau() + ComputeOpticalDepthForHit(LocalSigmaT, bIsFrontFace));
+		PackedPayload.SetTau(PackedPayload.GetTau() + TotalCoverage * ComputeOpticalDepthForHit(LocalSigmaT, bIsFrontFace));
 #endif
 
 		// If material is doing refraction, apply an extra factor to account for ApproximateCaustics
@@ -1458,14 +1424,17 @@ RAY_TRACING_ENTRY_ANY_HIT(PathTracingMaterialAHS,
 			           FakeCausticsPathFactor(RoughnessData.y, PathRoughness), RoughnessData.z);
 		}
 #endif
-		Transparency += TotalCoverage * Tr;
+		const float3 Transparency = saturate(lerp(1.0, Tr, TotalCoverage));
 
-		// SUBSTRATE_TODO: Is this sanity check necessary?
-		Transparency = saturate(Transparency);
 #else // MATERIALBLENDING_TRANSLUCENT && (REFRACTION_USE_INDEX_OF_REFRACTION || MATERIAL_ISTHINSURFACE)
 		// All other blend modes, do simple transparency handling
-		Transparency = lerp(1.0, TotalTransmittancePreCoverage, TotalCoverage);
+		const float3 Transparency = saturate(lerp(1.0, TotalTransmittancePreCoverage, TotalCoverage));
 #endif
+
+#else // MATERIAL_IS_SUBSTRATE
+
+		// TODO: What cases (if any) end up here?
+		const float3 Transparency = 0.0;
 #endif // MATERIAL_IS_SUBSTRATE
 
 #else // SUBSTRATE_ENABLED
