Substrate - fixed failing material cooking due to uninitialised texcoord leading to uninitialised Nanite derivatives.

#tests PC editor dx12, package win64.
#rnx

diff --git a/Engine/Shaders/Private/MaterialTemplate.ush b/Engine/Shaders/Private/MaterialTemplate.ush
--- a/Engine/Shaders/Private/MaterialTemplate.ush
+++ b/Engine/Shaders/Private/MaterialTemplate.ush
@@ -4086,6 +4086,13 @@ SHADER_PUSH_WARNINGS_STATE
 SHADER_DISABLE_WARNINGS
 void GetMaterialCustomizedUVs(FMaterialVertexParameters Parameters, inout float2 OutTexCoords[NUM_TEX_COORD_INTERPOLATORS])
 {
+	// Always initialise to avoid later operations to potentially cause compiler failure due to read uninitialized value when get_material_customized_u_vs is not defined.
+	UNROLL
+	for (int CoordinateIndex = 0; CoordinateIndex < NUM_TEX_COORD_INTERPOLATORS; CoordinateIndex++)
+	{
+		OutTexCoords[CoordinateIndex] = 0.0f;
+	}
+
 %{get_material_customized_u_vs}
 }
 
diff --git a/Engine/Shaders/Shared/SubstrateVisualizeDefinitions.h b/Engine/Shaders/Shared/SubstrateVisualizeDefinitions.h
--- a/Engine/Shaders/Shared/SubstrateVisualizeDefinitions.h
+++ b/Engine/Shaders/Shared/SubstrateVisualizeDefinitions.h
@@ -14,7 +14,7 @@
 #pragma message("UESHADERMETADATA_VERSION 24677C38-A71C-49D5-B60A-DFF59F57B39D")
 #include "/Engine/Shared/SubstrateDefinitions.h"
 
-#if PS4_PROFILE
+#if PS4_PROFILE || FEATURE_LEVEL <= FEATURE_LEVEL_SM5
 // For platform having trouble avoiding register spilling.
 // To fix this in general, we will have to break down the FSubstratePixelDebugData structure into Header and Closure structures with interleaved write to buffer.
 #define SUBSTRATE_VIS_CLOSURE_COUNT 1
