Substrate - fixed compilation of opaque material sampling depth or stencil.

#tests PC editor dx12
#rnx

diff --git a/Engine/Shaders/Private/MaterialTemplate.ush b/Engine/Shaders/Private/MaterialTemplate.ush
--- a/Engine/Shaders/Private/MaterialTemplate.ush
+++ b/Engine/Shaders/Private/MaterialTemplate.ush
@@ -203,7 +203,11 @@
 #else
 	#define SubstratePublicStruct Substrate
 #endif
+
+#if !SUBSTRATE_OPAQUE_DEFERRED
 #include "/Engine/Public/SubstratePublic.ush"
+#endif
+
 #undef SubstratePublicStruct
 
 #else
@@ -3290,6 +3294,7 @@ float4 SceneTextureLookup(float2 UV, int SceneTextureIndex, bool bFiltered)
 	}
 	#elif SUBSTRATE_GBUFFER_FORMAT==1 && !SHADING_PATH_MOBILE // Not on mobile that is not using dbuffer decal but the legacy GBuffer blending.
 	{
+		#if !SUBSTRATE_OPAQUE_DEFERRED // Prevents opaque material to sample the Gbuffer.
 		FSubstrateGBuffer SubstrateGBuffer = SubstratePublic_GetSubstrateGBufferFromFirstBSDF(PixelPos);
 		BRANCH
 		if (SubstrateGBuffer.bIsValid > 0)
@@ -3316,6 +3321,7 @@ float4 SceneTextureLookup(float2 UV, int SceneTextureIndex, bool bFiltered)
 				case PPI_IsFirstPerson:			return SubstrateGBuffer.bIsFirstPerson ? 1.0f : 0.0f;
 			}
 		}
+		#endif
 	}
 	#else // SUBSTRATE_ENABLED
 	{
