Fix histogram computation when disabling pre-calculatedIlluminance and using Substrate Format=1.

diff --git a/Engine/Shaders/Private/Histogram.usf b/Engine/Shaders/Private/Histogram.usf
--- a/Engine/Shaders/Private/Histogram.usf
+++ b/Engine/Shaders/Private/Histogram.usf
@@ -7,6 +7,7 @@
 #include "/Engine/Private/Common.ush"
 #include "/Engine/Private/ScreenPass.ush"
 #include "/Engine/Private/PostProcessHistogramCommon.ush"
+#include "/Engine/Private/PostProcessExposureMaterial.ush"
 
 /////////////////////////////////////////////////////////////////////////////////////////////////////////////////////
 //ATOMIC HISTOGRAM
@@ -323,34 +324,11 @@ void MainAtomicCS(
 			LuminanceVal = SceneColor.r;
 #elif USE_APPROX_ILLUMINANCE
 			const uint2 GBufferTexelPos = View.ViewRectMinAndSize.xy + ScreenUV * View.ViewRectMinAndSize.zw;
-			const float2 GBufferUV = (GBufferTexelPos + 0.5f) * View.BufferSizeAndInvSize.zw;
-
-			FGBufferData GBufferData = GetGBufferDataFromSceneTextures(GBufferUV);
 			
-			if (GBufferData.ShadingModelID != SHADINGMODELID_UNLIT)
+			const FExposureMaterial Material = GetExposureMaterial(GBufferTexelPos);
+			if (Material.bIsValid)
 			{
-				float3 Denominator = GBufferData.DiffuseColor;
-
-				if (GBufferData.ShadingModelID == SHADINGMODELID_SUBSURFACE
-					|| GBufferData.ShadingModelID == SHADINGMODELID_PREINTEGRATED_SKIN
-					|| GBufferData.ShadingModelID == SHADINGMODELID_TWOSIDED_FOLIAGE
-					|| GBufferData.ShadingModelID == SHADINGMODELID_CLOTH
-					|| GBufferData.ShadingModelID == SHADINGMODELID_EYE)
-				{
-					const float3 SubsurfaceColor = ExtractSubsurfaceColor(GBufferData);
-					Denominator += SubsurfaceColor;
-				}
-
-				{
-					const float3 TranslatedWorldPosition = ReconstructTranslatedWorldPositionFromDepth(GBufferUV, GBufferData.Depth);
-					const float3 N = GBufferData.WorldNormal;
-					const float3 V = -GetCameraVectorFromTranslatedWorldPosition(TranslatedWorldPosition);
-					const float3 EnvBrdf = EnvBRDF(GBufferData.SpecularColor, GBufferData.Roughness, max(0.0, dot(N, V)));
-
-					Denominator += EnvBrdf;
-				}
-
-				float BaseColorLuminanceVal = CalculateEyeAdaptationLuminance(Denominator);
+				float BaseColorLuminanceVal = CalculateEyeAdaptationLuminance(Material.DirectionalAlbedo);
 
 				LuminanceVal /= max(BaseColorLuminanceVal, EyeAdaptation_IgnoreMaterialsMinBaseColorLuminance);
 				LuminanceVal *= lerp(EyeAdaptation_IgnoreMaterialsLuminanceScale, 1.0f, EyeAdaptation_IgnoreMaterialsMinBaseColorLuminance);
diff --git a/Engine/Shaders/Private/PostProcessExposureMaterial.ush b/Engine/Shaders/Private/PostProcessExposureMaterial.ush
new file mode 100644
--- /dev/null
+++ b/Engine/Shaders/Private/PostProcessExposureMaterial.ush
@@ -0,0 +1,126 @@
+// Copyright Epic Games, Inc. All Rights Reserved.
+
+#include "DeferredShadingCommon.ush"
+#include "BRDF.ush"
+#include "PositionReconstructionCommon.ush"
+#include "Substrate/Substrate.ush"
+#include "Substrate/SubstrateEvaluation.ush"
+
+struct FExposureMaterial
+{
+	float3 TranslatedWorldPosition;
+	float3 WorldNormal;
+	float3 DirectionalAlbedo;
+
+	bool bIsValid;
+	bool bHasBackfaceLighting;
+	bool bHasValidDirectionalAlbedo;
+};
+
+// same weights as CalculateEyeAdaptationLuminance
+// but skip the min since we apply EyeAdaptation_IgnoreMaterialsMinBaseColorLuminance here
+float CalculateEyeAdaptationLuminanceWithoutMin(float3 Color)
+{
+	return dot(Color, EyeAdaptation_LuminanceWeights);
+}
+
+FExposureMaterial GetExposureMaterial(uint2 InPixelPos)
+{
+	FExposureMaterial Out = (FExposureMaterial)0;
+#if SHADING_PATH_DEFERRED
+	#if SUBSTRATE_GBUFFER_FORMAT==1
+	{	
+		FSubstrateAddressing SubstrateAddressing = GetSubstratePixelDataByteOffset(InPixelPos, uint2(View.BufferSizeAndInvSize.xy), Substrate.MaxBytesPerPixel);
+		FSubstratePixelHeader SubstratePixelHeader = UnpackSubstrateHeaderIn(Substrate.MaterialTextureArray, SubstrateAddressing, Substrate.TopLayerTexture);
+		FSubstrateTopLayerData TopLayerData = SubstrateUnpackTopLayerData(Substrate.TopLayerTexture.Load(uint3(InPixelPos, 0)));
+		FSubstrateSubsurfaceHeader SSSHeader = SubstrateLoadSubsurfaceHeader(Substrate.MaterialTextureArray, Substrate.FirstSliceStoringSubstrateSSSData, InPixelPos);
+
+		const float DeviceZ = ConvertFromDeviceZ(SceneDepthTexture.Load(int3(InPixelPos, 0)).r);
+		const float3 TranslatedWorldPosition = ReconstructTranslatedWorldPositionFromDeviceZ(InPixelPos, DeviceZ);
+		const float3 V = normalize(View.TranslatedWorldCameraOrigin - TranslatedWorldPosition);
+
+		float3 DiffuseDirectionalAlbedo = 0;
+		float3 DirectionalAlbedo = 0;
+		bool bHasMaterialBackfaceDiffuse = false;
+		{
+			FSubstrateDeferredLighting Out = GetInitialisedSubstrateDeferredLighting();
+			const FSubstrateIntegrationSettings Settings = InitSubstrateIntegrationSettings(false /*bForceFullyRough*/, Substrate.bRoughDiffuse, Substrate.PeelLayersAboveDepth, Substrate.bRoughnessTracking);
+			Substrate_for(uint ClosureIndex = 0, ClosureIndex < SubstratePixelHeader.ClosureCount, ++ClosureIndex)
+			{
+				FSubstrateBSDF BSDF = UnpackSubstrateBSDF(Substrate.MaterialTextureArray, SubstrateAddressing, SubstratePixelHeader);
+				FSubstrateBSDFContext SubstrateBSDFContext = SubstrateCreateBSDFContext(SubstratePixelHeader, BSDF, SubstrateAddressing, V);
+				const float3 BSDFThroughput = LuminanceWeight(SubstrateBSDFContext, BSDF);
+
+				// Evaluate environment lighting
+				FSubstrateEnvLightResult SubstrateEnvLight = SubstrateEvaluateForEnvLight(SubstrateBSDFContext, true /*bEnableSpecular*/, Settings);
+
+				DiffuseDirectionalAlbedo += BSDFThroughput * SubstrateEnvLight.DiffuseWeight;
+				DirectionalAlbedo        += BSDFThroughput * SubstrateEnvLight.DiffuseWeight;
+				DirectionalAlbedo        += BSDFThroughput * SubstrateEnvLight.SpecularWeight;
+
+				const bool bHasBackfaceDiffuse = SubstrateGetBSDFType(BSDF) == SUBSTRATE_BSDF_TYPE_SLAB && BSDF.HasBackScattering();
+				if (bHasBackfaceDiffuse)
+				{
+					DirectionalAlbedo += BSDFThroughput * SubstrateEnvLight.DiffuseBackFaceWeight; // SubstrateEnvLight.DiffuseBackFaceWeight is already divided PI
+				}
+				if (any(SubstrateEnvLight.SpecularHazeWeight > 0.0f))
+				{
+					DirectionalAlbedo += BSDFThroughput * SubstrateEnvLight.SpecularHazeWeight;
+				}
+			}
+		}
+
+		Out.TranslatedWorldPosition = ReconstructTranslatedWorldPositionFromDeviceZ(InPixelPos, DeviceZ);
+		Out.WorldNormal = TopLayerData.WorldNormal;
+		Out.DirectionalAlbedo = DirectionalAlbedo;
+		Out.bIsValid = SubstratePixelHeader.IsSubstrateMaterial();
+		Out.bHasBackfaceLighting = bHasMaterialBackfaceDiffuse;
+		Out.bHasValidDirectionalAlbedo = CalculateEyeAdaptationLuminanceWithoutMin(DiffuseDirectionalAlbedo) > EyeAdaptation_IgnoreMaterialsMinBaseColorLuminance;
+	}
+	#else // SUBSTRATE_GBUFFER_FORMAT==0 or LEGACY
+	{
+		const float2 GBufferUV = (InPixelPos + 0.5f) * View.BufferSizeAndInvSize.zw;
+		const FGBufferData GBufferData = GetGBufferDataFromSceneTextures(GBufferUV);
+		const float3 TranslatedWorldPosition = ReconstructTranslatedWorldPositionFromDepth(GBufferUV, GBufferData.Depth);
+
+		// Calculate approximate illuminance from SceneColor and GBufferData
+		// Illuminance ~= (SceneColorLuminance - Emissive) / (DiffuseColor + SubsurfaceColor + EnvBRDF)
+		float3 Denominator = 0.f;
+		Denominator += GBufferData.DiffuseColor;
+		if (GBufferData.ShadingModelID == SHADINGMODELID_SUBSURFACE
+			|| GBufferData.ShadingModelID == SHADINGMODELID_PREINTEGRATED_SKIN
+			|| GBufferData.ShadingModelID == SHADINGMODELID_TWOSIDED_FOLIAGE
+			|| GBufferData.ShadingModelID == SHADINGMODELID_CLOTH
+			|| GBufferData.ShadingModelID == SHADINGMODELID_EYE)
+		{
+			float3 SubsurfaceColor = ExtractSubsurfaceColor(GBufferData);
+
+			if (GBufferData.ShadingModelID == SHADINGMODELID_TWOSIDED_FOLIAGE)
+			{
+				SubsurfaceColor *= 1.0f / PI;
+			}
+
+			Denominator += SubsurfaceColor;
+		}
+
+		{
+			const float3 CameraVector = normalize(View.TranslatedWorldCameraOrigin - TranslatedWorldPosition);
+
+			const float3 N = GBufferData.WorldNormal;
+			const float3 V = CameraVector;
+			const float3 EnvBrdf = EnvBRDF(GBufferData.SpecularColor, GBufferData.Roughness, max(0.0, dot(N, V)));
+
+			Denominator += EnvBrdf;
+		}
+
+		Out.TranslatedWorldPosition = TranslatedWorldPosition;
+		Out.WorldNormal = GBufferData.WorldNormal;
+		Out.DirectionalAlbedo = Denominator;
+		Out.bIsValid = GBufferData.ShadingModelID != SHADINGMODELID_UNLIT;
+		Out.bHasBackfaceLighting = GetShadingModelRequiresBackfaceLighting(GBufferData.ShadingModelID);
+		Out.bHasValidDirectionalAlbedo = CalculateEyeAdaptationLuminanceWithoutMin(GBufferData.DiffuseColor) > EyeAdaptation_IgnoreMaterialsMinBaseColorLuminance;
+	}
+	#endif // SUBSTRATE_GBUFFER_FORMAT
+#endif
+	return Out;
+}
\ No newline at end of file
diff --git a/Engine/Shaders/Private/PostProcessEyeAdaptation.usf b/Engine/Shaders/Private/PostProcessEyeAdaptation.usf
--- a/Engine/Shaders/Private/PostProcessEyeAdaptation.usf
+++ b/Engine/Shaders/Private/PostProcessEyeAdaptation.usf
@@ -7,8 +7,7 @@
 #include "Common.ush"
 #include "PostProcessCommon.ush"
 #include "PostProcessHistogramCommon.ush"
-#include "Substrate/Substrate.ush"
-#include "Substrate/SubstrateEvaluation.ush"
+#include "PostProcessExposureMaterial.ush"
 
 #if COMPUTESHADER
 
@@ -376,13 +375,6 @@ RWTexture2D<float> RWIlluminanceTexture;
 SCREEN_PASS_TEXTURE_VIEWPORT(Illuminance)
 FScreenTransform IlluminanceRectToColorRect;
 
-// same weights as CalculateEyeAdaptationLuminance
-// but skip the min since we apply EyeAdaptation_IgnoreMaterialsMinBaseColorLuminance here
-float CalculateEyeAdaptationLuminanceWithoutMin(float3 Color)
-{
-	return dot(Color, EyeAdaptation_LuminanceWeights);
-}
-
 [numthreads(8, 8, 1)]
 void SetupExposureIlluminance(uint2 ThreadId : SV_DispatchThreadID)
 {
@@ -401,116 +393,6 @@ Texture3D TranslucencyLightingVolumeDirectionalOuter;
 
 #include "TranslucencyVolumeCommon.ush"
 
-struct FExposureMaterial
-{
-	float3 TranslatedWorldPosition;
-	float3 WorldNormal;
-	float3 DirectionalAlbedo;
-
-	bool bIsValid;
-	bool bHasBackfaceLighting;
-	bool bHasValidDirectionalAlbedo;
-};
-
-FExposureMaterial GetExposureMaterial(uint2 InPixelPos)
-{
-	FExposureMaterial Out = (FExposureMaterial)0;
-	#if SUBSTRATE_GBUFFER_FORMAT==1
-	{	
-		FSubstrateAddressing SubstrateAddressing = GetSubstratePixelDataByteOffset(InPixelPos, uint2(View.BufferSizeAndInvSize.xy), Substrate.MaxBytesPerPixel);
-		FSubstratePixelHeader SubstratePixelHeader = UnpackSubstrateHeaderIn(Substrate.MaterialTextureArray, SubstrateAddressing, Substrate.TopLayerTexture);
-		FSubstrateTopLayerData TopLayerData = SubstrateUnpackTopLayerData(Substrate.TopLayerTexture.Load(uint3(InPixelPos, 0)));
-		FSubstrateSubsurfaceHeader SSSHeader = SubstrateLoadSubsurfaceHeader(Substrate.MaterialTextureArray, Substrate.FirstSliceStoringSubstrateSSSData, InPixelPos);
-
-		const float DeviceZ = ConvertFromDeviceZ(SceneDepthTexture.Load(int3(InPixelPos, 0)).r);
-		const float3 TranslatedWorldPosition = ReconstructTranslatedWorldPositionFromDeviceZ(InPixelPos, DeviceZ);
-		const float3 V = normalize(View.TranslatedWorldCameraOrigin - TranslatedWorldPosition);
-
-		float3 DiffuseDirectionalAlbedo = 0;
-		float3 DirectionalAlbedo = 0;
-		bool bHasMaterialBackfaceDiffuse = false;
-		{
-			FSubstrateDeferredLighting Out = GetInitialisedSubstrateDeferredLighting();
-			const FSubstrateIntegrationSettings Settings = InitSubstrateIntegrationSettings(false /*bForceFullyRough*/, Substrate.bRoughDiffuse, Substrate.PeelLayersAboveDepth, Substrate.bRoughnessTracking);
-			Substrate_for(uint ClosureIndex = 0, ClosureIndex < SubstratePixelHeader.ClosureCount, ++ClosureIndex)
-			{
-				FSubstrateBSDF BSDF = UnpackSubstrateBSDF(Substrate.MaterialTextureArray, SubstrateAddressing, SubstratePixelHeader);
-				FSubstrateBSDFContext SubstrateBSDFContext = SubstrateCreateBSDFContext(SubstratePixelHeader, BSDF, SubstrateAddressing, V);
-				const float3 BSDFThroughput = LuminanceWeight(SubstrateBSDFContext, BSDF);
-
-				// Evaluate environment lighting
-				FSubstrateEnvLightResult SubstrateEnvLight = SubstrateEvaluateForEnvLight(SubstrateBSDFContext, true /*bEnableSpecular*/, Settings);
-
-				DiffuseDirectionalAlbedo += BSDFThroughput * SubstrateEnvLight.DiffuseWeight;
-				DirectionalAlbedo        += BSDFThroughput * SubstrateEnvLight.DiffuseWeight;
-				DirectionalAlbedo        += BSDFThroughput * SubstrateEnvLight.SpecularWeight;
-
-				const bool bHasBackfaceDiffuse = SubstrateGetBSDFType(BSDF) == SUBSTRATE_BSDF_TYPE_SLAB && BSDF.HasBackScattering();
-				if (bHasBackfaceDiffuse)
-				{
-					DirectionalAlbedo += BSDFThroughput * SubstrateEnvLight.DiffuseBackFaceWeight; // SubstrateEnvLight.DiffuseBackFaceWeight is already divided PI
-				}
-				if (any(SubstrateEnvLight.SpecularHazeWeight > 0.0f))
-				{
-					DirectionalAlbedo += BSDFThroughput * SubstrateEnvLight.SpecularHazeWeight;
-				}
-			}
-		}
-
-		Out.TranslatedWorldPosition = ReconstructTranslatedWorldPositionFromDeviceZ(InPixelPos, DeviceZ);
-		Out.WorldNormal = TopLayerData.WorldNormal;
-		Out.DirectionalAlbedo = DirectionalAlbedo;
-		Out.bIsValid = SubstratePixelHeader.IsSubstrateMaterial();
-		Out.bHasBackfaceLighting = bHasMaterialBackfaceDiffuse;
-		Out.bHasValidDirectionalAlbedo = CalculateEyeAdaptationLuminanceWithoutMin(DiffuseDirectionalAlbedo) > EyeAdaptation_IgnoreMaterialsMinBaseColorLuminance;
-	}
-	#else // SUBSTRATE_ENABLED
-	{
-		const float2 GBufferUV = (InPixelPos + 0.5f) * View.BufferSizeAndInvSize.zw;
-		const FGBufferData GBufferData = GetGBufferDataFromSceneTextures(GBufferUV);
-		const float3 TranslatedWorldPosition = ReconstructTranslatedWorldPositionFromDepth(GBufferUV, GBufferData.Depth);
-
-		// Calculate approximate illuminance from SceneColor and GBufferData
-		// Illuminance ~= (SceneColorLuminance - Emissive) / (DiffuseColor + SubsurfaceColor + EnvBRDF)
-		float3 Denominator = 0.f;
-		Denominator += GBufferData.DiffuseColor;
-		if (GBufferData.ShadingModelID == SHADINGMODELID_SUBSURFACE
-			|| GBufferData.ShadingModelID == SHADINGMODELID_PREINTEGRATED_SKIN
-			|| GBufferData.ShadingModelID == SHADINGMODELID_TWOSIDED_FOLIAGE
-			|| GBufferData.ShadingModelID == SHADINGMODELID_CLOTH
-			|| GBufferData.ShadingModelID == SHADINGMODELID_EYE)
-		{
-			float3 SubsurfaceColor = ExtractSubsurfaceColor(GBufferData);
-
-			if (GBufferData.ShadingModelID == SHADINGMODELID_TWOSIDED_FOLIAGE)
-			{
-				SubsurfaceColor *= 1.0f / PI;
-			}
-
-			Denominator += SubsurfaceColor;
-		}
-
-		{
-			const float3 CameraVector = normalize(View.TranslatedWorldCameraOrigin - TranslatedWorldPosition);
-
-			const float3 N = GBufferData.WorldNormal;
-			const float3 V = CameraVector;
-			const float3 EnvBrdf = EnvBRDF(GBufferData.SpecularColor, GBufferData.Roughness, max(0.0, dot(N, V)));
-
-			Denominator += EnvBrdf;
-		}
-
-		Out.TranslatedWorldPosition = TranslatedWorldPosition;
-		Out.WorldNormal = GBufferData.WorldNormal;
-		Out.DirectionalAlbedo = Denominator;
-		Out.bIsValid = GBufferData.ShadingModelID != SHADINGMODELID_UNLIT;
-		Out.bHasBackfaceLighting = GetShadingModelRequiresBackfaceLighting(GBufferData.ShadingModelID);
-		Out.bHasValidDirectionalAlbedo = CalculateEyeAdaptationLuminanceWithoutMin(GBufferData.DiffuseColor) > EyeAdaptation_IgnoreMaterialsMinBaseColorLuminance;
-	}
-	#endif // SUBSTRATE_ENABLED
-	return Out;
-}
-
 [numthreads(8, 8, 1)]
 void CalculateExposureIlluminance(uint2 ThreadId : SV_DispatchThreadID)
 {
diff --git a/Engine/Source/Runtime/Renderer/Private/PostProcess/PostProcessHistogram.cpp b/Engine/Source/Runtime/Renderer/Private/PostProcess/PostProcessHistogram.cpp
--- a/Engine/Source/Runtime/Renderer/Private/PostProcess/PostProcessHistogram.cpp
+++ b/Engine/Source/Runtime/Renderer/Private/PostProcess/PostProcessHistogram.cpp
@@ -37,6 +37,7 @@ public:
 		SHADER_PARAMETER_STRUCT_REF(FViewUniformShaderParameters, View)
 		SHADER_PARAMETER_STRUCT(FScreenPassTextureViewportParameters, Input)
 		SHADER_PARAMETER_STRUCT(FEyeAdaptationParameters, EyeAdaptation)
+		SHADER_PARAMETER_RDG_UNIFORM_BUFFER(FSubstrateGlobalUniformParameters, Substrate)
 		SHADER_PARAMETER_SAMPLER(SamplerState, InputSampler)
 		SHADER_PARAMETER_RDG_TEXTURE_SRV(Texture2D, InputTexture)
 		SHADER_PARAMETER_RDG_TEXTURE_UAV(RWTexture2D<float4>, HistogramRWTexture)
@@ -161,6 +162,7 @@ public:
 
 	BEGIN_SHADER_PARAMETER_STRUCT(FParameters, )
 		SHADER_PARAMETER_STRUCT_REF(FViewUniformShaderParameters, View)
+		SHADER_PARAMETER_RDG_UNIFORM_BUFFER(FSubstrateGlobalUniformParameters, Substrate)
 		SHADER_PARAMETER_STRUCT_INCLUDE(FSceneTextureParameters, SceneTextures)
 		SHADER_PARAMETER_STRUCT(FScreenPassTextureViewportParameters, Input)
 		SHADER_PARAMETER_STRUCT(FEyeAdaptationParameters, EyeAdaptation)
@@ -274,6 +276,7 @@ static FRDGTextureRef AddHistogramLegacyPass(
 		PassParameters->HistogramRWTexture = GraphBuilder.CreateUAV(HistogramTexture);
 		PassParameters->ThreadGroupCount = HistogramThreadGroupCount;
 		PassParameters->EyeAdaptation = EyeAdaptationParameters;
+		PassParameters->Substrate = Substrate::BindSubstrateGlobalUniformParameters(View);
 
 		{
 			const FRDGTextureDesc Desc = FRDGTextureDesc::Create2D(SceneColor.ViewRect.Size(), PF_R16F, FClearValueBinding::Black, TexCreate_UAV | TexCreate_ShaderResource);
@@ -380,6 +383,7 @@ static FRDGTextureRef AddHistogramAtomicPass(
 		PassParameters->EyeAdaptation = EyeAdaptationParameters;
 		// PassParameters->HistogramScatter64Output = GraphBuilder->CreateUAV(HistogramScatter64Texture);
 		PassParameters->HistogramScatter32Output = GraphBuilder.CreateUAV(HistogramScatter32Texture);
+		PassParameters->Substrate = Substrate::BindSubstrateGlobalUniformParameters(View);
 
 		{
 			const FRDGTextureDesc Desc = FRDGTextureDesc::Create2D(SceneColor.ViewRect.Size(), PF_R16F, FClearValueBinding::Black, TexCreate_UAV | TexCreate_ShaderResource);
@@ -517,6 +521,7 @@ FRDGTextureRef AddLocalExposurePass(
 		PassParameters->InputTexture = SceneColor.TextureSRV;
 		PassParameters->BilateralGridRWTexture = GraphBuilder.CreateUAV(LocalExposureTexture);
 		PassParameters->ThreadGroupCount = ThreadGroupCount;
+		PassParameters->Substrate = Substrate::BindSubstrateGlobalUniformParameters(View);
 
 		FHistogramCS::FPermutationDomain PermutationVector;
 		PermutationVector.Set<FHistogramCS::FBilateralGrid>(true);
