Fix DBuffer rendering with Substrate on forward mobile.

diff --git a/Engine/Shaders/Private/DeferredDecal.usf b/Engine/Shaders/Private/DeferredDecal.usf
--- a/Engine/Shaders/Private/DeferredDecal.usf
+++ b/Engine/Shaders/Private/DeferredDecal.usf
@@ -193,7 +193,7 @@ void FPixelShaderInOut_MainPS(inout FPixelShaderIn In, inout FPixelShaderOut Out
 
 	float3 EmissiveColor = 0.0f.xxx;
 	float Opacity = 0.0f;
-	FGBufferData DBufferData;
+	FGBufferData DBufferData = (FGBufferData)0;
 	if (SubstrateHeader.SubstrateTree.BSDFCount > 0)
 	{
 		const FSubstrateIntegrationSettings Settings = InitSubstrateIntegrationSettings(false /*bForceFullyRough*/, false /*bRoughDiffuseEnabled*/, 0, false/*bRoughnessTracking*/);
@@ -215,7 +215,6 @@ void FPixelShaderInOut_MainPS(inout FPixelShaderIn In, inout FPixelShaderOut Out
 		Opacity										= CoverageOverBackground;
 		EmissiveColor								= Export.EmissiveLuminance;
 
-		FGBufferData DBufferData					= (FGBufferData)0;
 		DBufferData.WorldNormal						= Export.WorldNormal;
 		DBufferData.BaseColor						= Export.BaseColor;
 		DBufferData.Metallic						= Export.Metallic;
