[GPU Scene] Removed redundant value PersistentPrimitiveId from the buffer and initialize using the argument to the GetPrimitiveData funciton (which is always the same). This is a small optimization as it avoids some current and future dependent reads.

#rnx

diff --git a/Engine/Shaders/Private/SceneData.ush b/Engine/Shaders/Private/SceneData.ush
--- a/Engine/Shaders/Private/SceneData.ush
+++ b/Engine/Shaders/Private/SceneData.ush
@@ -383,7 +383,11 @@ FPrimitiveSceneData GetPrimitiveData(uint PrimitiveId)
 	return (FPrimitiveSceneData)0;
 #else
 	FPrimitiveSceneData PrimitiveData = (FPrimitiveSceneData)0;
-	// Note: layout must match FPrimitiveSceneShaderData in C++
+
+	// short-circuit to avoid redundant indirection (as this is always the same value it is no longer stored in the buffer)
+	PrimitiveData.PersistentPrimitiveIndex = PrimitiveId;
+
+	// Note: layout must match FPrimitiveSceneShaderData::Setup in C++
 	// Relying on optimizer to remove unused loads
 	uint PrimitiveIndex = PrimitiveId * PRIMITIVE_SCENE_DATA_STRIDE;
 	
@@ -476,7 +480,7 @@ FPrimitiveSceneData GetPrimitiveData(uint PrimitiveId)
 	PrimitiveData.WireframeColor						= UnpackColorRGB24(LoadPrimitivePrimitiveSceneDataElement(PrimitiveIndex, 30).x);
 	PrimitiveData.PrimitiveColor						= UnpackColorRGB24(LoadPrimitivePrimitiveSceneDataElement(PrimitiveIndex, 30).y);
 	PrimitiveData.PackedNaniteFlags						= asuint(LoadPrimitivePrimitiveSceneDataElement(PrimitiveIndex, 30).z);
-	PrimitiveData.PersistentPrimitiveIndex				= asint(LoadPrimitivePrimitiveSceneDataElement(PrimitiveIndex, 30).w);
+	// 30.w is unused
 	PrimitiveData.InstanceDrawDistanceMinMaxSquared		= LoadPrimitivePrimitiveSceneDataElement(PrimitiveIndex, 31).xy;
 	PrimitiveData.InstanceWPODisableDistanceSquared		= LoadPrimitivePrimitiveSceneDataElement(PrimitiveIndex, 31).z;
 	PrimitiveData.NaniteRayTracingDataOffset			= asuint(LoadPrimitivePrimitiveSceneDataElement(PrimitiveIndex, 31).w);
diff --git a/Engine/Source/Runtime/Engine/Private/PrimitiveUniformShaderParameters.cpp b/Engine/Source/Runtime/Engine/Private/PrimitiveUniformShaderParameters.cpp
--- a/Engine/Source/Runtime/Engine/Private/PrimitiveUniformShaderParameters.cpp
+++ b/Engine/Source/Runtime/Engine/Private/PrimitiveUniformShaderParameters.cpp
@@ -168,7 +168,7 @@ void FPrimitiveSceneShaderData::Setup(const FPrimitiveUniformShaderParameters& P
 		PrimitiveUniformShaderParameters.WireframeAndPrimitiveColor.X, 
 		PrimitiveUniformShaderParameters.WireframeAndPrimitiveColor.Y, 
 		PrimitiveUniformShaderParameters.PackedNaniteFlags, 
-		uint32(PrimitiveUniformShaderParameters.PersistentPrimitiveIndex));
+		0u /*unused*/ );
 
 	Store4(OutData, 31, 
 		PrimitiveUniformShaderParameters.InstanceDrawDistanceMinMaxSquared.X, 
@@ -192,8 +192,8 @@ void FPrimitiveSceneShaderData::Setup(const FPrimitiveUniformShaderParameters& P
 		PrimitiveUniformShaderParameters.MaterialCacheDescriptor,
 		// Note: AnimationMinScreenSize could be packed to fewer bits as it needs to represent values in the range 0-1, and special case -1.
 		PrimitiveUniformShaderParameters.AnimationMinScreenSize,
-		0u,
-		0u);
+		0u /*unused*/,
+		0u /*unused*/);
 
 	// Set all the custom primitive data float4. This matches the loop in SceneData.ush
 	const int32 CustomPrimitiveDataStartIndex = 35;
