Fix tiled SSR with split screen
If we were unlucky, an SSR tiled would be marked as valid, and the content would contain NaNs. This also fixes the debug mode r.SSR.TiledComposite.Visualize 1 in splitscreen

#rnx

diff --git a/Engine/Shaders/Private/ReflectionEnvironmentPixelShader.usf b/Engine/Shaders/Private/ReflectionEnvironmentPixelShader.usf
--- a/Engine/Shaders/Private/ReflectionEnvironmentPixelShader.usf
+++ b/Engine/Shaders/Private/ReflectionEnvironmentPixelShader.usf
@@ -180,7 +180,7 @@ float3 ReflectionEnvironment(FGBufferData GBuffer, float AmbientOcclusion, float
 #elif REFLECTION_SOURCE_TYPE == REFLECTION_SOURCE_SSR
 	float4 ReflectionInput = Texture2DSample(ReflectionTexture, ReflectionTextureSampler, BufferUV);
 #else // REFLECTION_SOURCE_TILED_SSR
-	const uint2 PixelPos = uint2(SvPosition.xy);
+	const uint2 PixelPos = uint2(SvPosition.xy - View.ViewRectMin.xy);
 	bool bSSRTileUsed = IsSSRTileUsed(PixelPos, SSRTiledViewRes, SSRTileMaskBuffer);
 	float4 ReflectionInput = 0;
 
