Substrate: Tweak the logic for legacy conversion of cloth materials.

There is a difference between ClothBxDF and the new substrate slab structure. The previous logic was to blend the specular lobe into the cloth lobe using the cloth amount.

In substrate, the fuzz layer sits on top of the specular layer. This means that out of the box, one gets a specular response _under_ the fuzz which is more realistic but does not match the previous behavior.

This small change fades out the specular lobe to try and emulate the legacy behavior a bit better. While this is not a perfect match, this avoids the specular lobe from being visible when cloth amount is set to 1.0.

Additionaly, because the legacy cloth BRDF has a very different shape to the new fuzz model, remap the roughness range during legacy conversions to [0.5, 1.0] to avoid the low-roughness range which looks most visually different.

diff --git a/Engine/Shaders/Private/Substrate/SubstrateLegacyConversion.ush b/Engine/Shaders/Private/Substrate/SubstrateLegacyConversion.ush
--- a/Engine/Shaders/Private/Substrate/SubstrateLegacyConversion.ush
+++ b/Engine/Shaders/Private/Substrate/SubstrateLegacyConversion.ush
@@ -104,7 +104,7 @@ FSubstrateData CreateLegacySlab(
 	float3 DiffuseAlbedo = In.DiffuseAlbedo;
 	float3 F0 = In.F0;
 	float3 F90 = In.F90;
-	if (In.UseMetalness > 0.0f)
+	if (In.UseMetalness)
 	{
 		DiffuseAlbedo = ComputeDiffuseAlbedo(In.BaseColor, In.Metallic);
 		F0 = ComputeF0(In.Specular, In.BaseColor, In.Metallic);
@@ -133,7 +133,7 @@ FSubstrateData CreateLegacySlab(
 		In.ClearCoatBottomNormal,							// ClearCoatBottomNormal;
 		In.FuzzAmount,										// FuzzAmount
 		In.FuzzColor,										// FuzzColor
-		In.Roughness,										// FuzzRoughness
+		0.5 + 0.5 * In.Roughness,							// FuzzRoughness (adjust range because the legacy model has a very different response)
 		1.0f, (GLINT_UV_TYPE)0,								// GlintData
 		0.0f,												// Specular profile
 		In.Thickness,
@@ -430,9 +430,12 @@ FSubstrateData SubstrateConvertLegacyMaterialDynamic(
 	#if MATERIAL_SHADINGMODEL_CLOTH
 	if (ShadingModel == SHADINGMODELID_CLOTH)
 	{
-		LegacySlab.BaseColor = BaseColor;
-		LegacySlab.Specular = Specular;
-		LegacySlab.Metallic = Metallic;
+		LegacySlab.UseMetalness = false;
+		LegacySlab.DiffuseAlbedo = ComputeDiffuseAlbedo(BaseColor, Metallic);
+		LegacySlab.F0 = ComputeF0(Specular, BaseColor, Metallic);
+		LegacySlab.F90 = 1;
+		// Mask by cloth amount to emulate the old behavior where the cloth lobe gradually replaces the spec lobe depending on FuzzAmount
+		LegacySlab.F0 *= 1 - ClearCoat;
 		LegacySlab.Roughness = Roughness;
 		LegacySlab.Anisotropy = Anisotropy;
 		LegacySlab.SSSProfileID = 0.f;
