Substrate - fixed performance of aniso pass when GBufferFormat=0.
- Removed redundant BSDFCount check in SubstrateExport.
- Avoid too many normal/tangent normalisation. Now SharedLocalBases is normalised when written to.
- Added a fast path legacy opaque materials which always result in 1 BSDF with vull coverage (no need to run the tree for visibility testing).
- Added SHADINGMODEL filtering of exported data to reduce base pass shader size and complexity.
- Moved Specular Profiles baking into the "visible BSDF block".

#tests PC editor dx12 and console
#rnx

#DisabledValidations Preflight

diff --git a/Engine/Shaders/Private/AnisotropyPassShader.usf b/Engine/Shaders/Private/AnisotropyPassShader.usf
--- a/Engine/Shaders/Private/AnisotropyPassShader.usf
+++ b/Engine/Shaders/Private/AnisotropyPassShader.usf
@@ -148,8 +148,9 @@ void MainPixelShader(
 	float Anisotropy = Export.Anisotropy;
 
 	// Similar to CalculateAnisotropyTangent. SUBSTRATE_TODO: merge the code and account for bottom normal.
-	float3 WorldTangent = normalize(Export.WorldTangent);
-	float3 WorldNormal = normalize(Export.WorldNormal);
+	// Exported tangent and normal are already normalized (from storage in SharedLocalNormal).
+	float3 WorldTangent = Export.WorldTangent;
+	float3 WorldNormal = Export.WorldNormal;
 	float3 WorldBiTangent = cross(WorldNormal, WorldTangent);
 	WorldTangent = normalize(cross(WorldBiTangent, WorldNormal));
 
diff --git a/Engine/Shaders/Private/Substrate/Substrate.ush b/Engine/Shaders/Private/Substrate/Substrate.ush
--- a/Engine/Shaders/Private/Substrate/Substrate.ush
+++ b/Engine/Shaders/Private/Substrate/Substrate.ush
@@ -3350,6 +3350,11 @@ float3 SanitizeSSSMFP(float3 InMFP)
 	return clamp(InMFP, 0.0f, Max111110BitsFloat3);
 }
 
+float SanitizeAnisotropy(float Anisotropy)
+{
+	return clamp(Anisotropy, -1.f, 1.f);
+}
+
 void FSubstrateBSDF::SubstrateSanitizeBSDF()
 {
 	switch (BSDF_GETTYPE(this))
@@ -3370,7 +3375,7 @@ void FSubstrateBSDF::SubstrateSanitizeBSDF()
 				SLAB_HAZINESS(this) = PackHaziness(LocalHaziness);
 			}
 			SLAB_ROUGHNESS(this) = SanitizeRoughness(SLAB_ROUGHNESS(this)); // 
-			SLAB_ANISOTROPY(this) = clamp(SLAB_ANISOTROPY(this), -1.f, 1.f);
+			SLAB_ANISOTROPY(this) = SanitizeAnisotropy(SLAB_ANISOTROPY(this));
 			// SLAB_HAZINESS is sanitize when packed, see PackHaziness.
 		#if SUBSTRATE_INLINE_SHADING && SUBSTRATE_OPAQUE_MATERIAL
 			// We only clamp MFP when we are going to export it packed into the GBuffer.
diff --git a/Engine/Shaders/Private/Substrate/SubstrateExport.ush b/Engine/Shaders/Private/Substrate/SubstrateExport.ush
--- a/Engine/Shaders/Private/Substrate/SubstrateExport.ush
+++ b/Engine/Shaders/Private/Substrate/SubstrateExport.ush
@@ -915,13 +915,24 @@ void PackSubstrateOut(
 
 #if SUBSTRATE_ENABLED && (SUBSTRATE_MATERIAL_EXPORT_TYPE > 0 || SUBSTRATE_GPU_LIGHTMASS > 0 || SUBSTRATE_GBUFFER_FORMAT==0 || defined(SUBSTRATE_MATERIAL_EXPORT_REQUESTED))
 
+#ifndef MATERIAL_IS_SUBSTRATEHIDDENCONVERSION
+#define MATERIAL_IS_SUBSTRATEHIDDENCONVERSION 0
+#endif
+#ifndef SUBSTRATE_OPAQUE_DEFERRED
+#define SUBSTRATE_OPAQUE_DEFERRED 0
+#endif
+
+// Hidden conversion from root node for opaque material guarantee that there is a BSDF output with Coverage=1 and Transmittance=0 that is always visible.
+// SUBSTRATE_TODO a simple more general way to make this work would be to make sure there is no CoverageWeight node in the graph instead of MATERIAL_IS_SUBSTRATEHIDDENCONVERSION? In this case that optimisation would be valid too.
+#define SUBTRATE_OPTIMISED_LEGACY_OPAQUE_EXPORT (SUBSTRATE_OPAQUE_DEFERRED && MATERIAL_IS_SUBSTRATEHIDDENCONVERSION)
+
 struct FExportResult
 {
 	float  Coverage;
 	float3 TransmittancePreCoverage;
 	float3 BaseColorPostCoverage;
 	float3 BaseColor;
-	float3 WorldNormal;
+	float3 WorldNormal;					// Already normalized
 	float3 EmissiveLuminance;
 	float  Specular;
 	float  Roughness;
@@ -931,7 +942,7 @@ struct FExportResult
 	float3 SubsurfaceColor;
 	float  SubsurfaceProfileIDFloat;
 	float  SubsurfaceProfileRadiusScale;
-	float3 WorldTangent;
+	float3 WorldTangent;				// Already normalized
 	float  ShadingModelID;
 	float IndirectIrradianceScale;
 	float Curvature;
@@ -952,16 +963,34 @@ FExportResult SubstrateMaterialExportOut(
 	Out.IndirectIrradianceScale = 1.0;
 	Out.SubsurfaceProfileRadiusScale = 1.0;
 
+	const int BSDFIdx = 0;
+	#define CurrentBSDF SubstratePixelHeader.SubstrateTree.BSDFs[BSDFIdx]
+
+#if SUBTRATE_OPTIMISED_LEGACY_OPAQUE_EXPORT
+	{
+		// In this case, there is always a BSDF and it is always 100% visible.
+		// Sanitize BSDF before it is exported.
+		CurrentBSDF.SubstrateSanitizeBSDF();
+		{
+#else
 	if (SubstratePixelHeader.SubstrateTree.BSDFCount > 0)
 	{
-		// Update tree (coverage/transmittance/luminace weights)
+
+		// Update tree (coverage/transmittance/luminace weights) and Sanitize its content.
 		SubstratePixelHeader.SubstrateUpdateTree(SubstrateData, V, Settings, Out.Coverage, Out.TransmittancePreCoverage);
 
-		const int BSDFIdx = 0;
-		if(SubstratePixelHeader.SubstrateTree.BSDFCount > 0)
+		if (SubstrateIsBSDFVisible(CurrentBSDF))
 		{
-		#define CurrentBSDF SubstratePixelHeader.SubstrateTree.BSDFs[BSDFIdx]
+#endif
+			const uint BSDFType = BSDF_GETTYPE(CurrentBSDF);
+			const uint SSSType  = BSDF_GETSSSTYPE(CurrentBSDF);
+
+			const float3x3 TangentBasis = SubstrateGetBSDFSharedBasis_InlineShading(SubstratePixelHeader, CurrentBSDF);
+			const float3 N = TangentBasis[2];
+			const float3 T = TangentBasis[0];
 
+			const float GreyScaleWeight = dot(CurrentBSDF.LuminanceWeightV, (1.0 / 3.0).xxx);
+			
 			// Bake the Specular Profiles into F0/F90 of BSDFs to blend in order to retain some of the original look when using parameter blending or fully simplified materials.
 		#if SUBSTRATE_SPECPROFILE_ENABLED && MATERIAL_USES_SPECULAR_PROFILE
 			const float NoV = saturate(dot(SurfaceWorldNormal, V));
@@ -969,167 +998,173 @@ FExportResult SubstrateMaterialExportOut(
 			BSDF_SETHASSPECPROFILE(CurrentBSDF, 0);
 		#endif
 
-			if (SubstrateIsBSDFVisible(CurrentBSDF))
-			{
-				const uint BSDFType = BSDF_GETTYPE(CurrentBSDF);
-				const uint SSSType  = BSDF_GETSSSTYPE(CurrentBSDF);
-
-				const float3x3 TangentBasis = SubstrateGetBSDFSharedBasis_InlineShading(SubstratePixelHeader, CurrentBSDF);
-				const float3 N = TangentBasis[2];
-				const float3 T = TangentBasis[0];
-
-				const float GreyScaleWeight = dot(CurrentBSDF.LuminanceWeightV, (1.0 / 3.0).xxx);
-
-				Out.BaseColorPostCoverage	= SubstrateGetBSDFBaseColor(CurrentBSDF)	* CurrentBSDF.LuminanceWeightV;
-				Out.BaseColor				= SubstrateGetBSDFBaseColor(CurrentBSDF);
-				Out.WorldNormal				= N;
-				Out.EmissiveLuminance		= BSDF_GETEMISSIVE(CurrentBSDF)				* CurrentBSDF.LuminanceWeightV;
-				Out.Specular				= SubstrateGetBSDFSpecular(CurrentBSDF);
-				Out.Roughness				= SubstrateGetBSDFRoughness(CurrentBSDF);
-				Out.Anisotropy				= SubstrateGetBSDFAnisotropy(CurrentBSDF);
-				Out.Metallic				= SubstrateGetBSDFMetallic(CurrentBSDF);
-
-				Out.SubsurfaceColor			= SubstrateGetBSDFSubSurfaceColor(CurrentBSDF);
-				Out.SubsurfaceProfileIDFloat= ConvertSubsurfaceProfileToFloat(SubstrateGetBSDFSubSurfaceProfileID(CurrentBSDF));
-				Out.WorldTangent			= T;
-				Out.ShadingModelID			= SubstrateGetLegacyShadingModels(CurrentBSDF);
-
-				if (BSDFType == SUBSTRATE_BSDF_TYPE_HAIR)
-				{
-					Out.Metallic	= HAIR_SCATTER(CurrentBSDF);
-					Out.CustomData  = float4(0.0, 0.0, HAIR_BACKLIT(CurrentBSDF), 0.0);
-				}
-				else if (BSDFType == SUBSTRATE_BSDF_TYPE_EYE)
-				{
-					const float3 IrisNormal		= EYE_IRISNORMAL(CurrentBSDF);
-					const float3 IrisPlaneNormal= EYE_IRISPLANENORMAL(CurrentBSDF);
-					const float IrisMask		= saturate(EYE_IRISMASK(CurrentBSDF));
-					const float IrisDistance	= saturate(EYE_IRISDISTANCE(CurrentBSDF));
-
-					Out.CustomData.x = EncodeSubsurfaceProfile(Out.SubsurfaceProfileIDFloat).x;
-					Out.CustomData.w = 1.0f - IrisMask;	// Opacity
-
-					const float3 GBufferWorldNormal = N;
-				#if IRIS_NORMAL
-					float2 WorldNormalOct = UnitVectorToOctahedron(GBufferWorldNormal);
-
-					// CausticNormal stored as octahedron
-					#if NUM_MATERIAL_OUTPUTS_GETTANGENTOUTPUT > 0
-						// Blend in the negative intersection normal to create some concavity
-						// Not great as it ties the concavity to the convexity of the cornea surface
-						// No good justification for that. On the other hand, if we're just looking to
-						// introduce some concavity, this does the job.
-						float3 PlaneNormal = normalize( T /*GetTangentOutput0(MaterialParameters)*/);
-						float3 CausticNormal = normalize( lerp( PlaneNormal, -GBufferWorldNormal, IrisMask*IrisDistance ) );
-						float2 CausticNormalOct  = UnitVectorToOctahedron( CausticNormal );
-						float2 CausticNormalDelta = ( CausticNormalOct - WorldNormalOct ) * 0.5 + (128.0/255.0);
-						Out.Metallic = CausticNormalDelta.x;
-						Out.Specular = CausticNormalDelta.y;
-					#else
-						float3 PlaneNormal = GBufferWorldNormal;
-						Out.Metallic = 128.0/255.0;
-						Out.Specular = 128.0/255.0;
-					#endif
+			Out.BaseColorPostCoverage	= SubstrateGetBSDFBaseColor(CurrentBSDF)	* CurrentBSDF.LuminanceWeightV;
+			Out.BaseColor				= SubstrateGetBSDFBaseColor(CurrentBSDF);
+			Out.WorldNormal				= N;
+			Out.EmissiveLuminance		= BSDF_GETEMISSIVE(CurrentBSDF)				* CurrentBSDF.LuminanceWeightV;
+			Out.Specular				= SubstrateGetBSDFSpecular(CurrentBSDF);
+			Out.Roughness				= SubstrateGetBSDFRoughness(CurrentBSDF);
+			Out.Anisotropy				= SubstrateGetBSDFAnisotropy(CurrentBSDF);
+			Out.Metallic				= SubstrateGetBSDFMetallic(CurrentBSDF);
 
-					// IrisNormal CustomData.yz
-					//#if NUM_MATERIAL_OUTPUTS_CLEARCOATBOTTOMNORMAL > 0
-					//	float3 IrisNormal = normalize( ClearCoatBottomNormal0(MaterialParameters) );
-					//	#if MATERIAL_TANGENTSPACENORMAL
-					//		IrisNormal = normalize( TransformTangentVectorToWorld( MaterialParameters.TangentToWorld, IrisNormal ) );
-					//	#endif
-					//#else
-					//	float3 IrisNormal = PlaneNormal;
-					//#endif
-
-					float2 IrisNormalOct  = UnitVectorToOctahedron( IrisNormal );
-					float2 IrisNormalDelta = ( IrisNormalOct - WorldNormalOct ) * 0.5 + (128.0/255.0);
-					Out.CustomData.yz = IrisNormalDelta;
-				#else
-					Out.Metallic = IrisDistance;
+			Out.SubsurfaceColor			= SubstrateGetBSDFSubSurfaceColor(CurrentBSDF);
+			Out.SubsurfaceProfileIDFloat= ConvertSubsurfaceProfileToFloat(SubstrateGetBSDFSubSurfaceProfileID(CurrentBSDF));
+			Out.WorldTangent			= T;
+			Out.ShadingModelID			= SubstrateGetLegacyShadingModels(CurrentBSDF);
 
-					#if NUM_MATERIAL_OUTPUTS_GETTANGENTOUTPUT > 0
-						float3 Tangent = IrisNormal;
-						Out.CustomData.yz = UnitVectorToOctahedron( normalize(Tangent) ) * 0.5 + 0.5;
-					#endif
+
+#if MATERIAL_SHADINGMODEL_HAIR || !MATERIAL_IS_SUBSTRATEHIDDENCONVERSION // Hair shading model or if not a hidden conversion, then we cannot fully trust the shading model from simplified substrate nodes.
+			if (BSDFType == SUBSTRATE_BSDF_TYPE_HAIR)
+			{
+				Out.Metallic	= HAIR_SCATTER(CurrentBSDF);
+				Out.CustomData  = float4(0.0, 0.0, HAIR_BACKLIT(CurrentBSDF), 0.0);
+			}
+			else 
+#endif
+#if MATERIAL_SHADINGMODEL_EYE || !MATERIAL_IS_SUBSTRATEHIDDENCONVERSION
+				if (BSDFType == SUBSTRATE_BSDF_TYPE_EYE)
+			{
+				const float3 IrisNormal		= EYE_IRISNORMAL(CurrentBSDF);
+				const float3 IrisPlaneNormal= EYE_IRISPLANENORMAL(CurrentBSDF);
+				const float IrisMask		= saturate(EYE_IRISMASK(CurrentBSDF));
+				const float IrisDistance	= saturate(EYE_IRISDISTANCE(CurrentBSDF));
+
+				Out.CustomData.x = EncodeSubsurfaceProfile(Out.SubsurfaceProfileIDFloat).x;
+				Out.CustomData.w = 1.0f - IrisMask;	// Opacity
+
+				const float3 GBufferWorldNormal = N;
+			#if IRIS_NORMAL
+				float2 WorldNormalOct = UnitVectorToOctahedron(GBufferWorldNormal);
+
+				// CausticNormal stored as octahedron
+				#if NUM_MATERIAL_OUTPUTS_GETTANGENTOUTPUT > 0
+					// Blend in the negative intersection normal to create some concavity
+					// Not great as it ties the concavity to the convexity of the cornea surface
+					// No good justification for that. On the other hand, if we're just looking to
+					// introduce some concavity, this does the job.
+					float3 PlaneNormal = normalize( T /*GetTangentOutput0(MaterialParameters)*/);
+					float3 CausticNormal = normalize( lerp( PlaneNormal, -GBufferWorldNormal, IrisMask*IrisDistance ) );
+					float2 CausticNormalOct  = UnitVectorToOctahedron( CausticNormal );
+					float2 CausticNormalDelta = ( CausticNormalOct - WorldNormalOct ) * 0.5 + (128.0/255.0);
+					Out.Metallic = CausticNormalDelta.x;
+					Out.Specular = CausticNormalDelta.y;
+				#else
+					float3 PlaneNormal = GBufferWorldNormal;
+					Out.Metallic = 128.0/255.0;
+					Out.Specular = 128.0/255.0;
 				#endif
 
-				#if SHADING_PATH_MOBILE
-					#if MATERIAL_SHADINGMODEL_EYE_USE_CURVATURE
-						Out.Curvature = Metallic;
-					#else
-						Out.Curvature = CalculateCurvature(GBufferWorldNormal, WorldPosition_CamRelative);
-					#endif
+				// IrisNormal CustomData.yz
+				//#if NUM_MATERIAL_OUTPUTS_CLEARCOATBOTTOMNORMAL > 0
+				//	float3 IrisNormal = normalize( ClearCoatBottomNormal0(MaterialParameters) );
+				//	#if MATERIAL_TANGENTSPACENORMAL
+				//		IrisNormal = normalize( TransformTangentVectorToWorld( MaterialParameters.TangentToWorld, IrisNormal ) );
+				//	#endif
+				//#else
+				//	float3 IrisNormal = PlaneNormal;
+				//#endif
+
+				float2 IrisNormalOct  = UnitVectorToOctahedron( IrisNormal );
+				float2 IrisNormalDelta = ( IrisNormalOct - WorldNormalOct ) * 0.5 + (128.0/255.0);
+				Out.CustomData.yz = IrisNormalDelta;
+			#else
+				Out.Metallic = IrisDistance;
+
+				#if NUM_MATERIAL_OUTPUTS_GETTANGENTOUTPUT > 0
+					float3 Tangent = IrisNormal;
+					Out.CustomData.yz = UnitVectorToOctahedron( normalize(Tangent) ) * 0.5 + 0.5;
+				#endif
+			#endif
 
-						Out.Curvature = clamp(Out.Curvature, 0.001f, 1.0f);
+			#if SHADING_PATH_MOBILE
+				#if MATERIAL_SHADINGMODEL_EYE_USE_CURVATURE
+					Out.Curvature = Metallic;
+				#else
+					Out.Curvature = CalculateCurvature(GBufferWorldNormal, WorldPosition_CamRelative);
 				#endif
-				}
-				else if (SSSType == SSS_TYPE_DIFFUSION_PROFILE)  // Thin and not thin, Subsurface and TwoSided Foliage
-				{
-					// Roughness has been modify when calling GetSubstrateSlabBSDF(), but expending SSS Profile roughness
-					// Revert roughness to its original average value to get correct result
-					Out.Roughness = GetSubsurfaceProfileOriginalRoughness(SLAB_SSSPROFILEID(CurrentBSDF), Out.Roughness, SLAB_SSSPROFILERADIUSSCALE(CurrentBSDF));
 
-					Out.SubsurfaceProfileRadiusScale = SLAB_SSSPROFILERADIUSSCALE(CurrentBSDF);
-					Out.CustomData.rgb = EncodeSubsurfaceProfile(Out.SubsurfaceProfileIDFloat);
-					Out.CustomData.a   = Out.SubsurfaceProfileRadiusScale;
+					Out.Curvature = clamp(Out.Curvature, 0.001f, 1.0f);
+			#endif
+			}
+			else
+#endif
+#if MATERIAL_SHADINGMODEL_SUBSURFACE_PROFILE || !MATERIAL_IS_SUBSTRATEHIDDENCONVERSION
+				if (SSSType == SSS_TYPE_DIFFUSION_PROFILE)  // Thin and not thin, Subsurface and TwoSided Foliage
+			{
+				// Roughness has been modify when calling GetSubstrateSlabBSDF(), but expending SSS Profile roughness
+				// Revert roughness to its original average value to get correct result
+				Out.Roughness = GetSubsurfaceProfileOriginalRoughness(SLAB_SSSPROFILEID(CurrentBSDF), Out.Roughness, SLAB_SSSPROFILERADIUSSCALE(CurrentBSDF));
 
-					const float3 GBufferWorldNormal = N;
-					const float Opacity = SLAB_SSSPROFILERADIUSSCALE(CurrentBSDF);
-					
-					// Optimization: if opacity is 0 then revert to default shading model 
-				#if SUBSURFACE_PROFILE_OPACITY_THRESHOLD
-					if (Opacity > SSSS_OPACITY_THRESHOLD_EPS)
-				#endif
-					{
-					#if SHADING_PATH_MOBILE
-						#if MATERIAL_SUBSURFACE_PROFILE_USE_CURVATURE
-						Out.Curvature = MobileShadingPathCurvature;
-						#else
-						Out.Curvature = CalculateCurvature(GBufferWorldNormal, WorldPosition_CamRelative);
-						#endif
+				Out.SubsurfaceProfileRadiusScale = SLAB_SSSPROFILERADIUSSCALE(CurrentBSDF);
+				Out.CustomData.rgb = EncodeSubsurfaceProfile(Out.SubsurfaceProfileIDFloat);
+				Out.CustomData.a   = Out.SubsurfaceProfileRadiusScale;
 
-						Out.Curvature = clamp(Out.Curvature, 0.001f, 1.0f);
+				const float3 GBufferWorldNormal = N;
+				const float Opacity = SLAB_SSSPROFILERADIUSSCALE(CurrentBSDF);
+					
+				// Optimization: if opacity is 0 then revert to default shading model 
+			#if SUBSURFACE_PROFILE_OPACITY_THRESHOLD
+				if (Opacity > SSSS_OPACITY_THRESHOLD_EPS)
+			#endif
+				{
+				#if SHADING_PATH_MOBILE
+					#if MATERIAL_SUBSURFACE_PROFILE_USE_CURVATURE
+					Out.Curvature = MobileShadingPathCurvature;
+					#else
+					Out.Curvature = CalculateCurvature(GBufferWorldNormal, WorldPosition_CamRelative);
 					#endif
-					}
-				#if SUBSURFACE_PROFILE_OPACITY_THRESHOLD
-					else
-					{
-						Out.ShadingModelID = SHADINGMODELID_DEFAULT_LIT;
-						Out.CustomData = 0;
-					}
+
+					Out.Curvature = clamp(Out.Curvature, 0.001f, 1.0f);
 				#endif
 				}
-				else if (SSSType == SSS_TYPE_WRAP || SSSType == SSS_TYPE_TWO_SIDED_WRAP)
-				{
-					Out.CustomData.rgb = EncodeSubsurfaceColor(Out.SubsurfaceColor);
-					Out.CustomData.a   = SubstrateSubSurfaceGetWrapOpacityFromAnisotropy(SLAB_SSSPHASEANISOTROPY(CurrentBSDF));
-				}
-				else if (BSDF_GETHASFUZZ(CurrentBSDF))
+			#if SUBSURFACE_PROFILE_OPACITY_THRESHOLD
+				else
 				{
-					Out.CustomData.rgb = EncodeSubsurfaceColor(SLAB_FUZZ_COLOR(CurrentBSDF));
-					Out.CustomData.a   = SLAB_FUZZ_AMOUNT(CurrentBSDF); 
-					Out.IndirectIrradianceScale *= 1 - Out.CustomData.a;
+					Out.ShadingModelID = SHADINGMODELID_DEFAULT_LIT;
+					Out.CustomData = 0;
 				}
-				else if (BSDF_GETHASHAZINESS(CurrentBSDF))
+			#endif
+			}
+			else
+#endif
+#if MATERIAL_SHADINGMODEL_SUBSURFACE || MATERIAL_SHADINGMODEL_TWOSIDED_FOLIAGE || !MATERIAL_IS_SUBSTRATEHIDDENCONVERSION
+				if (SSSType == SSS_TYPE_WRAP || SSSType == SSS_TYPE_TWO_SIDED_WRAP)
+			{
+				Out.CustomData.rgb = EncodeSubsurfaceColor(Out.SubsurfaceColor);
+				Out.CustomData.a   = SubstrateSubSurfaceGetWrapOpacityFromAnisotropy(SLAB_SSSPHASEANISOTROPY(CurrentBSDF));
+			}
+			else
+#endif
+#if MATERIAL_SHADINGMODEL_CLOTH || !MATERIAL_IS_SUBSTRATEHIDDENCONVERSION
+				if (BSDF_GETHASFUZZ(CurrentBSDF))
+			{
+				Out.CustomData.rgb = EncodeSubsurfaceColor(SLAB_FUZZ_COLOR(CurrentBSDF));
+				Out.CustomData.a   = SLAB_FUZZ_AMOUNT(CurrentBSDF); 
+				Out.IndirectIrradianceScale *= 1 - Out.CustomData.a;
+			}
+			else
+#endif
+#if MATERIAL_SHADINGMODEL_CLEAR_COAT || !MATERIAL_IS_SUBSTRATEHIDDENCONVERSION
+				if (BSDF_GETHASHAZINESS(CurrentBSDF))
+			{
+				FHaziness LocalHaziness = UnpackHaziness(SLAB_HAZINESS(CurrentBSDF));
+				if (LocalHaziness.bSimpleClearCoat)
 				{
-					FHaziness LocalHaziness = UnpackHaziness(SLAB_HAZINESS(CurrentBSDF));
-					if (LocalHaziness.bSimpleClearCoat)
-					{
-						Out.CustomData.x = LocalHaziness.Weight;
-						Out.CustomData.y = LocalHaziness.Roughness;
+					Out.CustomData.x = LocalHaziness.Weight;
+					Out.CustomData.y = LocalHaziness.Roughness;
 
-					#if CLEAR_COAT_BOTTOM_NORMAL
-						Out.CustomData.wz = LocalHaziness.BottomNormalOct.xy;
-					#else
-						Out.CustomData.wz = DEFAULT_CLEAR_COAT_BOTTOM_NORMAL_OCT;
-					#endif
-					}
+				#if CLEAR_COAT_BOTTOM_NORMAL
+					Out.CustomData.wz = LocalHaziness.BottomNormalOct.xy;
+				#else
+					Out.CustomData.wz = DEFAULT_CLEAR_COAT_BOTTOM_NORMAL_OCT;
+				#endif
 				}
 			}
-		#undef CurrentBSDF
+#endif
+			// In case of an else, we need to provide a code block.
+			{}
 		}
-
-		Out.WorldNormal = all(Out.WorldNormal == 0) ? float3(0.0, 0.0, 0.0) : normalize(Out.WorldNormal);
+		#undef CurrentBSDF
 	}
 
 	return Out;
diff --git a/Engine/Shaders/Shared/SubstrateDefinitions.h b/Engine/Shaders/Shared/SubstrateDefinitions.h
--- a/Engine/Shaders/Shared/SubstrateDefinitions.h
+++ b/Engine/Shaders/Shared/SubstrateDefinitions.h
@@ -9,7 +9,7 @@
 
 #ifndef __cplusplus
 // Change this to force recompilation of all Substrate dependent shaders (for instance https://guidgenerator.com/online-guid-generator.aspx)
-#pragma message("UESHADERMETADATA_VERSION 161BBFDB-EBFE-41F0-8449-5FAB98EDAA82")
+#pragma message("UESHADERMETADATA_VERSION 96CDF7E5-2EB8-4E48-8C1E-23B557A7B5BB")
 #endif
 
 // Closure offsets are packed into 32bits, each entry using SUBSTRATE_CLOSURE_OFFSET_BIT_COUNT bits
diff --git a/Engine/Source/Runtime/Core/Private/UObject/FortniteMainBranchObjectVersion.cpp b/Engine/Source/Runtime/Core/Private/UObject/FortniteMainBranchObjectVersion.cpp
--- a/Engine/Source/Runtime/Core/Private/UObject/FortniteMainBranchObjectVersion.cpp
+++ b/Engine/Source/Runtime/Core/Private/UObject/FortniteMainBranchObjectVersion.cpp
@@ -14,6 +14,6 @@ TMap<FGuid, FGuid> FFortniteMainBranchObjectVersion::GetSystemGuids()
 	SystemGuids.Add(DevGuids.Niagara_LatestScriptCompileVersion, FGuid("783FF3700BB641C799B2D43552FBF91E"));
 	SystemGuids.Add(DevGuids.SkeletalMeshDerivedDataVersion, FGuid("C998F2F3B2884A02B7E290B8C25835F0"));
 	SystemGuids.Add(DevGuids.STATICMESH_DERIVEDDATA_VER, FGuid("3ABF32149F9D4D838750368AB95573BE"));
-	SystemGuids.Add(DevGuids.MaterialTranslationDDCVersion, FGuid("36B83A3B63AD4B19ABB959FFB26A9F24"));
+	SystemGuids.Add(DevGuids.MaterialTranslationDDCVersion, FGuid("9D0538021BDB47DAA63FFC09DDAAC00E"));
 	return SystemGuids;
 }
diff --git a/Engine/Source/Runtime/Engine/Private/Materials/HLSLMaterialTranslator.cpp b/Engine/Source/Runtime/Engine/Private/Materials/HLSLMaterialTranslator.cpp
--- a/Engine/Source/Runtime/Engine/Private/Materials/HLSLMaterialTranslator.cpp
+++ b/Engine/Source/Runtime/Engine/Private/Materials/HLSLMaterialTranslator.cpp
@@ -294,6 +294,7 @@ struct FHLSLMaterialTranslator::FEnvironmentDefines
 	bool bMaterialVolumetricAdvancedAdvancedGroundContribution;
 	bool bMaterialVolumetricCloudEmptySpaceSkippingOutput;
 	bool bMaterialIsSubstrate;
+	bool bUseRootNodeToSubstrateHiddenConversion;
 	bool bDualSourceColorBlendingEnabled;
 	bool bSubstratePremultipliedAlphaOpacityOverridden;
 	bool bSubstrateUsesConversionFromLegacy;
@@ -401,6 +402,7 @@ struct FHLSLMaterialTranslator::FEnvironmentDefines
 		Ar << bMaterialVolumetricAdvancedAdvancedGroundContribution;
 		Ar << bMaterialVolumetricCloudEmptySpaceSkippingOutput;
 		Ar << bMaterialIsSubstrate;
+		Ar << bUseRootNodeToSubstrateHiddenConversion;
 		Ar << bDualSourceColorBlendingEnabled;
 		Ar << bSubstratePremultipliedAlphaOpacityOverridden;
 		Ar << bSubstrateUsesConversionFromLegacy;
@@ -720,6 +722,7 @@ FHLSLMaterialTranslator::FHLSLMaterialTranslator(FMaterial* InMaterial,
 ,	bAllowCodeChunkGeneration(true)
 ,	bUsesAnisotropy(false)
 ,	bMaterialIsSubstrate(false)
+,	bUseRootNodeToSubstrateHiddenConversion(false)
 ,	bUsesCurvature(false)
 ,	bUsesPerInstanceFadeAmount(false)
 ,	bCullIntermediateUniformExpressions(GCullIntermediateUniformExpressions)
@@ -1467,7 +1470,7 @@ void FHLSLMaterialTranslator::TranslateMaterial()
 	FExpressionInput* ClearCoatNormal = &NullInput;
 	FExpressionInput* CustomTangent = &NullInput;
 
-	const bool bUseRootNodeToSubstrateHiddenConversion = bSubstrateEnabled && bSubstrateHiddenMaterialAssetConversion && !FrontMaterialExpr;
+	bUseRootNodeToSubstrateHiddenConversion = bSubstrateEnabled && bSubstrateHiddenMaterialAssetConversion && !FrontMaterialExpr;
 	if (bUseRootNodeToSubstrateHiddenConversion)
 	{
 		UMaterialInterface* MaterialInterface = Material->GetMaterialInterface();
@@ -2992,6 +2995,7 @@ void FHLSLMaterialTranslator::GetMaterialEnvironment(EShaderPlatform InPlatform,
 	}
 
 	OutEnvironment.SetDefine(TEXT("MATERIAL_IS_SUBSTRATE"), EnvironmentDefines->bMaterialIsSubstrate);
+	OutEnvironment.SetDefine(TEXT("MATERIAL_IS_SUBSTRATEHIDDENCONVERSION"), EnvironmentDefines->bUseRootNodeToSubstrateHiddenConversion);
 	OutEnvironment.SetDefine(TEXT("DUAL_SOURCE_COLOR_BLENDING_ENABLED"), EnvironmentDefines->bDualSourceColorBlendingEnabled);
 	OutEnvironment.SetDefine(TEXT("SUBSTRATE_PREMULTIPLIED_ALPHA_OPACITY_OVERRIDEN"), EnvironmentDefines->bSubstratePremultipliedAlphaOpacityOverridden);
 
@@ -11604,12 +11608,14 @@ int32 FHLSLMaterialTranslator::CalculatePeriodicWorldPositionOrigin(int TileScal
 
 int32 FHLSLMaterialTranslator::TransformNormalFromRequestedBasisToWorld(int32 NormalCodeChunk)
 {
+	// When feeding tangent space or world space, we want to have the final normal/tangent normalized when stored in SharedLocalBases.
+	// So that we do not have to normalise it later in any forward processes and avoid overhead in instructions.
 	if (IsTangentSpaceNormal())
 	{
 		// See TransformTangentNormalToWorld definitions in MaterialTemplate.ush
 		return AddCodeChunk(MCT_Float3, TEXT("TransformTangentNormalToWorld(Parameters.TangentToWorld, %s)"), *GetParameterCode(NormalCodeChunk));
 	}
-	return NormalCodeChunk;
+	return Normalize(NormalCodeChunk);
 }
 
 int32 FHLSLMaterialTranslator::DynamicParameter(FLinearColor& DefaultValue, uint32 ParameterIndex)
@@ -16146,6 +16152,7 @@ void FHLSLMaterialTranslator::PrepareEnvironmentDefines()
 	}
 
 	EnvironmentDefines->bMaterialIsSubstrate = bMaterialIsSubstrate;
+	EnvironmentDefines->bUseRootNodeToSubstrateHiddenConversion = bUseRootNodeToSubstrateHiddenConversion;
 
 	// Substrate requests dual source blending only for BLEND_TranslucentColoredTransmittance
 	bMaterialRequestsDualSourceBlending |= bMaterialIsSubstrate && Material->GetBlendMode() == EBlendMode::BLEND_TranslucentColoredTransmittance;
diff --git a/Engine/Source/Runtime/Engine/Private/Materials/HLSLMaterialTranslator.h b/Engine/Source/Runtime/Engine/Private/Materials/HLSLMaterialTranslator.h
--- a/Engine/Source/Runtime/Engine/Private/Materials/HLSLMaterialTranslator.h
+++ b/Engine/Source/Runtime/Engine/Private/Materials/HLSLMaterialTranslator.h
@@ -449,6 +449,9 @@ protected:
 	 */
 	uint32 bMaterialIsSubstrate : 1;
 
+	/** True if the material is using a hidden conversion from the root node. */
+	uint32 bUseRootNodeToSubstrateHiddenConversion : 1;
+
 	/** True if the opacity input is plugged in */
 	uint32 bOpacityPropertyIsUsed : 1;
 	
