Fix material buffer visualization with Substrate Format=1

diff --git a/Engine/Shaders/Private/MaterialTemplate.ush b/Engine/Shaders/Private/MaterialTemplate.ush
--- a/Engine/Shaders/Private/MaterialTemplate.ush
+++ b/Engine/Shaders/Private/MaterialTemplate.ush
@@ -204,7 +204,11 @@
 	#define SubstratePublicStruct Substrate
 #endif
 
-#if !SUBSTRATE_OPAQUE_DEFERRED
+#ifndef ALLOW_SUBSTRATE_PUBLIC_ACCESS
+#define ALLOW_SUBSTRATE_PUBLIC_ACCESS 0
+#endif
+
+#if !SUBSTRATE_OPAQUE_DEFERRED || ALLOW_SUBSTRATE_PUBLIC_ACCESS
 #include "/Engine/Public/SubstratePublic.ush"
 #endif
 
@@ -3298,7 +3302,7 @@ float4 SceneTextureLookup(float2 UV, int SceneTextureIndex, bool bFiltered)
 	}
 	#elif SUBSTRATE_GBUFFER_FORMAT==1 && !SHADING_PATH_MOBILE // Not on mobile that is not using dbuffer decal but the legacy GBuffer blending.
 	{
-		#if !SUBSTRATE_OPAQUE_DEFERRED // Prevents opaque material to sample the Gbuffer.
+		#if !SUBSTRATE_OPAQUE_DEFERRED || ALLOW_SUBSTRATE_PUBLIC_ACCESS // Prevents opaque material to sample the Gbuffer.
 		FSubstrateGBuffer SubstrateGBuffer = SubstratePublic_GetSubstrateGBufferFromFirstBSDF(PixelPos);
 		BRANCH
 		if (SubstrateGBuffer.bIsValid > 0)
diff --git a/Engine/Shaders/Private/PostProcessMaterialShaders.usf b/Engine/Shaders/Private/PostProcessMaterialShaders.usf
--- a/Engine/Shaders/Private/PostProcessMaterialShaders.usf
+++ b/Engine/Shaders/Private/PostProcessMaterialShaders.usf
@@ -25,6 +25,8 @@ uint YCbCrSrgbToLinear;
 #define POST_PROCESS_MATERIAL 0
 #endif
 
+#define ALLOW_SUBSTRATE_PUBLIC_ACCESS (SUBSTRATE_ENABLED)
+
 #if (POST_PROCESS_MATERIAL == 0)
 #error POST_PROCESS_MATERIAL must be defined to non-zero in the shader compilation environment.
 #endif
