[VSM] Added option to suppress dynamic invalidations for coarse pages on local lights (r.Shadow.Virtual.MarkCoarsePagesLocal=2). If enabled it can improve performance for local lights with VSM while leading to few visual issues.

diff --git a/Engine/Shaders/Private/VirtualShadowMaps/VirtualShadowMapCacheInvalidation.ush b/Engine/Shaders/Private/VirtualShadowMaps/VirtualShadowMapCacheInvalidation.ush
--- a/Engine/Shaders/Private/VirtualShadowMaps/VirtualShadowMapCacheInvalidation.ush
+++ b/Engine/Shaders/Private/VirtualShadowMaps/VirtualShadowMapCacheInvalidation.ush
@@ -135,7 +135,11 @@ void InvalidateInstancePages(FVirtualShadowMapProjectionShaderData ProjectionDat
 							FVSMPageOffset PageFlagOffset = CalcPageOffset(PageTableLevelOffset, MipLevel, uint2(x, y));
 							uint PageFlags = VirtualShadowMapGetPageFlag(PageFlagOffset);
 
-							if ((PageFlags & VSM_FLAG_ALLOCATED) != 0)
+							uint TestFlags = VSM_FLAG_ALLOCATED 
+								// Only invalidate dynamic pages marked with VSM_FLAG_DETAIL_GEOMETRY if bForceCacheDynamicCoarse is enabled for the SM
+								// Note that this is also suppressed downstream so not super important to skip here perhaps.
+								| ((ProjectionData.bForceCacheDynamicCoarse && !bInvalidateStaticPage) ? VSM_FLAG_DETAIL_GEOMETRY : 0u);
+							if ((PageFlags & TestFlags) == TestFlags)
 							{
 #if USE_HZB_OCCLUSION
 								if (!IsPageVisibleHZB(uint2(x, y), PageFlagOffset, false, HZBTestRect, true))
diff --git a/Engine/Shaders/Private/VirtualShadowMaps/VirtualShadowMapPhysicalPageManagement.usf b/Engine/Shaders/Private/VirtualShadowMaps/VirtualShadowMapPhysicalPageManagement.usf
--- a/Engine/Shaders/Private/VirtualShadowMaps/VirtualShadowMapPhysicalPageManagement.usf
+++ b/Engine/Shaders/Private/VirtualShadowMaps/VirtualShadowMapPhysicalPageManagement.usf
@@ -306,9 +306,15 @@ void UpdatePhysicalPages(uint3 Index : SV_DispatchThreadID)
 				else
 				{
 					uint NextPageFlags = VSM_FLAG_ALLOCATED;
+					const uint RequestDetailGeometryFlag = (RequestFlags & VSM_FLAG_DETAIL_GEOMETRY);
+
+					bool bIsStaticInvalidated = (PrevPhysicalFlags & VSM_EXTENDED_FLAG_INVALIDATE_STATIC) != 0u;
+
+					bool bSkipDynamicInvalidations = VirtualShadowMapHandle.IsSinglePage() 
+						|| (RequestDetailGeometryFlag == 0u && Projection.bForceCacheDynamicCoarse && !bIsStaticInvalidated);
 
 					// Distant lights ignore invalidations as they are round-robin invalidated
-					if (bDynamicPageInvalidation && !VirtualShadowMapHandle.IsSinglePage())
+					if (bDynamicPageInvalidation && !bSkipDynamicInvalidations)
 					{
 						if ((PrevPhysicalFlags & VSM_EXTENDED_FLAG_ANY_INVALIDATED) != 0)
 						{
@@ -343,7 +349,6 @@ void UpdatePhysicalPages(uint3 Index : SV_DispatchThreadID)
 						// If the detail geometry flag doesn't match the cached page we treat it as a full invalidation.
 						// TODO: This could potentially be a problem for interleaved multiview rendering;
 						// If the flag differs in the two views it will cause cache thrashing.									
-						const uint RequestDetailGeometryFlag = (RequestFlags & VSM_FLAG_DETAIL_GEOMETRY);
 						if (RequestDetailGeometryFlag != PhysicalPageDetailGeometryFlag)
 						{
 							NextPageFlags |= (VSM_FLAG_STATIC_UNCACHED | VSM_FLAG_DYNAMIC_UNCACHED);
diff --git a/Engine/Shaders/Private/VirtualShadowMaps/VirtualShadowMapProjectionStructs.ush b/Engine/Shaders/Private/VirtualShadowMaps/VirtualShadowMapProjectionStructs.ush
--- a/Engine/Shaders/Private/VirtualShadowMaps/VirtualShadowMapProjectionStructs.ush
+++ b/Engine/Shaders/Private/VirtualShadowMaps/VirtualShadowMapProjectionStructs.ush
@@ -48,6 +48,7 @@ struct FVirtualShadowMapProjectionShaderData
 	bool bUnreferenced; // see VSM_PROJ_FLAG_UNREFERENCED
 	bool bIsCoarseClipLevel; // see VSM_PROJ_FLAG_IS_COARSE_CLIP_LEVEL
 	bool bUseReceiverMask; // see VSM_PROJ_FLAG_USE_RECEIVER_MASK
+	bool bForceCacheDynamicCoarse; // see VSM_PROJ_FLAG_FORCE_CACHE_DYNAMIC_COARSE
 };
 
 FVirtualShadowMapProjectionShaderData DecodeVirtualShadowMapProjectionData(ByteAddressBuffer ProjectionData, FVirtualShadowMapHandle VirtualShadowMapHandle)
@@ -116,6 +117,7 @@ FVirtualShadowMapProjectionShaderData DecodeVirtualShadowMapProjectionData(ByteA
 	Result.bUnreferenced = (Result.Flags & VSM_PROJ_FLAG_UNREFERENCED) != 0U;
 	Result.bIsCoarseClipLevel = (Result.Flags & VSM_PROJ_FLAG_IS_COARSE_CLIP_LEVEL) != 0U;
 	Result.bUseReceiverMask = (Result.Flags & VSM_PROJ_FLAG_USE_RECEIVER_MASK) != 0;
+	Result.bForceCacheDynamicCoarse = (Result.Flags & VSM_PROJ_FLAG_FORCE_CACHE_DYNAMIC_COARSE) != 0;
 
 	return Result;
 }
diff --git a/Engine/Shaders/Shared/VirtualShadowMapDefinitions.h b/Engine/Shaders/Shared/VirtualShadowMapDefinitions.h
--- a/Engine/Shaders/Shared/VirtualShadowMapDefinitions.h
+++ b/Engine/Shaders/Shared/VirtualShadowMapDefinitions.h
@@ -50,6 +50,7 @@ VirtualShadowMapDefinitions.h: used in virtual shadow map shaders and C++ code t
 #define VSM_PROJ_FLAG_IS_COARSE_CLIP_LEVEL (1U << 3) // Used to indicate that the clip level is a coarse level
 #define VSM_PROJ_FLAG_IS_FIRST_PERSON_SHADOW (1U << 4) // Used to indicate that this is a "first-person" shadow 
 #define VSM_PROJ_FLAG_USE_RECEIVER_MASK (1U << 5) // Used to enable receiver masks on this light
+#define VSM_PROJ_FLAG_FORCE_CACHE_DYNAMIC_COARSE (1U << 6) // Used to prevent invalidations of dynamic coarse pages, that is pages NOT marked with detail geometry.
 
 // Hard limit for max distant lights supported 8k for now - we may revise later. We need to keep them in a fixed range for now to make allocation easy and minimize overhead for indexing.
 #define VSM_MAX_SINGLE_PAGE_SHADOW_MAPS (1024U * 8U)
diff --git a/Engine/Source/Runtime/Renderer/Private/Shadows/ShadowSceneRenderer.cpp b/Engine/Source/Runtime/Renderer/Private/Shadows/ShadowSceneRenderer.cpp
--- a/Engine/Source/Runtime/Renderer/Private/Shadows/ShadowSceneRenderer.cpp
+++ b/Engine/Source/Runtime/Renderer/Private/Shadows/ShadowSceneRenderer.cpp
@@ -123,6 +123,8 @@ static TAutoConsoleVariable<int32>  CVarForceInvalidateLocalVSM(
 extern TAutoConsoleVariable<int32> CVarMarkPixelPagesMipModeLocal;
 extern float GMinScreenRadiusForShadowCaster;
 
+extern TAutoConsoleVariable<int32> CVarMarkCoarsePagesLocal;
+
 bool IsVSMOnePassProjectionEnabled(const FEngineShowFlags& ShowFlags)
 {
 	return CVarVirtualShadowOnePassProjection.GetValueOnAnyThread() 
@@ -473,6 +475,14 @@ TSharedPtr<FVirtualShadowMapPerLightCacheEntry> FShadowSceneRenderer::AddLocalLi
 		uint32 PackedCullingViewId = FVirtualShadowMapProjectionShaderData::PackCullingViewId(SceneRenderer.Views[ClosestCullingViewIndex].SceneRendererPrimaryViewId, SceneRenderer.Views[ClosestCullingViewIndex].PersistentViewId);
 		uint32 Flags = PerLightCacheEntry->IsUncached() ? VSM_PROJ_FLAG_UNCACHED : 0U;
 		Flags |= PerLightCacheEntry->ShouldUseReceiverMask() ? VSM_PROJ_FLAG_USE_RECEIVER_MASK : 0U;
+
+		// If the coarse pages mode is 2, we suppress invalidaitons of pages not marked as "detail geometry". 
+		// Coarse pages are not marked with this flag so will not be invalidated.
+		if (CVarMarkCoarsePagesLocal.GetValueOnRenderThread() == 2)
+		{
+			Flags |= VSM_PROJ_FLAG_FORCE_CACHE_DYNAMIC_COARSE;
+		}
+
 		const FLightSceneProxy* Proxy = ProjectedShadowInfo->GetLightSceneInfo().Proxy;
 
 		// For now just tie this to whether anything has invalidated the light (including movement)
diff --git a/Engine/Source/Runtime/Renderer/Private/VirtualShadowMaps/VirtualShadowMapArray.cpp b/Engine/Source/Runtime/Renderer/Private/VirtualShadowMaps/VirtualShadowMapArray.cpp
--- a/Engine/Source/Runtime/Renderer/Private/VirtualShadowMaps/VirtualShadowMapArray.cpp
+++ b/Engine/Source/Runtime/Renderer/Private/VirtualShadowMaps/VirtualShadowMapArray.cpp
@@ -327,8 +327,10 @@ TAutoConsoleVariable<int32> CVarMarkPixelPagesMipModeLocal(
 TAutoConsoleVariable<int32> CVarMarkCoarsePagesLocal(
 	TEXT("r.Shadow.Virtual.MarkCoarsePagesLocal"),
 	1,
-	TEXT("Marks coarse pages in local light virtual shadow maps so that low resolution data is available everywhere.")
-	TEXT("Ability to disable is primarily for profiling and debugging."),
+	TEXT("Marks coarse pages in local light virtual shadow maps so that low resolution data is available everywhere.\n")
+	TEXT(" 0 - Disabled\n")
+	TEXT(" 1 - (default) Always mark the last mip level.\n")
+	TEXT(" 2 - Performance Mode: Suppress dynamic invalidations due to geometry changes (e.g., moving objects or WPO or animation).\n"),
 	ECVF_Scalability | ECVF_RenderThreadSafe
 );
 
@@ -2581,7 +2583,7 @@ void FVirtualShadowMapArray::BeginMarkPages(
 	{
 		FMarkCoarsePagesCS::FParameters* PassParameters = GraphBuilder.AllocParameters< FMarkCoarsePagesCS::FParameters >();
 		PassParameters->MarkingParameters = GetMarkingParameters(GraphBuilder);
-		PassParameters->bMarkCoarsePagesLocal = CVarMarkCoarsePagesLocal.GetValueOnRenderThread();
+		PassParameters->bMarkCoarsePagesLocal = CVarMarkCoarsePagesLocal.GetValueOnRenderThread() != 0 ? 1 : 0;
 		PassParameters->bIncludeNonNaniteGeometry = CVarCoarsePagesIncludeNonNanite.GetValueOnRenderThread();
 
 		auto ComputeShader = GetGlobalShaderMap(Scene.GetFeatureLevel())->GetShader<FMarkCoarsePagesCS>();
