substrate - fix for compilation error on mobile.

#tests PC editor dx12
#rnx

#DisabledValidations Preflight

diff --git a/Engine/Shaders/Private/Substrate/SubstrateRead.ush b/Engine/Shaders/Private/Substrate/SubstrateRead.ush
--- a/Engine/Shaders/Private/Substrate/SubstrateRead.ush
+++ b/Engine/Shaders/Private/Substrate/SubstrateRead.ush
@@ -12,6 +12,7 @@
 
 #include "../Common.ush"
 #include "Substrate.ush"
+#include "/Engine/Private/DeferredShadingCommon.ush"
 
 struct FSubstrateGBufferBSDF
 {
@@ -204,8 +205,16 @@ FSubstrateGBufferBSDF SubstrateReadGBufferBSDF(in FScreenSpaceData ScreenSpaceDa
 
 FSubstrateBSDF UnpackSubstrateBSDFIn(inout FSubstrateAddressing SubstrateAddressing)
 {
+	FSubstrateGBufferBSDF GBufferBSDF = (FSubstrateGBufferBSDF)0; 
+#if SHADING_PATH_DEFERRED
+#if FEATURE_LEVEL >= FEATURE_LEVEL_SM5
 	FScreenSpaceData ScreenSpaceData = GetScreenSpaceDataUint(SubstrateAddressing.PixelCoords);
-	FSubstrateGBufferBSDF GBufferBSDF = SubstrateReadGBufferBSDF(ScreenSpaceData);
+#else
+	const float2 UV = (float2(SubstrateAddressing.PixelCoords) + 0.5f) * View.BufferSizeAndInvSize.zw;
+	FScreenSpaceData ScreenSpaceData = GetScreenSpaceData(UV);
+#endif
+	GBufferBSDF = SubstrateReadGBufferBSDF(ScreenSpaceData);
+#endif
 	return GBufferBSDF.BSDF;
 }
 
diff --git a/Engine/Shaders/Public/SubstratePublic.ush b/Engine/Shaders/Public/SubstratePublic.ush
--- a/Engine/Shaders/Public/SubstratePublic.ush
+++ b/Engine/Shaders/Public/SubstratePublic.ush
@@ -11,7 +11,7 @@
 
 #if SUBSTRATE_ENABLED 
 
-#define SUBSTRATE_ALLOWS_GBUFFER_SAMPLING  ((NEEDS_SCENE_TEXTURES && !SCENE_TEXTURES_DISABLED))
+#define SUBSTRATE_ALLOWS_GBUFFER_SAMPLING  ((NEEDS_SCENE_TEXTURES && !SCENE_TEXTURES_DISABLED)) && SHADING_PATH_DEFERRED
 
 /** Public facing functions allowing external users access to some Substrate data. */
 
