Make sure Substrate simple path DCEs everything that isn't the unlit or default lit

diff --git a/Engine/Shaders/Private/DeferredLightPixelShaders.usf b/Engine/Shaders/Private/DeferredLightPixelShaders.usf
--- a/Engine/Shaders/Private/DeferredLightPixelShaders.usf
+++ b/Engine/Shaders/Private/DeferredLightPixelShaders.usf
@@ -279,7 +279,8 @@ void DeferredLightPixelMain(
 	FSubstratePixelHeader SubstratePixelHeader = UnpackSubstrateHeaderIn(Substrate.MaterialTextureArray, SubstrateAddressing, Substrate.TopLayerTexture);
 	const bool bMaterialIsLit = SubstratePixelHeader.ClosureCount > 0;	// This test is also enough to exclude sky pixels
 	#else
-	FScreenSpaceData ScreenSpaceData = GetScreenSpaceData(InputParams.ScreenUV);
+	const bool bForceUnlitOrDefaultLit = SUBSTRATE_FASTPATH;
+	FScreenSpaceData ScreenSpaceData = GetScreenSpaceData(InputParams.ScreenUV, true, bForceUnlitOrDefaultLit);
 	const FGBufferData GBuffer = ScreenSpaceData.GBuffer;
 	const bool bMaterialIsLit = GBuffer.ShadingModelID != SHADINGMODELID_UNLIT;
 	#endif
diff --git a/Engine/Shaders/Private/DeferredShadingCommon.ush b/Engine/Shaders/Private/DeferredShadingCommon.ush
--- a/Engine/Shaders/Private/DeferredShadingCommon.ush
+++ b/Engine/Shaders/Private/DeferredShadingCommon.ush
@@ -1153,10 +1153,10 @@ float ConvertSubsurfaceProfileToFloat(uint InProfileIdUint)
 
 // @param UV - UV space in the GBuffer textures (BufferSize resolution)
 // TOOD: need to find a way to make this more convenient.
-FGBufferData GetGBufferDataFromSceneTextures(float2 UV, bool bGetNormalizedNormal = true)
+FGBufferData GetGBufferDataFromSceneTextures(float2 UV, bool bGetNormalizedNormal = true, bool bForceUnlitOrDefaultLit = false)
 {
 #if GBUFFER_REFACTOR
-	return DecodeGBufferDataSceneTextures(UV,bGetNormalizedNormal);
+	return DecodeGBufferDataSceneTextures(UV,bGetNormalizedNormal,bForceUnlitOrDefaultLit);
 #else
 	float4 GBufferA = GBufferATexture.SampleLevel(GBufferATextureSampler, UV, 0);
 	float4 GBufferB = GBufferBTexture.SampleLevel(GBufferBTextureSampler, UV, 0);
@@ -1189,10 +1189,10 @@ uint GetSceneLightingChannel(uint2 PixelCoord)
 }
 
 // @param UV - UV space in the GBuffer textures (BufferSize resolution)
-FGBufferData GetGBufferData(float2 UV, bool bGetNormalizedNormal = true)
+FGBufferData GetGBufferData(float2 UV, bool bGetNormalizedNormal = true, bool bForceUnlitOrDefaultLit = false)
 {
 #if GBUFFER_REFACTOR
-	return DecodeGBufferDataUV(UV,bGetNormalizedNormal);
+	return DecodeGBufferDataUV(UV,bGetNormalizedNormal,bForceUnlitOrDefaultLit);
 #else
 	float4 GBufferA = Texture2DSampleLevel(SceneTexturesStruct.GBufferATexture, SceneTexturesStruct_GBufferATextureSampler, UV, 0);
 	float4 GBufferB = Texture2DSampleLevel(SceneTexturesStruct.GBufferBTexture, SceneTexturesStruct_GBufferBTextureSampler, UV, 0);
@@ -1232,11 +1232,11 @@ uint GetShadingModelId(float2 UV)
 }
 
 // @param UV - UV space in the GBuffer textures (BufferSize resolution)
-FScreenSpaceData GetScreenSpaceData(float2 UV, bool bGetNormalizedNormal = true)
+FScreenSpaceData GetScreenSpaceData(float2 UV, bool bGetNormalizedNormal = true, bool bForceUnlitOrDefaultLit = false)
 {
 	FScreenSpaceData Out;
 
-	Out.GBuffer = GetGBufferData(UV, bGetNormalizedNormal);
+	Out.GBuffer = GetGBufferData(UV, bGetNormalizedNormal, bForceUnlitOrDefaultLit);
 	float4 ScreenSpaceAO = Texture2DSampleLevel(SceneTexturesStruct.ScreenSpaceAOTexture, SceneTexturesStruct_ScreenSpaceAOTextureSampler, UV, 0);
 
 	Out.AmbientOcclusion = ScreenSpaceAO.r;
diff --git a/Engine/Shaders/Private/ReflectionEnvironmentPixelShader.usf b/Engine/Shaders/Private/ReflectionEnvironmentPixelShader.usf
--- a/Engine/Shaders/Private/ReflectionEnvironmentPixelShader.usf
+++ b/Engine/Shaders/Private/ReflectionEnvironmentPixelShader.usf
@@ -429,8 +429,9 @@ void ReflectionEnvironmentSkyLighting(
 	FSubstratePixelHeader SubstratePixelHeader = UnpackSubstrateHeaderIn(Substrate.MaterialTextureArray, SubstrateAddressing, Substrate.TopLayerTexture);
 	const bool bMaterialIsLit = SubstratePixelHeader.ClosureCount > 0;	// This test is also enough to exclude sky pixels
 	#else
+	const bool bForceUnlitOrDefaultLit = SUBSTRATE_FASTPATH;
+	const FGBufferData GBuffer = GetGBufferDataFromSceneTextures(BufferUV, true, bForceUnlitOrDefaultLit);
 	FScreenSpaceData ScreenSpaceData;
-	const FGBufferData GBuffer = GetGBufferDataFromSceneTextures(BufferUV);
 	ScreenSpaceData.GBuffer = GBuffer;
 	ScreenSpaceData.AmbientOcclusion = 1.0f; // Used, fetched from AmbientOcclusionTexture.
 	const bool bMaterialIsLit = GBuffer.ShadingModelID != SHADINGMODELID_UNLIT;
diff --git a/Engine/Source/Runtime/Engine/Private/ShaderCompiler/ShaderGenerationUtil.cpp b/Engine/Source/Runtime/Engine/Private/ShaderCompiler/ShaderGenerationUtil.cpp
--- a/Engine/Source/Runtime/Engine/Private/ShaderCompiler/ShaderGenerationUtil.cpp
+++ b/Engine/Source/Runtime/Engine/Private/ShaderCompiler/ShaderGenerationUtil.cpp
@@ -1164,7 +1164,8 @@ static FString CreateGBufferDecodeFunctionDirect(const FGBufferInfo& BufferInfo)
 	FullStr += TEXT(",\n\tuint CustomStencil");
 	FullStr += TEXT(",\n\tfloat SceneDepth");
 	FullStr += TEXT(",\n\tbool bGetNormalizedNormal");
-	FullStr += TEXT(",\n\tbool bChecker)\n");
+	FullStr += TEXT(",\n\tbool bChecker");
+	FullStr += TEXT(",\n\tbool bForceUnlitOrDefaultLit = false)\n");
 
 	FullStr += TEXT("{\n");
 
@@ -1349,6 +1350,8 @@ static FString CreateGBufferDecodeFunctionDirect(const FGBufferInfo& BufferInfo)
 		}
 	}
 
+	FullStr += TEXT("\tif(bForceUnlitOrDefaultLit && Ret.ShadingModelID != 0) Ret.ShadingModelID = 1;\n");
+
 	FullStr += TEXT("\tRet.WorldTangent = AnisotropicData.xyz;\n");
 	FullStr += TEXT("\tRet.Anisotropy = AnisotropicData.w;\n");
 
@@ -1411,7 +1414,7 @@ static FString CreateGBufferDecodeFunctionVariation(const FGBufferInfo& BufferIn
 
 	FullStr += TEXT("// @param PixelPos relative to left top of the rendertarget (not viewport)\n");
 
-	FullStr += FString::Printf(TEXT("FGBufferData DecodeGBufferData%s(%s %s, bool bGetNormalizedNormal = true)\n"),
+	FullStr += FString::Printf(TEXT("FGBufferData DecodeGBufferData%s(%s %s, bool bGetNormalizedNormal = true, bool bForceUnlitOrDefaultLit = false)\n"),
 		Suffix.GetCharArray().GetData(),
 		CoordType.GetCharArray().GetData(),
 		CoordName.GetCharArray().GetData());
@@ -1580,9 +1583,10 @@ static FString CreateGBufferDecodeFunctionVariation(const FGBufferInfo& BufferIn
 	FullStr += TEXT(",\n\t\tCustomStencil");
 	FullStr += TEXT(",\n\t\tSceneDepth");
 	FullStr += TEXT(",\n\t\tbGetNormalizedNormal");
-	FullStr += FString::Printf(TEXT(",\n\t\t%s(%s));\n"),
+	FullStr += FString::Printf(TEXT(",\n\t\t%s(%s)"),
 					CheckerFunc.GetCharArray().GetData(),
 					CoordName.GetCharArray().GetData());
+	FullStr += TEXT(",\n\t\tbForceUnlitOrDefaultLit);\n");
 
 	FullStr += TEXT("\n");
 	FullStr += TEXT("\treturn Ret;\n");
