1. Fix per-pixel MFP causing incorrect SSS in substrate mode due to profile id caching. SSS_PROFILE_PERPIXEL(512) is stored as Depth>0, SSS_PROFILE_INVALID_ID(255) inside the profile id caching texture (R8).
2. Fix a potential bug. Set the SSS implementation default of SSS_PROFILE_PERPIXEL to AFIS.

#tests Test done on the jira reported bug level Mateiral_Substrate in the content example
#rnx

diff --git a/Engine/Shaders/Private/PostProcessSubsurface.usf b/Engine/Shaders/Private/PostProcessSubsurface.usf
--- a/Engine/Shaders/Private/PostProcessSubsurface.usf
+++ b/Engine/Shaders/Private/PostProcessSubsurface.usf
@@ -261,6 +261,10 @@ uint ExtractSubsurfaceProfileIntWithInvalid(float2 UV,float Depth)
 	{
 		uint2 PosOfSampledUV = ConvertUVPos2Input0GridPos(UV);
 		ProfileID = ProfileIdTexture[PosOfSampledUV];
+#if SUBSTRATE_GBUFFER_FORMAT == 1
+		//When Depth > 0, the profile id is valid, SSS_PROFILE_ID_PERPIXEL is compressed as SSS_PROFILE_ID_INVALID due to Profile ID caching texture being R8.
+		ProfileID = (ProfileID == SSS_PROFILE_ID_INVALID) ? SSS_PROFILE_ID_PERPIXEL : ProfileID; 
+#endif
 	}
 	return ProfileID;
 }
@@ -639,6 +643,10 @@ void BuildIndirectDispatchArgsCS(uint2 DispatchThreadId : SV_DispatchThreadID)
 
 uint GetSubsurfaceType(uint ProfileId)
 {
+#if SUBSTRATE_GBUFFER_FORMAT==1
+	// We don't have a bit to indicate which SSS implementation to use, so default to TYPE_BURLEY that uses AFIS. This is default behavior previously.
+	if (ProfileId == SSS_PROFILE_ID_PERPIXEL) return TYPE_BURLEY;
+#endif
 	return GetSubsurfaceProfileUseBurley(ProfileId) ? TYPE_BURLEY : TYPE_SEPARABLE;
 }
 
diff --git a/Engine/Source/Runtime/Renderer/Private/PostProcess/PostProcessSubsurface.cpp b/Engine/Source/Runtime/Renderer/Private/PostProcess/PostProcessSubsurface.cpp
--- a/Engine/Source/Runtime/Renderer/Private/PostProcess/PostProcessSubsurface.cpp
+++ b/Engine/Source/Runtime/Renderer/Private/PostProcess/PostProcessSubsurface.cpp
@@ -265,10 +265,7 @@ END_SHADER_PARAMETER_STRUCT()
 
 uint32 GetSSSMipmapsMinTileCount(const FViewInfo& View)
 {
-	const int32 MinGenerateMipsTileCount = FMath::Max(0, CVarSSSMipmapsMinTileCount.GetValueOnRenderThread());
-	// When using Substrate and Adaptive GBuffer format, we support per-pixel MFP, which requires to run Burley categorization. 
-	// This forces Burley to run even when a limited number of tile requires MIPs/Burley in order to output the correct look
-	return Substrate::IsSubstrateEnabled() && !Substrate::IsSubstrateBlendableGBufferEnabled(View.GetShaderPlatform()) ? 0u : MinGenerateMipsTileCount;
+	return FMath::Max(0, CVarSSSMipmapsMinTileCount.GetValueOnRenderThread());
 }
 
 FSubsurfaceParameters GetSubsurfaceCommonParameters(FRDGBuilder& GraphBuilder, const FViewInfo& View, TRDGUniformBufferRef<FSceneTextureUniformParameters> SceneTextures)
