-Emit warning when hlsl code block in the material editor are using the non common SceneTextureLookup function. This ensures that material authors are made aware when they rely on a low-level, render path-specific scene texture lookup.
-Fix SceneTextures warnings when mobile supports SM5 nodes.

#tests replayruns

diff --git a/Engine/Shaders/Private/MaterialTemplate.ush b/Engine/Shaders/Private/MaterialTemplate.ush
--- a/Engine/Shaders/Private/MaterialTemplate.ush
+++ b/Engine/Shaders/Private/MaterialTemplate.ush
@@ -3396,6 +3396,15 @@ float4 SceneTextureFetchFunc(FMaterialPixelParameters Parameters, int SceneTextu
 	return SceneTextureFetchFunc(Parameters, SceneTextureIndex, float2(PixelOffsetX, PixelOffsetY));
 }
 
+float4 SceneTextureLookupCommon(FMaterialPixelParameters Parameters, float2 UV, int SceneTextureIndex, bool bFiltered)
+{
+#if SHADING_PATH_MOBILE
+	return MobileSceneTextureLookup(Parameters, SceneTextureIndex, UV);
+#else
+	return SceneTextureLookup(UV, SceneTextureIndex, bFiltered);
+#endif
+}
+
 // SceneTextureLookup function used by new translator.  Shading path agnostic, and handles conditional application of exposure scale, so generated code doesn't need to vary.
 float4 SceneTextureLookup(FMaterialPixelParameters Parameters, float2 UV, int SceneTextureIndex, bool bFiltered)
 {
diff --git a/Engine/Source/Runtime/Engine/Private/Materials/HLSLMaterialTranslator.cpp b/Engine/Source/Runtime/Engine/Private/Materials/HLSLMaterialTranslator.cpp
--- a/Engine/Source/Runtime/Engine/Private/Materials/HLSLMaterialTranslator.cpp
+++ b/Engine/Source/Runtime/Engine/Private/Materials/HLSLMaterialTranslator.cpp
@@ -8439,7 +8439,9 @@ void FHLSLMaterialTranslator::UseSceneTextureId(ESceneTextureId SceneTextureId,
 		const bool bSceneTextureRequiresSM5 = SceneTextureId == PPI_WorldNormal;
 		if (bSceneTextureRequiresSM5)
 		{
-			ErrorUnlessFeatureLevelSupported(ERHIFeatureLevel::SM5);
+			// Minimum requirement for mobile supporting SM5 nodes is ES3_1.
+			ERHIFeatureLevel::Type RequiredFeatureLevel = MobileSupportsSM5MaterialNodes(Platform) ? ERHIFeatureLevel::ES3_1 : ERHIFeatureLevel::SM5;
+			ErrorUnlessFeatureLevelSupported(RequiredFeatureLevel);
 		}
 
 		if (SceneTextureId == PPI_WorldNormal && Material->HasNormalConnected() && !IsUsingDBuffers(Platform))
