Fix accidental commenting out of VSM stats

diff --git a/Engine/Shaders/Private/VirtualShadowMaps/VirtualShadowMapPhysicalPageManagement.usf b/Engine/Shaders/Private/VirtualShadowMaps/VirtualShadowMapPhysicalPageManagement.usf
--- a/Engine/Shaders/Private/VirtualShadowMaps/VirtualShadowMapPhysicalPageManagement.usf
+++ b/Engine/Shaders/Private/VirtualShadowMaps/VirtualShadowMapPhysicalPageManagement.usf
@@ -340,7 +340,7 @@ void UpdatePhysicalPages(uint3 Index : SV_DispatchThreadID)
 					uint PhysicalPageDetailGeometryFlag = (PrevPhysicalFlags & VSM_FLAG_DETAIL_GEOMETRY);
 					if (bRequestedThisFrame)
 					{
-						//StatsBufferInterlockedInc(VSM_STAT_REQUESTED_THIS_FRAME_PAGES);
+						StatsBufferInterlockedInc(VSM_STAT_REQUESTED_THIS_FRAME_PAGES);
 						// Remove from LRU list and add to requested list
 						PushPhysicalPageList(PHYSICAL_PAGE_LIST_REQUESTED, PhysicalPageIndex);
 						OutPhysicalPageMetaData[PhysicalPageIndex].LastRequestedSceneFrameNumber = VirtualShadowMap.SceneFrameNumber;
@@ -432,7 +432,7 @@ void AllocateNewPageMappings(FVirtualShadowMapHandle VirtualShadowMapHandle, FVS
 
 		if (PageFlags == 0u)
 		{
-			//StatsBufferInterlockedInc(VSM_STAT_REQUESTED_THIS_FRAME_PAGES);
+			StatsBufferInterlockedInc(VSM_STAT_REQUESTED_THIS_FRAME_PAGES);
 
 			int PhysicalPageIndex = PopPhysicalPageList(PHYSICAL_PAGE_LIST_AVAILABLE);
 			if (PhysicalPageIndex >= 0)
