Substrate: Update Diffuse_Chan model to remove the built-in F0=0.04 which double counts the specular energy preservation

This factor could only partially be compensated for by dividing out the GGX albedo, which led to incorrect behavior at grazing angles, particularly at low roughness.

The new formulation of the Chan BSDF uses a lambertian response at Roughness=0 which is a better match for the ground truth.

The previous implementation is still available via #ifdef for now so we can do additional comparisons if needed.

[RN] Substrate: Improve rough diffuse model to have less energy loss at low roughness

diff --git a/Engine/Shaders/Private/BRDF.ush b/Engine/Shaders/Private/BRDF.ush
--- a/Engine/Shaders/Private/BRDF.ush
+++ b/Engine/Shaders/Private/BRDF.ush
@@ -231,6 +231,7 @@ float3 Diffuse_Gotanda( float3 DiffuseColor, float Roughness, float NoV, float N
 
 // [ Chan 2018, "Material Advances in Call of Duty: WWII" ]
 // It has been extended here to fade out retro reflectivity contribution from area light in order to avoid visual artefacts.
+#define DIFFUSE_CHAN_VERSION 2
 float3 Diffuse_Chan( float3 DiffuseColor, float a2, float NoV, float NoL, float VoH, float NoH, float RetroReflectivityWeight)
 {
 	// We saturate each input to avoid out of range negative values which would result in weird darkening at the edge of meshes (resulting from tangent space interpolation).
@@ -238,7 +239,21 @@ float3 Diffuse_Chan( float3 DiffuseColor, float a2, float NoV, float NoL, float
 	NoL = saturate(NoL);
 	VoH = saturate(VoH);
 	NoH = saturate(NoH);
-
+#if DIFFUSE_CHAN_VERSION == 2
+	// Updated model [Multiscattering Diffuse and Specular BRDFs, Danny Chan, 2024 - Unpublished]
+	const float Roughness = a2;
+	const float Alpha = Roughness * Roughness;
+	// The original writeup uses an FSmooth term inspired by Burley diffuse to balance energy between spec/diffuse.
+	// However in our implementation the energy balance between diffuse and spec is handled externally, so we stick
+	// to a plain lambertian for the Roughness=0 limit.
+	const float FSmooth = 1;
+	const float Scale = max(0.55 - 0.2 * Roughness, 1.25 - 1.6 * Roughness);
+	const float Bias = saturate(4 * Alpha);
+	const float FRough = Scale * (NoH + Bias) * rcp(NoH + 0.025) * VoH * VoH;
+	const float DiffuseSS = lerp(FSmooth, lerp(FSmooth, FRough, RetroReflectivityWeight), Roughness);
+	const float DiffuseMS = Alpha * 0.38;
+	return (1 / PI) * DiffuseColor * (DiffuseSS + DiffuseMS);
+#else
 	// a2 = 2 / ( 1 + exp2( 18 * g )
 	float g = saturate( (1.0 / 18.0) * log2( 2 * rcpFast(a2) - 1 ) );
 
@@ -263,6 +278,7 @@ float3 Diffuse_Chan( float3 DiffuseColor, float a2, float NoV, float NoL, float
 	Lobe = min(1.0, Lobe);
 
 	return DiffuseColor * Lobe;
+#endif
 }
 
 // [Blinn 1977, "Models of light reflection for computer synthesized pictures"]
diff --git a/Engine/Shaders/Private/Substrate/SubstrateEvaluation.ush b/Engine/Shaders/Private/Substrate/SubstrateEvaluation.ush
--- a/Engine/Shaders/Private/Substrate/SubstrateEvaluation.ush
+++ b/Engine/Shaders/Private/Substrate/SubstrateEvaluation.ush
@@ -797,7 +797,7 @@ FSubstrateEvaluateResult SubstrateEvaluateBSDFCommon(FSubstrateBSDFContext BSDFC
 
 			// Specular to Diffuse energy preservation
 			{
-				#if MATERIAL_ROUGHDIFFUSE && SUBSTRATE_SHADING_QUALITY == 1	
+				#if DIFFUSE_CHAN_VERSION == 1 && MATERIAL_ROUGHDIFFUSE && SUBSTRATE_SHADING_QUALITY == 1	
 				// We remove the Chan model baked transmittance specular directional albedo assuming F=0.04.
 				// Then we apply the transmittance according to the current surface setup.
 				// To make this a safe operation:
