Substrate - fixed missing emissive with optimised export to GBuffer.

#rnx
#tests PC editor dx12

diff --git a/Engine/Shaders/Private/Substrate/SubstrateExport.ush b/Engine/Shaders/Private/Substrate/SubstrateExport.ush
--- a/Engine/Shaders/Private/Substrate/SubstrateExport.ush
+++ b/Engine/Shaders/Private/Substrate/SubstrateExport.ush
@@ -971,6 +971,7 @@ FExportResult SubstrateMaterialExportOut(
 		// In this case, there is always a BSDF and it is always 100% visible.
 		// Sanitize BSDF before it is exported.
 		CurrentBSDF.SubstrateSanitizeBSDF();
+		CurrentBSDF.LuminanceWeightV = 1.0f; 
 		{
 #else
 	if (SubstratePixelHeader.SubstrateTree.BSDFCount > 0)
