Path Tracing: Implement missing cases to handle the new Refraction method (IOR from F0) on materials

Also tweak the material editor logic to disable the pin in that mode as it is not used.

#rnx

diff --git a/Engine/Shaders/Private/PathTracing/PathTracingMaterialHitShader.usf b/Engine/Shaders/Private/PathTracing/PathTracingMaterialHitShader.usf
--- a/Engine/Shaders/Private/PathTracing/PathTracingMaterialHitShader.usf
+++ b/Engine/Shaders/Private/PathTracing/PathTracingMaterialHitShader.usf
@@ -57,6 +57,8 @@ float AdjustMaterialRoughness(float Roughness, float PathRoughness)
 
 #if REFRACTION_USE_INDEX_OF_REFRACTION || REFRACTION_USE_INDEX_OF_REFRACTION_FROM_F0
 
+#define PATH_TRACER_REFRACTION_ENABLED 1
+
 #if SUBSTRATE_ENABLED
 
 float GetRefractionIor(FPixelMaterialInputs PixelMaterialInputs, float3 F0)
@@ -123,8 +125,11 @@ float3 ComputeOpticalDepthForHit(float3 LocalSigmaT, bool bIsFrontFace)
 	return LocalSigmaT * Distance;
 }
 
+#else
+
+#define PATH_TRACER_REFRACTION_ENABLED 0
 
-#endif // REFRACTION_USE_INDEX_OF_REFRACTION
+#endif // REFRACTION_USE_INDEX_OF_REFRACTION || REFRACTION_USE_INDEX_OF_REFRACTION_FROM_F0
 
 float3 DecodeSSSProfileRadius(uint ProfileId, float3 DiffuseColor, float Opacity)
 {
@@ -247,14 +252,14 @@ RAY_TRACING_ENTRY_CLOSEST_HIT(PathTracingMaterialCHS,
 	float3 TotalTransmittancePreCoverage = 0;
 	const float3 V_World = MaterialParameters.CameraVector;
 	SubstratePixelHeader.SubstrateUpdateTree(SubstrateData, V_World, IntegrationSettings, TotalCoverage, TotalTransmittancePreCoverage);
-#if REFRACTION_USE_INDEX_OF_REFRACTION
+#if PATH_TRACER_REFRACTION_ENABLED
 	// transparency is only used for coverage, transmittance happens through refraction (except for shadows potentially, which is handled below)
 	float3 Transparency = 1 - TotalCoverage; 
 #else
 	// transparency is used for both coverage and transmittance
 	// SUBSTRATE_TODO: Thin case does not seem to be properly handled here, so the TotalTransmittancePreCoverage is recomputed later
 	const float3 Transparency = saturate(lerp(1.0, TotalTransmittancePreCoverage, TotalCoverage));
-#endif // REFRACTION_USE_INDEX_OF_REFRACTION
+#endif // PATH_TRACER_REFRACTION_ENABLED
 #endif // SUBSTRATE_OPTIMIZED_UNLIT
 
 #else // MATERIAL_IS_SUBSTRATE
@@ -615,7 +620,7 @@ RAY_TRACING_ENTRY_CLOSEST_HIT(PathTracingMaterialCHS,
 
 							if (BSDF.bIsBottom)
 							{
-								#if REFRACTION_USE_INDEX_OF_REFRACTION
+								#if PATH_TRACER_REFRACTION_ENABLED
 								float Ior = GetRefractionIor(PixelMaterialInputs, F0);
 								#endif
 
@@ -804,7 +809,7 @@ RAY_TRACING_ENTRY_CLOSEST_HIT(PathTracingMaterialCHS,
 	half  Specular = GetMaterialSpecular(PixelMaterialInputs);
 	half Roughness = GetMaterialRoughness(PixelMaterialInputs);
 	float Ior = 0.0;
-#if MATERIALBLENDING_TRANSLUCENT && REFRACTION_USE_INDEX_OF_REFRACTION && (MATERIAL_SHADINGMODEL_DEFAULT_LIT || MATERIAL_SHADINGMODEL_THIN_TRANSLUCENT)
+#if MATERIALBLENDING_TRANSLUCENT && PATH_TRACER_REFRACTION_ENABLED && (MATERIAL_SHADINGMODEL_DEFAULT_LIT || MATERIAL_SHADINGMODEL_THIN_TRANSLUCENT)
 	// NOTE: only default-lit and thin-translucent use ior, the other material models only support plain transparency
 	// This is an attempt to limit the complexity of the material eval/sample API which must take into account an extra lobe if supporting refraction
 	Ior = GetRefractionIor(PixelMaterialInputs);
@@ -1335,7 +1340,7 @@ RAY_TRACING_ENTRY_ANY_HIT(PathTracingMaterialAHS,
 		float3 TotalTransmittancePreCoverage = 0;
 		const float3 V_World = MaterialParameters.CameraVector;
 		SubstratePixelHeader.SubstrateUpdateTree(SubstrateData, V_World, IntegrationSettings, TotalCoverage, TotalTransmittancePreCoverage);
-#if MATERIALBLENDING_TRANSLUCENT && (REFRACTION_USE_INDEX_OF_REFRACTION || MATERIAL_ISTHINSURFACE)
+#if MATERIALBLENDING_TRANSLUCENT && (PATH_TRACER_REFRACTION_ENABLED || MATERIAL_ISTHINSURFACE)
 
 		// SUBSTRATE_TODO: refraction handling
 		float PathRoughness = PackedPayload.GetPathRoughness();
@@ -1378,7 +1383,7 @@ RAY_TRACING_ENTRY_ANY_HIT(PathTracingMaterialAHS,
 						{
 							// only non-SSS slabs 
 							LocalSigmaT = rcp(max(SLAB_SSSMFP(BSDF), 0.0001));
-#if MATERIALBLENDING_TRANSLUCENT && (SUBSTRATE_BLENDING_TRANSLUCENT_GREYTRANSMITTANCE || SUBSTRATE_BLENDING_TRANSLUCENT_COLOREDTRANSMITTANCE) && REFRACTION_USE_INDEX_OF_REFRACTION
+#if MATERIALBLENDING_TRANSLUCENT && (SUBSTRATE_BLENDING_TRANSLUCENT_GREYTRANSMITTANCE || SUBSTRATE_BLENDING_TRANSLUCENT_COLOREDTRANSMITTANCE) && PATH_TRACER_REFRACTION_ENABLED
 							Ior = GetRefractionIor(PixelMaterialInputs, F0);
 #endif
 						}
@@ -1409,7 +1414,7 @@ RAY_TRACING_ENTRY_ANY_HIT(PathTracingMaterialAHS,
 		float3 Transmission = exp(-LocalSigmaT * SUBSTRATE_SIMPLEVOLUME_THICKNESS_CM);
 		float3 Tr = ComputeThinSlabWeights(Transmission, NoV, Ior, F0).Transmitted;
 #else
-		HLSL_STATIC_ASSERT(REFRACTION_USE_INDEX_OF_REFRACTION == 1, "Expect refraction to be true here");
+		HLSL_STATIC_ASSERT(PATH_TRACER_REFRACTION_ENABLED == 1, "Expect refraction to be true here");
 		// Solid Glass Fake Caustics
 		float Tr = FakeCaustics(F0, Ior, NoV);
 		// Account for transmission 
@@ -1417,7 +1422,7 @@ RAY_TRACING_ENTRY_ANY_HIT(PathTracingMaterialAHS,
 #endif
 
 		// If material is doing refraction, apply an extra factor to account for ApproximateCaustics
-#if REFRACTION_USE_INDEX_OF_REFRACTION
+#if PATH_TRACER_REFRACTION_ENABLED
 		if (Ior > 0.0)
 		{
 			Tr *= lerp(FakeCausticsPathFactor(RoughnessData.x, PathRoughness),
@@ -1426,7 +1431,7 @@ RAY_TRACING_ENTRY_ANY_HIT(PathTracingMaterialAHS,
 #endif
 		const float3 Transparency = saturate(lerp(1.0, Tr, TotalCoverage));
 
-#else // MATERIALBLENDING_TRANSLUCENT && (REFRACTION_USE_INDEX_OF_REFRACTION || MATERIAL_ISTHINSURFACE)
+#else // MATERIALBLENDING_TRANSLUCENT && (PATH_TRACER_REFRACTION_ENABLED || MATERIAL_ISTHINSURFACE)
 		// All other blend modes, do simple transparency handling
 		const float3 Transparency = saturate(lerp(1.0, TotalTransmittancePreCoverage, TotalCoverage));
 #endif
@@ -1462,7 +1467,7 @@ RAY_TRACING_ENTRY_ANY_HIT(PathTracingMaterialAHS,
 		if (ShadingModelID == SHADINGMODELID_THIN_TRANSLUCENT)
 		{
 			float3 Transmission = GetThinTranslucentMaterialOutput0(MaterialParameters);
-#if REFRACTION_USE_INDEX_OF_REFRACTION
+#if PATH_TRACER_REFRACTION_ENABLED
 			float Ior = GetRefractionIor(PixelMaterialInputs);
 #else
 			float Ior = 0.0;
@@ -1488,7 +1493,7 @@ RAY_TRACING_ENTRY_ANY_HIT(PathTracingMaterialAHS,
 				Transparency *= ComputeThinSlabWeights(Transmission, VoN, Ior, F0).Transmitted;
 
 				// fake caustic approximation (see comments below)
-#if REFRACTION_USE_INDEX_OF_REFRACTION
+#if PATH_TRACER_REFRACTION_ENABLED
 				if (Ior > 0.0)
 				{
 					Transparency *= FakeCausticsPathFactor(Roughness, PathRoughness);
@@ -1502,7 +1507,7 @@ RAY_TRACING_ENTRY_ANY_HIT(PathTracingMaterialAHS,
 		}
 		else
 #endif // MATERIAL_SHADINGMODEL_THIN_TRANSLUCENT
-#if MATERIAL_SHADINGMODEL_DEFAULT_LIT && REFRACTION_USE_INDEX_OF_REFRACTION
+#if MATERIAL_SHADINGMODEL_DEFAULT_LIT && PATH_TRACER_REFRACTION_ENABLED
 		if (ShadingModelID == SHADINGMODELID_DEFAULT_LIT)
 		{
 			// Is refraction enabled?
diff --git a/Engine/Source/Runtime/Engine/Private/Materials/Material.cpp b/Engine/Source/Runtime/Engine/Private/Materials/Material.cpp
--- a/Engine/Source/Runtime/Engine/Private/Materials/Material.cpp
+++ b/Engine/Source/Runtime/Engine/Private/Materials/Material.cpp
@@ -7613,6 +7613,7 @@ static bool IsPropertyActive_Internal(EMaterialProperty InProperty,
 	EBlendMode BlendMode,
 	FMaterialShadingModelField ShadingModels,
 	ETranslucencyLightingMode TranslucencyLightingMode,
+	ERefractionMode RefractionMethod,
 	bool bIsTessellationEnabled,
 	bool bBlendableOutputAlpha,
 	bool bUsesDistortion,
@@ -7766,7 +7767,7 @@ static bool IsPropertyActive_Internal(EMaterialProperty InProperty,
 			switch (InProperty)
 			{
 			case MP_Refraction:
-				Active = (bIsTranslucentBlendMode && !IsAlphaHoldoutBlendMode(BlendMode) && !IsModulateBlendMode(BlendMode) && bUsesDistortion) || ShadingModels.HasShadingModel(MSM_SingleLayerWater);
+				Active = (bIsTranslucentBlendMode && !IsAlphaHoldoutBlendMode(BlendMode) && !IsModulateBlendMode(BlendMode) && bUsesDistortion && RefractionMethod != RM_IndexOfRefractionFromF0) || ShadingModels.HasShadingModel(MSM_SingleLayerWater);
 				break;
 			case MP_Opacity:
 				// Opacity is used as alpha override for alpha composite blending. 
@@ -7813,7 +7814,7 @@ static bool IsPropertyActive_Internal(EMaterialProperty InProperty,
 			Active = false;
 			break;
 		case MP_Refraction:
-			Active = (bIsTranslucentBlendMode && !IsAlphaHoldoutBlendMode(BlendMode) && !IsModulateBlendMode(BlendMode) && bUsesDistortion) || ShadingModels.HasShadingModel(MSM_SingleLayerWater);
+			Active = (bIsTranslucentBlendMode && !IsAlphaHoldoutBlendMode(BlendMode) && !IsModulateBlendMode(BlendMode) && bUsesDistortion && RefractionMethod != RM_IndexOfRefractionFromF0) || ShadingModels.HasShadingModel(MSM_SingleLayerWater);
 			break;
 		case MP_Opacity:
 			Active = (bIsTranslucentBlendMode && !IsModulateBlendMode(BlendMode)) || ShadingModels.HasShadingModel(MSM_SingleLayerWater);
@@ -7902,6 +7903,7 @@ bool UMaterial::IsPropertyActiveInEditor(EMaterialProperty InProperty) const
 		BlendMode,
 		ShadingModels,
 		TranslucencyLightingMode,
+		RefractionMethod,
 		IsTessellationEnabled(),
 		IsPostProcessMaterialOutputingAlpha(),
 		bUsesDistortion,
@@ -7931,6 +7933,7 @@ bool UMaterial::IsPropertyActiveInDerived(EMaterialProperty InProperty, const UM
 		DerivedMaterial->GetBlendMode(),
 		DerivedMaterial->GetShadingModels(),
 		TranslucencyLightingMode,
+		RefractionMethod,
 		DerivedMaterial->IsTessellationEnabled(),
 		IsPostProcessMaterialOutputingAlpha(),
 		bUsesDistortion,
diff --git a/Engine/Source/Runtime/Engine/Private/Materials/MaterialShared.cpp b/Engine/Source/Runtime/Engine/Private/Materials/MaterialShared.cpp
--- a/Engine/Source/Runtime/Engine/Private/Materials/MaterialShared.cpp
+++ b/Engine/Source/Runtime/Engine/Private/Materials/MaterialShared.cpp
@@ -6117,7 +6117,7 @@ void FinalCompileValidation(
 		{
 			OutErrors.Add(TEXT("AbsorptionMedium custom output requires Default Lit shading model."));
 		}
-		if (Material->RefractionMethod != RM_IndexOfRefraction)
+		if (Material->RefractionMethod != RM_IndexOfRefraction && Material->RefractionMethod != RM_IndexOfRefractionFromF0)
 		{
 			OutErrors.Add(TEXT("AbsorptionMedium custom output requires \"Index of Refraction\" as the \"Refraction Mode\"."));
 		}
