[Substrate] Fixed blendable GBuffer Substrate SLW applying the indirect lighting component twice.

diff --git a/Engine/Shaders/Private/BasePassPixelShader.usf b/Engine/Shaders/Private/BasePassPixelShader.usf
--- a/Engine/Shaders/Private/BasePassPixelShader.usf
+++ b/Engine/Shaders/Private/BasePassPixelShader.usf
@@ -1933,7 +1933,8 @@ void FPixelShaderInOut_MainPS(
 		}
 	#endif
 
-		if (MATERIAL_SUBSTRATE_OPAQUE_PRECOMPUTED_LIGHTING && GBuffer.ShadingModelID != SHADINGMODELID_UNLIT)
+		// Water currently takes the legacy path for forward shading, so we need to skip it here so we don't apply this term twice
+		if (MATERIAL_SUBSTRATE_OPAQUE_PRECOMPUTED_LIGHTING && GBuffer.ShadingModelID != SHADINGMODELID_UNLIT && SUBSTRATE_INLINE_SINGLELAYERWATER == 0)
 		{
 			float3 InputBentNormal = MaterialParameters.WorldNormal;
 			// Clear Coat Bottom Normal
