Fixed local fog volume computing wrong UV when sampling depth in splitscreen on mobile renderer

#rnx

diff --git a/Engine/Shaders/Private/LocalFogVolumes/LocalFogVolumeSplat.usf b/Engine/Shaders/Private/LocalFogVolumes/LocalFogVolumeSplat.usf
--- a/Engine/Shaders/Private/LocalFogVolumes/LocalFogVolumeSplat.usf
+++ b/Engine/Shaders/Private/LocalFogVolumes/LocalFogVolumeSplat.usf
@@ -85,11 +85,10 @@ void LocalFogVolumeTiledPS(
 #else
 	const float2 HalfResScale = 1.0f;
 #endif
-	const float2 ViewRectSVPos = SVPos.xy * HalfResScale;
-	const float2 SourcePixelPos = ViewRectSVPos + ResolvedView.ViewRectMin.xy;
+	const float2 ViewRectSVPos = SVPos.xy * HalfResScale;  
 
 #if SHADING_PATH_MOBILE
-	float2 UVDepth = SourcePixelPos * ResolvedView.BufferSizeAndInvSize.zw;
+	float2 UVDepth = ViewRectSVPos * ResolvedView.BufferSizeAndInvSize.zw;
 	float DeviceZ = LookupDeviceZ(UVDepth); // LookupDeviceZ with uint2(SVPos.xy) as input does not work on mobile.
 #else
 	float DeviceZ = LookupDeviceZ(uint2(SVPos.xy));
