Fix virtual texture debug views for screen res scaling.
#rnx

diff --git a/Engine/Shaders/Private/VisualizeVirtualTexture.usf b/Engine/Shaders/Private/VisualizeVirtualTexture.usf
--- a/Engine/Shaders/Private/VisualizeVirtualTexture.usf
+++ b/Engine/Shaders/Private/VisualizeVirtualTexture.usf
@@ -30,7 +30,7 @@ void Main(
 	OutColor.xyz = SceneColorLuminance;
 	OutColor.w = 1;
 
-	uint2 PixelPos = (uint2)SvPosition.xy;
+	uint2 PixelPos = (uint2)(SvPosition.xy * (float2)ViewSize / Input_ViewportSize);
 	uint ReadIndex = PixelPos.y * ViewSize.x + PixelPos.x;
 	uint PackedDebugValue = DebugBuffer[ReadIndex];
 	uint DebugValue = 0;
