Substrate - help the compiler to unroll a for into a simple if when there is a single BSDF.

#rnx
#tests PC and console
#DisabledValidations CheckCopyright

diff --git a/Engine/Shaders/Private/Substrate/SubstrateForwardLighting.ush b/Engine/Shaders/Private/Substrate/SubstrateForwardLighting.ush
--- a/Engine/Shaders/Private/Substrate/SubstrateForwardLighting.ush
+++ b/Engine/Shaders/Private/Substrate/SubstrateForwardLighting.ush
@@ -173,8 +173,13 @@ float3 SubstrateForwardLighting(
 		//
 		// Forward lighting
 		//
+	#if SUBSTRATE_CLAMPED_CLOSURE_COUNT==1
+		int BSDFIdx = 0;
+		if(SubstratePixelHeader.SubstrateTree.BSDFCount > 0)
+	#else
 		SUBSTRATE_UNROLL_N(SUBSTRATE_CLAMPED_CLOSURE_COUNT)
 		for (int BSDFIdx = 0; BSDFIdx < SubstratePixelHeader.SubstrateTree.BSDFCount; ++BSDFIdx)
+	#endif
 		{
 		#define CurrentBSDF SubstratePixelHeader.SubstrateTree.BSDFs[BSDFIdx]
 
