Fix ShaderPrint incorrectly rendering hidden primitives.

diff --git a/Engine/Shaders/Private/ShaderPrintDrawPrimitive.usf b/Engine/Shaders/Private/ShaderPrintDrawPrimitive.usf
--- a/Engine/Shaders/Private/ShaderPrintDrawPrimitive.usf
+++ b/Engine/Shaders/Private/ShaderPrintDrawPrimitive.usf
@@ -101,6 +101,9 @@ void ShaderDrawDebugPS(
 	const float Depth = DepthTexture.SampleLevel(DepthSampler, DepthUV, 0);
 	const bool bIsHidden = Position.z < Depth; // Reverse-Z
 
+	if (bIsHidden)
+		discard;
+
 	const float ColorScale = bIsHidden ? 0.4f : 1;
 	const float CheckerboardScale = ((PixelCoord.x%2) ^ (PixelCoord.y%2)) > 0 ? 1 : 0.0;
 	const float OcclusionScale = (bCheckerboardEnabled && bIsHidden) ? (bDrawOccludedLines ? CheckerboardScale : 0.0f) : 1;
