Fix incorrect MRTs output when using unlit material with Format=0.

diff --git a/Engine/Shaders/Private/BasePassPixelShader.usf b/Engine/Shaders/Private/BasePassPixelShader.usf
--- a/Engine/Shaders/Private/BasePassPixelShader.usf
+++ b/Engine/Shaders/Private/BasePassPixelShader.usf
@@ -2308,7 +2308,7 @@ void FPixelShaderInOut_MainPS(
 
 			EncodeGBufferToMRT(Out, GBuffer, QuantizationBias);
 
-			if (GBuffer.ShadingModelID == SHADINGMODELID_UNLIT && !SUBSTRATE_ENABLED) // Do not touch what Substrate outputs
+			if (GBuffer.ShadingModelID == SHADINGMODELID_UNLIT && (!SUBSTRATE_ENABLED || SUBSTRATE_GBUFFER_FORMAT==0)) // Do not touch what Substrate outputs with Format=1
 			{
 				Out.MRT[1] = 0;
 				SetGBufferForUnlit(Out.MRT[2]);
