Substrate - follow up to 45318925. Using Substrate_for_unroll now. Next: update all similare spots.

#rnx

diff --git a/Engine/Shaders/Private/Substrate/SubstrateForwardLighting.ush b/Engine/Shaders/Private/Substrate/SubstrateForwardLighting.ush
--- a/Engine/Shaders/Private/Substrate/SubstrateForwardLighting.ush
+++ b/Engine/Shaders/Private/Substrate/SubstrateForwardLighting.ush
@@ -173,13 +173,7 @@ float3 SubstrateForwardLighting(
 		//
 		// Forward lighting
 		//
-	#if SUBSTRATE_CLAMPED_CLOSURE_COUNT==1
-		int BSDFIdx = 0;
-		if(SubstratePixelHeader.SubstrateTree.BSDFCount > 0)
-	#else
-		SUBSTRATE_UNROLL_N(SUBSTRATE_CLAMPED_CLOSURE_COUNT)
-		for (int BSDFIdx = 0; BSDFIdx < SubstratePixelHeader.SubstrateTree.BSDFCount; ++BSDFIdx)
-	#endif
+		Substrate_for_unroll(int BSDFIdx = 0, BSDFIdx < SubstratePixelHeader.SubstrateTree.BSDFCount, ++BSDFIdx)
 		{
 		#define CurrentBSDF SubstratePixelHeader.SubstrateTree.BSDFs[BSDFIdx]
 
