Substrate - fix for MFP and vertical layer operator simplification.

#tests pc editor dx12
#rnx

diff --git a/Engine/Shaders/Private/Substrate/Substrate.ush b/Engine/Shaders/Private/Substrate/Substrate.ush
--- a/Engine/Shaders/Private/Substrate/Substrate.ush
+++ b/Engine/Shaders/Private/Substrate/Substrate.ush
@@ -3174,13 +3174,13 @@ FSubstrateData SubstrateVerticalLayeringParameterBlending(in FSubstrateData Top,
 	const float3 NewMFP					= 1.0f / NewExtinction;
 #else
 	// MFP looks more perceptually linear even though it is not perfect
-	const float3 NewMFP = Info2.TransmittanceOnlyBottom * BaseBSDF.TmpMFP + Info2.TransmittanceOnlyTop * TopBSDF.TmpMFP + Info2.TransmittanceTopAndBottom * min(TopBSDF.TmpMFP, BaseBSDF.TmpMFP);
+	const float3 NewMFP					= Info2.TransmittanceOnlyBottom* BaseBSDF.TmpMFP + Info2.TransmittanceOnlyTop * TopBSDF.TmpMFP + Info2.TransmittanceTopAndBottom * min(TopBSDF.TmpMFP, BaseBSDF.TmpMFP);
 #endif
 
-	ResultBSDF.TmpMFP				= NewMFP;
-	SLAB_SSSMFP(ResultBSDF)			= NewMFP;
+	ResultBSDF.Coverage					= MaxCoverage;
+	ResultBSDF.TmpMFP					= NewMFP;
+	ResultBSDF.EnableSlabBSDFSimpleVolumetric();	// Derive Slab MFP from TmpMFP and thickness. Due to approximation above (thickness & MFP based combination), it is not a perfect match.
 
-	ResultBSDF.Coverage				= MaxCoverage;
 #endif
 
 #endif
