Small instruction/occupancy improvement by excluding SSS_TYPE_SIMPLEVOLUME from SUBSTRATE_GBUFFER_FORMAT 0.

diff --git a/Engine/Shaders/Private/Substrate/SubstrateEvaluation.ush b/Engine/Shaders/Private/Substrate/SubstrateEvaluation.ush
--- a/Engine/Shaders/Private/Substrate/SubstrateEvaluation.ush
+++ b/Engine/Shaders/Private/Substrate/SubstrateEvaluation.ush
@@ -64,6 +64,8 @@
 #define SUBSTRATE_USE_FULLYSIMPLIFIED_MATERIAL 0
 #endif
 
+#define SUBSTRATE_SUPPORTS_SIMPLEVOLUME (MATERIAL_IS_SUBSTRATE || SUBSTRATE_GBUFFER_FORMAT==1)
+
 ///////////////////////////////////////////////////////////////////////////////
 // Substrate Macros for reducing permutation boiler plate
 // 
@@ -985,7 +987,7 @@ FSubstrateEvaluateResult SubstrateEvaluateBSDFCommon(FSubstrateBSDFContext BSDFC
 			//// This is used for optically thin translucent objects, or non-bottom layer of a material.
 			////
 
-			#if (SUBSTRATE_FASTPATH==0 || MATERIALBLENDING_ANY_TRANSLUCENT) // !Fast path or rendering translucent materials
+			#if ((SUBSTRATE_FASTPATH==0 && SUBSTRATE_SUPPORTS_SIMPLEVOLUME) || MATERIALBLENDING_ANY_TRANSLUCENT) // !Fast path or rendering translucent materials
 			BRANCH
 			if (BSDF_GETSSSTYPE(BSDFContext.BSDF) == SSS_TYPE_SIMPLEVOLUME)
 			{
@@ -1616,6 +1618,7 @@ FSubstrateEnvLightResult SubstrateEvaluateForEnvLight(FSubstrateBSDFContext BSDF
 			const bool bIsSubsurfaceDiffusionMFP = BSDF_GETSSSTYPE(BSDFContext.BSDF) == SSS_TYPE_DIFFUSION;
 			SubstrateEnvLightResult.bPostProcessSubsurface = (bIsSubsurfaceDiffusionProfile || bIsSubsurfaceDiffusionMFP) && SUBSTRATE_SSS_MATERIAL_OVERRIDE;
 
+			#if SUBSTRATE_SUPPORTS_SIMPLEVOLUME
 			BRANCH
 			if (BSDF_GETSSSTYPE(BSDFContext.BSDF) == SSS_TYPE_SIMPLEVOLUME)
 			{
@@ -1641,6 +1644,7 @@ FSubstrateEnvLightResult SubstrateEvaluateForEnvLight(FSubstrateBSDFContext BSDF
 				// We keep the same normal
 			}
 			#endif
+			#endif
 			
 
 			BRANCH
