Substrate - fixed distortion not showing up with blendable gbuffer.

#rnx
#tests PC editor dx12

#DisabledValidations Preflight

diff --git a/Engine/Shaders/Private/DistortAccumulatePS.usf b/Engine/Shaders/Private/DistortAccumulatePS.usf
--- a/Engine/Shaders/Private/DistortAccumulatePS.usf
+++ b/Engine/Shaders/Private/DistortAccumulatePS.usf
@@ -32,7 +32,7 @@
 #endif // ADAPTIVE_VOLUMETRIC_SHADOW_MAP
 
 // SUBSTRATE_TODO enable distortion on mobile.
-#if SUBSTRATE_ENABLED && (!MATERIAL_IS_SUBSTRATE || SHADING_PATH_MOBILE)
+#if SUBSTRATE_ENABLED && SHADING_PATH_MOBILE
 #undef SUBSTRATE_ENABLED
 #define SUBSTRATE_ENABLED 0
 #endif
@@ -70,6 +70,11 @@ void Main(
 	StereoSetupPS(StereoInput);
 	const uint EyeIndex = GetEyeIndex(StereoInput);
 
+#if SUBSTRATE_ENABLED
+	OutVarianceCoverage = float2(0.0f, 0.0f);
+	OutClosestDepthMeter = MaxHalfFloat;
+#endif
+
 	// material parameter inputs
 	FMaterialPixelParameters MaterialParameters = GetMaterialPixelParameters(Interpolants, SvPosition);
 	FPixelMaterialInputs PixelMaterialInputs;
@@ -191,7 +196,7 @@ void Main(
 	// Sample depth at distortion offset
 	float2 DistortBufferUV = ScreenUV + BufferUVDistortion;
 
-#else // SUBSTRATE_ENABLED
+#else // SUBSTRATE_ENABLED && TEMPLATE_USES_SUBSTRATE
 
 	float3 BaseColor = GetMaterialBaseColor(PixelMaterialInputs);
 	float  Metallic = GetMaterialMetallic(PixelMaterialInputs);
@@ -208,8 +213,17 @@ void Main(
 	// Sample depth at distortion offset
 	float2 DistortBufferUV = ScreenUV + BufferUVDistortion;
 
+#if DISTORTION_ACCOUNT_FOR_COVERAGE
+	const float RefractionCoverage = GetMaterialOpacity(PixelMaterialInputs);	// Coverage of the material.
+#else
+	const float RefractionCoverage = 1.0f;										// Behavior where coverage is always ignored (legacy support)
+#endif
+	const float RefractionLobeVariance = 0.0f;
+#if SUBSTRATE_ENABLED
+	OutVarianceCoverage = float2(RefractionLobeVariance, RefractionCoverage);
+#endif
 
-#endif // SUBSTRATE_ENABLED
+#endif // SUBSTRATE_ENABLED && TEMPLATE_USES_SUBSTRATE
 
 	float DistortSceneDepth = CalcSceneDepth(DistortBufferUV);
 
