Substrate - make sure coverage is accounted for correctly only for Substrate materials for legacy compatibility reasons.
Clarification of coverage.

#rnx
#tests PC editor dx12

#DisabledValidations Preflight

diff --git a/Engine/Shaders/Private/Substrate/SubstrateTree.ush b/Engine/Shaders/Private/Substrate/SubstrateTree.ush
--- a/Engine/Shaders/Private/Substrate/SubstrateTree.ush
+++ b/Engine/Shaders/Private/Substrate/SubstrateTree.ush
@@ -189,10 +189,10 @@ void FSubstrateTree::UpdateSingleOperatorCoverageTransmittance(int OpIndex)
 		Op.ThroughputAlongV		= OpA.ThroughputAlongV;
 		Op.TransmittanceAlongN	= OpA.TransmittanceAlongN;
 
-		Op.TopDownRefractionLobeStat = SubstrateWeightLobe(OpA.TopDownRefractionLobeStat, Weight);
+		Op.TopDownRefractionLobeStat = SubstrateWeightLobe(OpA.TopDownRefractionLobeStat, Weight);		// This is used for rough refraction and we weight energy and roughness to make sure coverage from weight node can make rough refraction vanish.
 		Op.TopDownRefractionWorldNormal = OpA.TopDownRefractionWorldNormal * Weight;
 
-		Op.BottomUpRefractionLobeStat = SubstrateWeightLobe(OpA.BottomUpRefractionLobeStat, Weight);
+		Op.BottomUpRefractionLobeStat = SubstrateWeightLobe(OpA.BottomUpRefractionLobeStat, Weight);	// This is used for opaue rough refraction, we scale the lobe to make it vanish in order to have a smooth continuous blur behavior.
 
 		// Coverage does not affect rough refraction data
 		Op.OpaqueRoughRefractThicknessCm  = OpA.OpaqueRoughRefractThicknessCm;
diff --git a/Engine/Source/Runtime/Core/Private/UObject/FortniteMainBranchObjectVersion.cpp b/Engine/Source/Runtime/Core/Private/UObject/FortniteMainBranchObjectVersion.cpp
--- a/Engine/Source/Runtime/Core/Private/UObject/FortniteMainBranchObjectVersion.cpp
+++ b/Engine/Source/Runtime/Core/Private/UObject/FortniteMainBranchObjectVersion.cpp
@@ -13,7 +13,7 @@ TMap<FGuid, FGuid> FFortniteMainBranchObjectVersion::GetSystemGuids()
 	SystemGuids.Add(DevGuids.MATERIALSHADERMAP_DERIVEDDATA_VER, FGuid("F051877974B44D409435F321BDB39DBB"));
 	SystemGuids.Add(DevGuids.NIAGARASHADERMAP_DERIVEDDATA_VER, FGuid("6360A977062842A29ED30E3A7ACB0E64"));
 	SystemGuids.Add(DevGuids.Niagara_LatestScriptCompileVersion, FGuid("783FF3700BB641C799B2D43552FBF91E"));
-	SystemGuids.Add(DevGuids.MaterialTranslationDDCVersion, FGuid("0EE6461792BB4D3C885E06E5CE03B068"));
+	SystemGuids.Add(DevGuids.MaterialTranslationDDCVersion, FGuid("287723BF015541FB90E91D2F0AF0AD5E"));
 
 	// These GUIDs are in AppendRestrictedFortniteMainStreamObjectVersionGuids as they should rarely be changed.
 	//SystemGuids.Add(DevGuids.NANITE_DERIVEDDATA_VER, FGuid("ACF98184978742A592537DA223A1E6BE"));
diff --git a/Engine/Source/Runtime/Engine/Private/Materials/HLSLMaterialTranslator.cpp b/Engine/Source/Runtime/Engine/Private/Materials/HLSLMaterialTranslator.cpp
--- a/Engine/Source/Runtime/Engine/Private/Materials/HLSLMaterialTranslator.cpp
+++ b/Engine/Source/Runtime/Engine/Private/Materials/HLSLMaterialTranslator.cpp
@@ -16050,7 +16050,7 @@ void FHLSLMaterialTranslator::PrepareEnvironmentDefines()
 	// Distortion uses tangent space transform 
 	EnvironmentDefines->bUsesDistortion = Material->IsDistorted();
 	EnvironmentDefines->bUsesExplicitDerivatives = bUsesExplicitDerivatives;
-	EnvironmentDefines->bDistortionAccountForCoverage = bSubstrateEnabled && Material->GetRefractionCoverageMode() == RCM_CoverageAccountedFor ? 1 : 0;
+	EnvironmentDefines->bDistortionAccountForCoverage = bSubstrateEnabled && !bMaterialUsesRootNodeToSubstrateHiddenConversion && Material->GetRefractionCoverageMode() == RCM_CoverageAccountedFor ? 1 : 0;
 	EnvironmentDefines->bMaterialEnableTranslucencyFogging = Material->ShouldApplyFogging();
 	EnvironmentDefines->bMaterialEnableTranslucencyCloudFogging = Material->ShouldApplyCloudFogging();
 	EnvironmentDefines->bMaterialIsSky = Material->IsSky();
diff --git a/Engine/Source/Runtime/Engine/Private/Materials/Material.cpp b/Engine/Source/Runtime/Engine/Private/Materials/Material.cpp
--- a/Engine/Source/Runtime/Engine/Private/Materials/Material.cpp
+++ b/Engine/Source/Runtime/Engine/Private/Materials/Material.cpp
@@ -5197,7 +5197,7 @@ bool UMaterial::CanEditChange(const FProperty* InProperty) const
 
 		if (PropertyName == GET_MEMBER_NAME_STRING_CHECKED(UMaterial, RefractionCoverageMode))
 		{
-			return bSubstrateEnabled;
+			return bSubstrateEnabled && !bSubstrateHiddenMaterialAssetConversion;
 		}
 
 		if (PropertyName == GET_MEMBER_NAME_STRING_CHECKED(UMaterial, bHasPixelAnimation))
diff --git a/Engine/Source/Runtime/Engine/Public/Materials/Material.h b/Engine/Source/Runtime/Engine/Public/Materials/Material.h
--- a/Engine/Source/Runtime/Engine/Public/Materials/Material.h
+++ b/Engine/Source/Runtime/Engine/Public/Materials/Material.h
@@ -1118,7 +1118,7 @@ public:
 	UPROPERTY(EditAnywhere, Category=Refraction)
 	TEnumAsByte<ERefractionMode> RefractionMethod;
 
-	/** Controls whether refraction takes into account the material surface coverage, or not. */
+	/** Controls whether refraction takes into account the material surface coverage, or not. This is only for when substrate FrontMaterial is plugged in, not legacy root node material for the sake of compatibility. */
 	UPROPERTY(EditAnywhere, Category=Refraction)
 	TEnumAsByte<ERefractionCoverageMode> RefractionCoverageMode;
 
