Substrate - clariying TEMPLATE_USES_SUBSTRATE and TEMPLATE_USES_LEGACYGBUFFERDATA.

#rnx
#tests PC editor dx12 with RT

#DisabledValidations Preflight

diff --git a/Engine/Shaders/Private/MaterialTemplate.ush b/Engine/Shaders/Private/MaterialTemplate.ush
--- a/Engine/Shaders/Private/MaterialTemplate.ush
+++ b/Engine/Shaders/Private/MaterialTemplate.ush
@@ -174,7 +174,11 @@
 	#define ALLOW_CLIP 1
 #endif
 
+// Material shader uses Front material or root node is converted to Substrate data.
 #define TEMPLATE_USES_SUBSTRATE					(SUBSTRATE_ENABLED && MATERIAL_IS_SUBSTRATE)
+// Material shader is NOT exported to Substrate data (no FrontMaterial or SubstrateTree data). It uses legacy GBuffer data.
+#define TEMPLATE_USES_LEGACYGBUFFERDATA			(SUBSTRATE_ENABLED && !MATERIAL_IS_SUBSTRATE)
+
 #define SUBSTRATE_OPAQUE_MATERIAL				(TEMPLATE_USES_SUBSTRATE && ((!SUBSTRATE_MATERIAL_EXPORT_EXECUTED && MATERIALBLENDING_ANY_TRANSLUCENT<=0) || SUBSTRATE_MATERIAL_EXPORT_FROM_OPAQUE))
 #define SUBSTRATE_TRANSLUCENT_MATERIAL			(TEMPLATE_USES_SUBSTRATE && ((!SUBSTRATE_MATERIAL_EXPORT_EXECUTED && MATERIALBLENDING_ANY_TRANSLUCENT>0 ) || SUBSTRATE_MATERIAL_EXPORT_FROM_TRANSLUCENT))
 #define SUBSTRATE_OPAQUE_DEFERRED				(SUBSTRATE_OPAQUE_MATERIAL && !FORWARD_SHADING)
@@ -185,7 +189,7 @@
 // Premultiplied alpha override is only enabled when alpha composite blending is used and the root node input is fed.
 #define SUBSTRATE_USE_PREMULTALPHA_OVERRIDE	(SUBSTRATE_ENABLED && MATERIALBLENDING_ALPHACOMPOSITE && SUBSTRATE_PREMULTIPLIED_ALPHA_OPACITY_OVERRIDEN)
 
-#if TEMPLATE_USES_SUBSTRATE
+#if TEMPLATE_USES_SUBSTRATE || TEMPLATE_USES_LEGACYGBUFFERDATA
 #define SUBSTRATE_INLINE_SHADING 1
 #endif
 
diff --git a/Engine/Shaders/Private/PathTracing/PathTracingMaterialHitShader.usf b/Engine/Shaders/Private/PathTracing/PathTracingMaterialHitShader.usf
--- a/Engine/Shaders/Private/PathTracing/PathTracingMaterialHitShader.usf
+++ b/Engine/Shaders/Private/PathTracing/PathTracingMaterialHitShader.usf
@@ -21,12 +21,7 @@
 // Ensure that SSS albedo comes through in the material
 #define SUBSTRATE_SSS_MATERIAL_OVERRIDE 0
 
-// This define happens when we compile material with BlendableGBuffer (legacy) and no substrate nodes are used in a materials.
-// As oftoday, we are then going to compile the material with Substrate disabled. But the payload is still a Substrate payload!
-// So to quickly and simply convert from LEgacy material into a Substrate payload, We create a slab from SubstrateConvertLegacyMaterialStatic
-// that is then exported to a path tracer material using the recular code path. Only sampling from a Single BSDF instead of the substrate tree of BSDFs.
-#define SUBSTRATE_CONVERT_FROM_LEGACY (SUBSTRATE_ENABLED && !TEMPLATE_USES_SUBSTRATE)
-#if SUBSTRATE_CONVERT_FROM_LEGACY
+#if (SUBSTRATE_ENABLED && !TEMPLATE_USES_SUBSTRATE) // == TEMPLATE_USES_LEGACYGBUFFERDATA from MaterialTemplate, needs to happen here before other #include in this file.
 #define SUBSTRATE_INLINE_SHADING 1
 #endif
 
@@ -49,11 +44,15 @@
 #include "/Engine/Private/PathTracing/PathTracingCommon.ush"
 #include "/Engine/Private/PathTracing/Material/PathTracingFresnel.ush"
 
-#if SUBSTRATE_CONVERT_FROM_LEGACY
+// This define happens when we compile material with BlendableGBuffer (legacy) data and no substrate nodes are used or hidden conversion happening.
+// As of today, we are then going to compile the material with Substrate disabled. But the payload is still a Substrate payload!
+// So to quickly and simply convert from LEgacy material into a Substrate payload, We create a slab from SubstrateConvertLegacyMaterialStatic
+// that is then exported to a path tracer material using the recular code path. Only sampling from a Single BSDF instead of the substrate tree of BSDFs.
+#if TEMPLATE_USES_LEGACYGBUFFERDATA
 #include "/Engine/Private/Substrate/SubstrateLegacyConversion.ush"
 #endif
 
-#if (SUBSTRATE_ENABLED && TEMPLATE_USES_SUBSTRATE) || SUBSTRATE_CONVERT_FROM_LEGACY
+#if (SUBSTRATE_ENABLED && TEMPLATE_USES_SUBSTRATE) || TEMPLATE_USES_LEGACYGBUFFERDATA
 #define MATERIAL_SUBSTRATE_OPAQUE_PRECOMPUTED_LIGHTING 0
 #include "/Engine/Private/Substrate/SubstrateExport.ush"
 #include "/Engine/Private/PathTracing/Material/PathTracingSubstrateCommon.ush"
@@ -272,7 +271,7 @@ RAY_TRACING_ENTRY_CLOSEST_HIT(PathTracingMaterialCHS,
 	const float3 Transparency = 0;
 #endif
 
-#elif SUBSTRATE_CONVERT_FROM_LEGACY
+#elif TEMPLATE_USES_LEGACYGBUFFERDATA
 
 #if SUBSTRATE_BLENDING_COLOREDTRANSMITTANCEONLY
 	const uint BSDFCount = 0;
@@ -307,7 +306,7 @@ RAY_TRACING_ENTRY_CLOSEST_HIT(PathTracingMaterialCHS,
 
 	FPathTracingPayload Payload = (FPathTracingPayload)0;
 
-#if !(SUBSTRATE_ENABLED && TEMPLATE_USES_SUBSTRATE) || SUBSTRATE_CONVERT_FROM_LEGACY
+#if !(SUBSTRATE_ENABLED && TEMPLATE_USES_SUBSTRATE) || TEMPLATE_USES_LEGACYGBUFFERDATA
 	/**
 	 * Set common material attributes for both full and simplified materials
 	 **/
@@ -358,7 +357,7 @@ RAY_TRACING_ENTRY_CLOSEST_HIT(PathTracingMaterialCHS,
 #error Unknown material blending mode
 #endif
 
-#endif //  !(SUBSTRATE_ENABLED && TEMPLATE_USES_SUBSTRATE) || SUBSTRATE_CONVERT_FROM_LEGACY
+#endif //  !(SUBSTRATE_ENABLED && TEMPLATE_USES_SUBSTRATE) || TEMPLATE_USES_LEGACYGBUFFERDATA
 
 	// fetch primitive flags only once
 	// TODO: would be nice to keep this inside MaterialParameters as it is also needed there as well
@@ -397,7 +396,7 @@ RAY_TRACING_ENTRY_CLOSEST_HIT(PathTracingMaterialCHS,
 	Payload.WorldSmoothNormal = AdjustShadingNormal(Payload.WorldSmoothNormal, Payload.WorldGeoNormal, WorldRayDirection());
 #endif
 
-#if SUBSTRATE_ENABLED && TEMPLATE_USES_SUBSTRATE  || SUBSTRATE_CONVERT_FROM_LEGACY
+#if SUBSTRATE_ENABLED && TEMPLATE_USES_SUBSTRATE  || TEMPLATE_USES_LEGACYGBUFFERDATA
 	// unless proven otherwise, we are a basic unlit material
 	Payload.ShadingModelID = SHADINGMODELID_UNLIT;
 
@@ -414,7 +413,7 @@ RAY_TRACING_ENTRY_CLOSEST_HIT(PathTracingMaterialCHS,
 
 #else // SUBSTRATE_OPTIMIZED_UNLIT
 
-#if !SUBSTRATE_CONVERT_FROM_LEGACY // In this case we must keep BSDFOpacity and TransparencyColor evaluated above.
+#if !TEMPLATE_USES_LEGACYGBUFFERDATA // In this case we must keep BSDFOpacity and TransparencyColor evaluated above.
 	Payload.BSDFOpacity = 1;
 #if !SUBSTRATE_OPAQUE_MATERIAL
 	Payload.TransparencyColor = Transparency;
@@ -422,7 +421,7 @@ RAY_TRACING_ENTRY_CLOSEST_HIT(PathTracingMaterialCHS,
 #endif
 
 
-#if SUBSTRATE_CONVERT_FROM_LEGACY
+#if TEMPLATE_USES_LEGACYGBUFFERDATA
 
 	Payload.ShadingModelID = GetMaterialShadingModel(PixelMaterialInputs);
 
@@ -597,16 +596,16 @@ RAY_TRACING_ENTRY_CLOSEST_HIT(PathTracingMaterialCHS,
 		1,//dummy uint ClearCoatBottomNormal_SharedLocalBasisIndex,
 		SharedLocalBasisTypes);
 
-#endif // SUBSTRATE_CONVERT_FROM_LEGACY
+#endif // TEMPLATE_USES_LEGACYGBUFFERDATA
 
 
 
 
 
 
-	if (BSDFCount > 0 || SUBSTRATE_CONVERT_FROM_LEGACY)
+	if (BSDFCount > 0 || TEMPLATE_USES_LEGACYGBUFFERDATA)
 	{
-#if SUBSTRATE_CONVERT_FROM_LEGACY
+#if TEMPLATE_USES_LEGACYGBUFFERDATA
 		{
 	#if HAIR_STRAND_MESH_FACTORY
 			Payload.WorldSmoothNormal = Payload.WorldGeoNormal;
@@ -620,7 +619,7 @@ RAY_TRACING_ENTRY_CLOSEST_HIT(PathTracingMaterialCHS,
 	#endif
 
 		}
-#else // SUBSTRATE_CONVERT_FROM_LEGACY
+#else // TEMPLATE_USES_LEGACYGBUFFERDATA
 		{
 	#if HAIR_STRAND_MESH_FACTORY
 			Payload.WorldSmoothNormal = Payload.WorldGeoNormal;
@@ -639,10 +638,10 @@ RAY_TRACING_ENTRY_CLOSEST_HIT(PathTracingMaterialCHS,
 			}
 	#endif
 		}
-#endif // SUBSTRATE_CONVERT_FROM_LEGACY
+#endif // TEMPLATE_USES_LEGACYGBUFFERDATA
 
 		int BSDFIdx = 0;
-#if SUBSTRATE_CONVERT_FROM_LEGACY
+#if TEMPLATE_USES_LEGACYGBUFFERDATA
 #define BSDF SubstrateData.InlinedBSDF
 		BSDF.SubstrateSanitizeBSDF();
 		BSDF.LuminanceWeightV = 1;
@@ -650,12 +649,13 @@ RAY_TRACING_ENTRY_CLOSEST_HIT(PathTracingMaterialCHS,
 		BSDF.TransmittanceAboveAlongN = 1;
 		BSDF.bIsBottom = 1;
 		BSDF.bIsTop = 1;
-		
+	#if SUBSTRATE_INLINE_SHADING
 		Payload.Radiance = BSDF.Emissive * Payload.BSDFOpacity;
+	#endif
 #else
 #define BSDF SubstratePixelHeader.SubstrateTree.BSDFs[BSDFIdx]
 
-	#if !SUBSTRATE_CONVERT_FROM_LEGACY
+	#if !TEMPLATE_USES_LEGACYGBUFFERDATA
 		// Accumulate radiance across all BSDFs
 		Payload.Radiance = 0.0;
 		for (BSDFIdx = 0; BSDFIdx < BSDFCount; ++BSDFIdx)
@@ -665,7 +665,7 @@ RAY_TRACING_ENTRY_CLOSEST_HIT(PathTracingMaterialCHS,
 	#endif
 #endif
 
-#if SUBSTRATE_CLAMPED_CLOSURE_COUNT > 1 && !SUBSTRATE_CONVERT_FROM_LEGACY
+#if SUBSTRATE_CLAMPED_CLOSURE_COUNT > 1 && !TEMPLATE_USES_LEGACYGBUFFERDATA
 		// The material might be using more than one slab -- insert code to stochastically pick one 
 		FRISContext SlabSampler = InitRISContext(PackedPayload.GetStochasticSlabRand());
 		for (BSDFIdx = 0; BSDFIdx < BSDFCount; ++BSDFIdx)
@@ -686,7 +686,7 @@ RAY_TRACING_ENTRY_CLOSEST_HIT(PathTracingMaterialCHS,
 					// Is the slab even visible at all?
 					if (any(MaxLobeWeight > 0.0))
 					{
-					#if SUBSTRATE_CONVERT_FROM_LEGACY
+					#if TEMPLATE_USES_LEGACYGBUFFERDATA
 						const float3 N = WorldNormal;
 					#else
 						const float3 N = SubstratePixelHeader.SharedLocalBases.Normals[BSDF_GETSHAREDLOCALBASISID(BSDF)];
@@ -795,7 +795,7 @@ RAY_TRACING_ENTRY_CLOSEST_HIT(PathTracingMaterialCHS,
 							Payload.BSDFOpacity = 1.0 / Contrib;
 							#endif
 							Payload.WorldNormal = N;
-						#if SUBSTRATE_CONVERT_FROM_LEGACY
+						#if TEMPLATE_USES_LEGACYGBUFFERDATA
 							if (SharedLocalBasisTypes == SUBSTRATE_BASIS_TYPE_TANGENT)
 							{
 								Payload.WorldTangent = WorldTangent;
@@ -945,7 +945,7 @@ RAY_TRACING_ENTRY_CLOSEST_HIT(PathTracingMaterialCHS,
 				// NOTE: The cases below are for back-compatibility and cannot be layered
 				case SUBSTRATE_BSDF_TYPE_EYE:
 				{
-				#if SUBSTRATE_CONVERT_FROM_LEGACY
+				#if TEMPLATE_USES_LEGACYGBUFFERDATA
 					Payload.WorldNormal = WorldNormal;
 				#else
 					Payload.WorldNormal = SubstratePixelHeader.SharedLocalBases.Normals[BSDF_GETSHAREDLOCALBASISID(BSDF)];
@@ -995,7 +995,7 @@ RAY_TRACING_ENTRY_CLOSEST_HIT(PathTracingMaterialCHS,
 				}
 				case SUBSTRATE_BSDF_TYPE_SINGLELAYERWATER:
 				{
-				#if SUBSTRATE_CONVERT_FROM_LEGACY
+				#if TEMPLATE_USES_LEGACYGBUFFERDATA
 					Payload.WorldNormal = WorldNormal;
 				#else
 					Payload.WorldNormal = SubstratePixelHeader.SharedLocalBases.Normals[BSDF_GETSHAREDLOCALBASISID(BSDF)];
@@ -1054,7 +1054,7 @@ RAY_TRACING_ENTRY_CLOSEST_HIT(PathTracingMaterialCHS,
 	}
 #endif
 
-#else // SUBSTRATE_ENABLED && TEMPLATE_USES_SUBSTRATE  || SUBSTRATE_CONVERT_FROM_LEGACY 
+#else // SUBSTRATE_ENABLED && TEMPLATE_USES_SUBSTRATE  || TEMPLATE_USES_LEGACYGBUFFERDATA 
 	Payload.Radiance = GetMaterialEmissive(PixelMaterialInputs);
 
 	Payload.WorldNormal = MaterialParameters.WorldNormal;
@@ -1448,7 +1448,7 @@ RAY_TRACING_ENTRY_CLOSEST_HIT(PathTracingMaterialCHS,
 
 #endif // SUBSTRATE_ENABLED && TEMPLATE_USES_SUBSTRATE 
 
-	// Only opaque surfaces can recieve decals. This is both for compatibility with raster and improving performance when we have stacks of transparent surfaces
+	// Only opaque surfaces can receive decals. This is both for compatibility with raster and improving performance when we have stacks of transparent surfaces
 #if MATERIALBLENDING_SOLID || MATERIALBLENDING_MASKED
 	if ((PrimitiveFlags & PRIMITIVE_SCENE_DATA_FLAG_DECAL_RECEIVER) != 0)
 	{
diff --git a/Engine/Shaders/Private/VolumetricCloud.usf b/Engine/Shaders/Private/VolumetricCloud.usf
--- a/Engine/Shaders/Private/VolumetricCloud.usf
+++ b/Engine/Shaders/Private/VolumetricCloud.usf
@@ -18,11 +18,6 @@
 #include "VolumetricCloudCommon.ush"
 #include "/Engine/Shared/EnvironmentComponentsFlags.h"
 
-#if SUBSTRATE_ENABLED && !MATERIAL_IS_SUBSTRATE
-#undef SUBSTRATE_ENABLED
-#define SUBSTRATE_ENABLED 0
-#endif
-
 #ifndef SHADER_RENDERVIEW_CS
 #define SHADER_RENDERVIEW_CS 0
 #endif
@@ -111,7 +106,7 @@ void VolumetricCloudCalcPixelMaterialInputs(in out FMaterialPixelParameters Para
 float3 SampleExtinctionCoefficients(in FPixelMaterialInputs PixelMaterialInputs)
 {
 	float3 Extinction = 0.0f;
-#if SUBSTRATE_ENABLED
+#if TEMPLATE_USES_SUBSTRATE
 	FSubstrateBSDF BSDF = PixelMaterialInputs.FrontMaterial.InlinedBSDF;
 	Extinction = VOLUMETRICFOGCLOUD_EXTINCTION(BSDF).rgb;
 #else
@@ -125,7 +120,7 @@ float3 SampleExtinctionCoefficients(in FPixelMaterialInputs PixelMaterialInputs)
 float3 SampleEmissive(in FPixelMaterialInputs PixelMaterialInputs)
 {
 	float3 EmissiveColor = 0.0f;
-#if SUBSTRATE_ENABLED
+#if TEMPLATE_USES_SUBSTRATE
 	FSubstrateBSDF BSDF = PixelMaterialInputs.FrontMaterial.InlinedBSDF;
 	EmissiveColor = BSDF_GETEMISSIVE(BSDF).rgb;
 #else
@@ -137,7 +132,7 @@ float3 SampleEmissive(in FPixelMaterialInputs PixelMaterialInputs)
 float3 SampleAlbedo(in FPixelMaterialInputs PixelMaterialInputs)
 {
 	float3 Albedo = 0.0f;
-#if SUBSTRATE_ENABLED
+#if TEMPLATE_USES_SUBSTRATE
 	FSubstrateBSDF BSDF = PixelMaterialInputs.FrontMaterial.InlinedBSDF;
 	Albedo = VOLUMETRICFOGCLOUD_ALBEDO(BSDF).rgb;
 #else
@@ -151,7 +146,7 @@ float3 SampleAlbedo(in FPixelMaterialInputs PixelMaterialInputs)
 float SampleAmbientOcclusion(in FPixelMaterialInputs PixelMaterialInputs)
 {
 	float AO = 0.0f;
-#if SUBSTRATE_ENABLED
+#if TEMPLATE_USES_SUBSTRATE
 	FSubstrateBSDF BSDF = PixelMaterialInputs.FrontMaterial.InlinedBSDF;
 	AO = VOLUMETRICFOGCLOUD_AO(BSDF);
 #else
@@ -177,7 +172,7 @@ void ConvertCloudPixelMaterialInputsToWorkingColorSpace(inout FPixelMaterialInpu
 	float3 AlbedoWCS		= SRGBColorToWorkingColorSpace(AlbedoSRGB);
 	float3 EmissiveWCS		= SRGBColorToWorkingColorSpace(EmissiveSRGB);
 
-#if SUBSTRATE_ENABLED	
+#if TEMPLATE_USES_SUBSTRATE	
 	VOLUMETRICFOGCLOUD_EXTINCTION(PixelMaterialInputs.FrontMaterial.InlinedBSDF)	= ExtinctionWCS;
 	VOLUMETRICFOGCLOUD_ALBEDO(PixelMaterialInputs.FrontMaterial.InlinedBSDF)		= AlbedoWCS;
 	BSDF_GETEMISSIVE(PixelMaterialInputs.FrontMaterial.InlinedBSDF)					= EmissiveWCS;
diff --git a/Engine/Shaders/Private/VolumetricFogVoxelization.usf b/Engine/Shaders/Private/VolumetricFogVoxelization.usf
--- a/Engine/Shaders/Private/VolumetricFogVoxelization.usf
+++ b/Engine/Shaders/Private/VolumetricFogVoxelization.usf
@@ -15,11 +15,6 @@
 
 #include "Common.ush"
 
-#if SUBSTRATE_ENABLED && !MATERIAL_IS_SUBSTRATE
-#undef SUBSTRATE_ENABLED
-#define SUBSTRATE_ENABLED 0
-#endif
-
 #define SceneTexturesStruct VoxelizeVolumePass.SceneTextures 
 
 #include "/Engine/Generated/Material.ush"
@@ -34,7 +29,7 @@ uint VoxelizationPassIndex;
 float SampleExtinctionCoefficients(in FPixelMaterialInputs PixelMaterialInputs)
 {
 	float Extinction = 0.0f;
-#if SUBSTRATE_ENABLED
+#if TEMPLATE_USES_SUBSTRATE
 	FSubstrateBSDF BSDF = PixelMaterialInputs.FrontMaterial.InlinedBSDF;
 	Extinction = VOLUMETRICFOGCLOUD_EXTINCTION(BSDF).r;
 #else
@@ -48,7 +43,7 @@ float SampleExtinctionCoefficients(in FPixelMaterialInputs PixelMaterialInputs)
 float3 SampleEmissive(in FPixelMaterialInputs PixelMaterialInputs)
 {
 	float3 EmissiveColor = 0.0f;
-#if SUBSTRATE_ENABLED
+#if TEMPLATE_USES_SUBSTRATE
 	FSubstrateBSDF BSDF = PixelMaterialInputs.FrontMaterial.InlinedBSDF;
 	EmissiveColor = BSDF_GETEMISSIVE(BSDF);
 #else
@@ -60,7 +55,7 @@ float3 SampleEmissive(in FPixelMaterialInputs PixelMaterialInputs)
 float3 SampleAlbedo(in FPixelMaterialInputs PixelMaterialInputs, in ViewState InputView)
 {
 	float3 Albedo = 0.0f;
-#if SUBSTRATE_ENABLED
+#if TEMPLATE_USES_SUBSTRATE
 	FSubstrateBSDF BSDF = PixelMaterialInputs.FrontMaterial.InlinedBSDF;
 	Albedo = VOLUMETRICFOGCLOUD_ALBEDO(BSDF);
 #else
