Fix Details Lighting view with Substrate Format=0

diff --git a/Engine/Shaders/Private/Substrate/SubstrateRead.ush b/Engine/Shaders/Private/Substrate/SubstrateRead.ush
--- a/Engine/Shaders/Private/Substrate/SubstrateRead.ush
+++ b/Engine/Shaders/Private/Substrate/SubstrateRead.ush
@@ -70,8 +70,10 @@ FSubstrateGBufferBSDF SubstrateReadGBufferBSDF(in FScreenSpaceData ScreenSpaceDa
 		float		ThicknessCm = 1.0f;
 		bool		bIsThinSurface = false;
 
-		float3 DiffuseAlbedo = ComputeDiffuseAlbedo(GBuffer.BaseColor, GBuffer.Metallic);
-		float3 F0 = ComputeF0(GBuffer.Specular, GBuffer.BaseColor, GBuffer.Metallic);
+		// Directly use DiffuseColor and SpecularColor (instead of ComputeDiffuseAlbedo() / ComputeF0() with BaseColor, Metallic, Specular), 
+		// as they contain Diffuse/Specular override which are used for view modes.
+		float3 DiffuseAlbedo = GBuffer.DiffuseColor;
+		float3 F0 = GBuffer.SpecularColor;
 		float3 F90 = 1.0f;
 
 		float ClearCoatUseSecondNormal = 0.0f;
@@ -105,7 +107,7 @@ FSubstrateGBufferBSDF SubstrateReadGBufferBSDF(in FScreenSpaceData ScreenSpaceDa
 			IrisNormal = OctahedronToUnitVector(GBuffer.CustomData.yz * 2 - 1);
 			IrisDistance = GBuffer.StoredMetallic;
 		#endif
-			BSDF = GetSubstrateEyeBSDF(GBuffer.BaseColor, DielectricSpecularToF0(GBuffer.Specular), GBuffer.Roughness, IrisMask, IrisDistance, IrisNormal, IrisPlaneNormal, SSSProfileID, Emissive, SharedLocalBasisIndex).InlinedBSDF;
+			BSDF = GetSubstrateEyeBSDF(DiffuseAlbedo, F0, GBuffer.Roughness, IrisMask, IrisDistance, IrisNormal, IrisPlaneNormal, SSSProfileID, Emissive, SharedLocalBasisIndex).InlinedBSDF;
 		}
 		else
 		{
