Fix incorrect WorldNormal view mode with unlit material when using Substrate Format=0.

diff --git a/Engine/Shaders/Private/MaterialTemplate.ush b/Engine/Shaders/Private/MaterialTemplate.ush
--- a/Engine/Shaders/Private/MaterialTemplate.ush
+++ b/Engine/Shaders/Private/MaterialTemplate.ush
@@ -3307,37 +3307,33 @@ float4 SceneTextureLookup(float2 UV, int SceneTextureIndex, bool bFiltered)
 	#elif SUBSTRATE_ENABLED && (SUBSTRATE_GBUFFER_FORMAT==0 || SUBSTRATE_GBUFFER_FORMAT==1) && !SHADING_PATH_MOBILE // Not on mobile that is not using dbuffer decal but the legacy GBuffer blending.
 	{
 		#if !SUBSTRATE_OPAQUE_DEFERRED || ALLOW_SUBSTRATE_PUBLIC_ACCESS // Prevents opaque material to sample the Gbuffer.
-		FSubstrateGBuffer SubstrateGBuffer = SubstratePublic_GetSubstrateGBufferFromFirstBSDF(PixelPos);
-		BRANCH
-		if (SubstrateGBuffer.bIsValid > 0)
+		FSubstrateGBuffer SubstrateGBuffer = SubstratePublic_GetSubstrateGBufferFromFirstBSDF(PixelPos, false/*bGetNormalizedNormal*/);
+		switch (SceneTextureIndex)
 		{
-			switch (SceneTextureIndex)
-			{
-				// order needs to match to ESceneTextureId
-				case PPI_DiffuseColor:			return float4(SubstrateGBuffer.DiffuseColor, 0);
-				case PPI_SpecularColor:			return float4(SubstrateGBuffer.SpecularColor, 0);
-				case PPI_SubsurfaceColor:		return float4(SubstrateGBuffer.SubsurfaceColor, 0);
-				case PPI_BaseColor:				return float4(SubstrateGBuffer.BaseColor, 0);
-				case PPI_Specular:				return SubstrateGBuffer.Specular;
-				case PPI_Metallic:				return SubstrateGBuffer.Metallic;
-				case PPI_WorldNormal:			return float4(SubstrateGBuffer.WorldNormal, 0);
-				case PPI_Opacity:				return SubstrateGBuffer.Opacity;
-				case PPI_Roughness:				return SubstrateGBuffer.Roughness;
-				case PPI_MaterialAO:			return SubstrateGBuffer.MaterialAO;
-				case PPI_ShadingModelColor:		return float4(SubstrateGBuffer.ShadingModelColor, 1);
-				case PPI_ShadingModelID:		return float4(SubstrateGBuffer.ShadingModelID, 0, 0, 0);
-				case PPI_StoredBaseColor:		return float4(SubstrateGBuffer.StoredBaseColor, 0);
-				case PPI_StoredSpecular:		return float4(SubstrateGBuffer.StoredSpecular.rrr, 0);
-				case PPI_WorldTangent:			return float4(SubstrateGBuffer.WorldTangent, 0);
-				case PPI_Anisotropy:			return SubstrateGBuffer.Anisotropy;
-				case PPI_IsFirstPerson:			return SubstrateGBuffer.bIsFirstPerson ? 1.0f : 0.0f;
-			}
+			// order needs to match to ESceneTextureId
+			case PPI_DiffuseColor:			return float4(SubstrateGBuffer.DiffuseColor, 0);
+			case PPI_SpecularColor:			return float4(SubstrateGBuffer.SpecularColor, 0);
+			case PPI_SubsurfaceColor:		return float4(SubstrateGBuffer.SubsurfaceColor, 0);
+			case PPI_BaseColor:				return float4(SubstrateGBuffer.BaseColor, 0);
+			case PPI_Specular:				return SubstrateGBuffer.Specular;
+			case PPI_Metallic:				return SubstrateGBuffer.Metallic;
+			case PPI_WorldNormal:			return float4(SubstrateGBuffer.WorldNormal, 0);
+			case PPI_Opacity:				return SubstrateGBuffer.Opacity;
+			case PPI_Roughness:				return SubstrateGBuffer.Roughness;
+			case PPI_MaterialAO:			return SubstrateGBuffer.MaterialAO;
+			case PPI_ShadingModelColor:		return float4(SubstrateGBuffer.ShadingModelColor, 1);
+			case PPI_ShadingModelID:		return float4(SubstrateGBuffer.ShadingModelID, 0, 0, 0);
+			case PPI_StoredBaseColor:		return float4(SubstrateGBuffer.StoredBaseColor, 0);
+			case PPI_StoredSpecular:		return float4(SubstrateGBuffer.StoredSpecular.rrr, 0);
+			case PPI_WorldTangent:			return float4(SubstrateGBuffer.WorldTangent, 0);
+			case PPI_Anisotropy:			return SubstrateGBuffer.Anisotropy;
+			case PPI_IsFirstPerson:			return SubstrateGBuffer.bIsFirstPerson ? 1.0f : 0.0f;
 		}
 		#endif
 	}
 	#else // SUBSTRATE_ENABLED
 	{
-		FScreenSpaceData ScreenSpaceData = GetScreenSpaceData(UV, false);
+		FScreenSpaceData ScreenSpaceData = GetScreenSpaceData(UV, false/*bGetNormalizedNormal*/);
 		switch (SceneTextureIndex)
 		{
 			// order needs to match to ESceneTextureId
diff --git a/Engine/Shaders/Public/SubstratePublic.ush b/Engine/Shaders/Public/SubstratePublic.ush
--- a/Engine/Shaders/Public/SubstratePublic.ush
+++ b/Engine/Shaders/Public/SubstratePublic.ush
@@ -42,7 +42,7 @@ struct FSubstrateGBuffer
 	bool   bIsFirstPerson;
 };
 
-FSubstrateGBuffer SubstratePublic_GetSubstrateGBufferFromFirstBSDF(uint2 PixelCoord)
+FSubstrateGBuffer SubstratePublic_GetSubstrateGBufferFromFirstBSDF(uint2 PixelCoord, bool bGetNormalizedNormal = true)
 {
 	FSubstrateGBuffer Out = (FSubstrateGBuffer)0; 
 	Out.bIsValid = false;
@@ -117,7 +117,7 @@ FSubstrateGBuffer SubstratePublic_GetSubstrateGBufferFromFirstBSDF(uint2 PixelCo
 	}
 	#elif SUBSTRATE_GBUFFER_FORMAT == 0
 	const float2 UV = (float2(PixelCoord) + 0.5) * View.BufferSizeAndInvSize.zw;
-	const FGBufferData GBuffer = GetScreenSpaceData(UV).GBuffer;
+	const FGBufferData GBuffer = GetScreenSpaceData(UV, bGetNormalizedNormal).GBuffer;
 	Out.bIsValid			= GBuffer.ShadingModelID != SHADINGMODELID_UNLIT;
 	Out.DiffuseColor		= GBuffer.DiffuseColor;
 	Out.SpecularColor		= GBuffer.SpecularColor;
