Fix HeterogeneousVolumes with Substrate Format=0.

diff --git a/Engine/Shaders/Private/HeterogeneousVolumes/HeterogeneousVolumesLiveShadingUtils.ush b/Engine/Shaders/Private/HeterogeneousVolumes/HeterogeneousVolumesLiveShadingUtils.ush
--- a/Engine/Shaders/Private/HeterogeneousVolumes/HeterogeneousVolumesLiveShadingUtils.ush
+++ b/Engine/Shaders/Private/HeterogeneousVolumes/HeterogeneousVolumesLiveShadingUtils.ush
@@ -36,7 +36,7 @@ FPrimitiveSceneData GetPrimitiveData(FMaterialPixelParameters Parameters, float4
 float3 SampleExtinctionCoefficients(in FPixelMaterialInputs PixelMaterialInputs)
 {
 	float3 Extinction = 0.0f;
-#if SUBSTRATE_ENABLED
+#if TEMPLATE_USES_SUBSTRATE
 	FSubstrateBSDF BSDF = PixelMaterialInputs.FrontMaterial.InlinedBSDF;
 	Extinction = VOLUMETRICFOGCLOUD_EXTINCTION(BSDF).rgb;
 #else
@@ -50,7 +50,7 @@ float3 SampleExtinctionCoefficients(in FPixelMaterialInputs PixelMaterialInputs)
 float3 SampleEmissive(in FPixelMaterialInputs PixelMaterialInputs)
 {
 	float3 EmissiveColor = 0.0f;
-#if SUBSTRATE_ENABLED
+#if TEMPLATE_USES_SUBSTRATE
 	FSubstrateBSDF BSDF = PixelMaterialInputs.FrontMaterial.InlinedBSDF;
 	EmissiveColor = BSDF_GETEMISSIVE(BSDF);
 #else
@@ -62,7 +62,7 @@ float3 SampleEmissive(in FPixelMaterialInputs PixelMaterialInputs)
 float3 SampleAlbedo(in FPixelMaterialInputs PixelMaterialInputs)
 {
 	float3 Albedo = 0.0f;
-#if SUBSTRATE_ENABLED
+#if TEMPLATE_USES_SUBSTRATE
 	FSubstrateBSDF BSDF = PixelMaterialInputs.FrontMaterial.InlinedBSDF;
 	Albedo = VOLUMETRICFOGCLOUD_ALBEDO(BSDF);
 #else
