Substrate - fixed shader cook error. No need for inline data when !TEMPLATE_USES_SUBSTRATE in fact.

#rnx
#tests PC editor dx12

#DisabledValidations Preflight

diff --git a/Engine/Shaders/Private/MaterialTemplate.ush b/Engine/Shaders/Private/MaterialTemplate.ush
--- a/Engine/Shaders/Private/MaterialTemplate.ush
+++ b/Engine/Shaders/Private/MaterialTemplate.ush
@@ -189,7 +189,7 @@
 // Premultiplied alpha override is only enabled when alpha composite blending is used and the root node input is fed.
 #define SUBSTRATE_USE_PREMULTALPHA_OVERRIDE	(SUBSTRATE_ENABLED && MATERIALBLENDING_ALPHACOMPOSITE && SUBSTRATE_PREMULTIPLIED_ALPHA_OPACITY_OVERRIDEN)
 
-#if TEMPLATE_USES_SUBSTRATE || TEMPLATE_USES_LEGACYGBUFFERDATA
+#if TEMPLATE_USES_SUBSTRATE
 #define SUBSTRATE_INLINE_SHADING 1
 #endif
 
