Shader build fix for virtual texture lightmap permutations.
#rnx

diff --git a/Engine/Shaders/Private/BasePassPixelShader.usf b/Engine/Shaders/Private/BasePassPixelShader.usf
--- a/Engine/Shaders/Private/BasePassPixelShader.usf
+++ b/Engine/Shaders/Private/BasePassPixelShader.usf
@@ -946,7 +946,7 @@ void FPixelShaderInOut_MainPS(
 #if LIGHTMAP_VT_ENABLED
 	// This must occur after CalcMaterialParameters(), which is required to initialize the VT feedback mechanism
 	// Lightmap request is always the first VT sample in the shader
-	StoreVirtualTextureFeedback(MaterialParameters.VirtualTextureFeedback, 0, LightmapVTPageTableResult.PackedRequest, LightmapVTPageTableResult.PendingMips);
+	StoreVirtualTextureFeedback(MaterialParameters.VirtualTextureFeedback, LightmapVTPageTableResult.PackedRequest, LightmapVTPageTableResult.PendingMips);
 #endif
 
 #if USE_EDITOR_COMPOSITING && (FEATURE_LEVEL >= FEATURE_LEVEL_SM4 || MOBILE_EMULATION)
diff --git a/Engine/Shaders/Private/MobileBasePassPixelShader.usf b/Engine/Shaders/Private/MobileBasePassPixelShader.usf
--- a/Engine/Shaders/Private/MobileBasePassPixelShader.usf
+++ b/Engine/Shaders/Private/MobileBasePassPixelShader.usf
@@ -618,7 +618,7 @@ void Main(
 #if LIGHTMAP_VT_ENABLED
 	// This must occur after CalcMaterialParameters(), which is required to initialize the VT feedback mechanism
 	// Lightmap request is always the first VT sample in the shader
-	StoreVirtualTextureFeedback(MaterialParameters.VirtualTextureFeedback, 0, LightmapVTPageTableResult.PackedRequest, LightmapVTPageTableResult.PendingMips);
+	StoreVirtualTextureFeedback(MaterialParameters.VirtualTextureFeedback, LightmapVTPageTableResult.PackedRequest, LightmapVTPageTableResult.PendingMips);
 #endif
 
 #if TEMPLATE_USES_SUBSTRATE && !SUBSTRATE_INLINE_SINGLELAYERWATER
