Fixed height/weight alpha precision loss when merging landscape

diff --git a/Engine/Shaders/Private/Landscape/LandscapeCommon.ush b/Engine/Shaders/Private/Landscape/LandscapeCommon.ush
--- a/Engine/Shaders/Private/Landscape/LandscapeCommon.ush
+++ b/Engine/Shaders/Private/Landscape/LandscapeCommon.ush
@@ -51,9 +51,11 @@ float2 PackHeightAlpha(float InAlpha, uint InAlphaFlags)
 void UnpackHeightAlpha(float2 InPackedAlpha, out float OutAlpha, out uint OutFlags)
 {
 	uint AlphaWithFlags = ((uint)round(InPackedAlpha.r * 255.0f) << 8) | (uint)round(InPackedAlpha.g * 255.0f);
-	uint Alpha = AlphaWithFlags & 0xFFFC; // First 14 bits are used to store alpha (0xFFFC in float)
-	OutAlpha = float(Alpha) / 65535.0f; 
-	OutFlags = AlphaWithFlags & 0x3; // Last 2 bits are used to store flags
+	// The first 14 bits are used to store alpha
+	uint ExtractedAlpha = (AlphaWithFlags >> 2);
+	OutAlpha = (float)ExtractedAlpha / (float)((1 << 14) - 1);
+	// The last 2 bits are used to store flags
+	OutFlags = AlphaWithFlags & 0x3;
 }
 
 // It's a bit overkill but we have 2 channels available for weight, so we can use 16 bits for storing it
@@ -80,9 +82,11 @@ float2 PackWeightAlpha(float InAlpha, uint InAlphaFlags)
 void UnpackWeightAlpha(float2 InPackedAlpha, out float OutAlpha, out uint OutFlags)
 {
 	uint AlphaWithFlags = ((uint)round(InPackedAlpha.r * 255.0f) << 8) | (uint)round(InPackedAlpha.g * 255.0f);
-	uint Alpha = AlphaWithFlags & 0xFFFC; // First 14 bits are used to store alpha (0xFFFC in float)
-	OutAlpha = float(Alpha) / 65535.0f; 
-	OutFlags = AlphaWithFlags & 0x3; // Last 2 bits are used to store flags
+	// The first 14 bits are used to store alpha
+	uint ExtractedAlpha = (AlphaWithFlags >> 2);
+	OutAlpha = (float)ExtractedAlpha / (float)((1 << 14) - 1);
+	// The last 2 bits are used to store flags
+	OutFlags = AlphaWithFlags & 0x3;
 }
 
 // Compute the normal of the triangle formed by the 3 points (in winding order).
