Fix missing default SSS profile on legacy material when Substrate is enabled.

diff --git a/Engine/Shaders/Private/DebugProbes.usf b/Engine/Shaders/Private/DebugProbes.usf
--- a/Engine/Shaders/Private/DebugProbes.usf
+++ b/Engine/Shaders/Private/DebugProbes.usf
@@ -127,6 +127,7 @@ void WriteOutMaterial(
 		/*float Roughness						*/Roughness,
 		/*float Anisotropy						*/0.0f,
 		/*uint SSSProfileID						*/0,
+		/*bool bSupportDefaultSSSProfile		*/false,
 		/*float3 SSSMFP							*/0.0f,
 		/*float SSSMFPScale						*/0.0f,
 		/*float SSSPhaseAniso					*/0.0f,
diff --git a/Engine/Shaders/Private/Substrate/Substrate.ush b/Engine/Shaders/Private/Substrate/Substrate.ush
--- a/Engine/Shaders/Private/Substrate/Substrate.ush
+++ b/Engine/Shaders/Private/Substrate/Substrate.ush
@@ -1529,7 +1529,7 @@ FSubstrateData InternalGetSubstrateSlabBSDF(
 	float3 RawNormal,
 	float3 DiffuseAlbedo, float3 F0, float3 F90,
 	float Roughness, float Anisotropy,
-	uint SSSProfileID, float3 SSSMFP, float SSSMFPScale, float SSSPhaseAnisotropy, uint SSSType,
+	uint SSSProfileID, bool bSupportDefaultSSSProfile, float3 SSSMFP, float SSSMFPScale, float SSSPhaseAnisotropy, uint SSSType,
 	float3 Emissive, 
 	float SecondRoughness, float SecondRoughnessWeight, float SecondRoughnessAsSimpleClearCoat,
 	float ClearCoatUseSecondNormal, float3 ClearCoatBottomNormal,
@@ -1550,8 +1550,8 @@ FSubstrateData InternalGetSubstrateSlabBSDF(
 	FSubstrateData SubstrateData = GetInitialisedSubstrateData();
 
 	SSSMFPScale = saturate(SSSMFPScale);
-	const bool bSSSProfilePlugged = SSSProfileID > 0;
-	const bool bMFPPlugged = any(SSSMFP > 0);
+	const bool bSSSProfilePlugged = SSSProfileID > (bSupportDefaultSSSProfile ? 0u : 1u);
+	const bool bMFPPlugged = any(SSSMFP >= 0);
 	const bool bIsThin = bIsAtTheBottomOfTopology && (bMFPPlugged && bIsThinSurface); // Enable thin only if MFP is plugged, as otherwise there will be no back-tramission
 
 	// In forward mode (blend mode != Opaque/Masked), diffusion is not supported. 
@@ -1735,7 +1735,7 @@ FSubstrateData GetSubstrateSlabBSDF(
 	float3 RawNormal,
 	float3 DiffuseAlbedo, float3 F0, float3 F90,
 	float Roughness, float Anisotropy,
-	uint SSSProfileID, float3 SSSMFP, float SSSMFPScale, float SSSPhaseAnisotropy, uint SSSType,
+	uint SSSProfileID, bool bSupportDefaultSSSProfile, float3 SSSMFP, float SSSMFPScale, float SSSPhaseAnisotropy, uint SSSType,
 	float3 Emissive,
 	float SecondRoughness, float SecondRoughnessWeight, float SecondRoughnessAsSimpleClearCoat,
 	float ClearCoatUseSecondNormal, float3 ClearCoatBottomNormal,
@@ -1753,7 +1753,7 @@ FSubstrateData GetSubstrateSlabBSDF(
 		RawNormal,
 		DiffuseAlbedo, F0, F90,
 		Roughness, Anisotropy,
-		SSSProfileID, SSSMFP, SSSMFPScale, SSSPhaseAnisotropy, SSSType,
+		SSSProfileID, bSupportDefaultSSSProfile, SSSMFP, SSSMFPScale, SSSPhaseAnisotropy, SSSType,
 		Emissive,
 		SecondRoughness, SecondRoughnessWeight, SecondRoughnessAsSimpleClearCoat,
 		ClearCoatUseSecondNormal, ClearCoatBottomNormal,
@@ -1776,7 +1776,7 @@ FSubstrateData GetSubstrateSlabBSDF(
 	float3 RawNormal,
 	float3 DiffuseAlbedo, float3 F0, float3 F90,
 	float Roughness, float Anisotropy,
-	uint SSSProfileID, float3 SSSMFP, float SSSMFPScale, float SSSPhaseAnisotropy, uint SSSType,
+	uint SSSProfileID, bool bSupportDefaultSSSProfile, float3 SSSMFP, float SSSMFPScale, float SSSPhaseAnisotropy, uint SSSType,
 	float3 Emissive,
 	float SecondRoughness, float SecondRoughnessWeight, float SecondRoughnessAsSimpleClearCoat,
 	float ClearCoatUseSecondNormal, float3 ClearCoatBottomNormal,
@@ -1794,7 +1794,7 @@ FSubstrateData GetSubstrateSlabBSDF(
 		RawNormal,
 		DiffuseAlbedo, F0, F90,
 		Roughness, Anisotropy,
-		SSSProfileID, SSSMFP, SSSMFPScale, SSSPhaseAnisotropy, SSSType,
+		SSSProfileID, bSupportDefaultSSSProfile, SSSMFP, SSSMFPScale, SSSPhaseAnisotropy, SSSType,
 		Emissive,
 		SecondRoughness, SecondRoughnessWeight, SecondRoughnessAsSimpleClearCoat,
 		ClearCoatUseSecondNormal, ClearCoatBottomNormal,
@@ -2529,6 +2529,7 @@ void FSubstratePixelHeader::SubstrateConvertFromDBuffer(
 		In.Roughness,		// Roughness
 		0.f,				// Anisotropy
 		0,					// SSSProfileID
+		false, 				// bSupportDefaultSSSProfile
 		0.f,				// SSSMFP
 		0.f,				// SSSMFPScale
 		0.f,				// SSSPhaseAniso
diff --git a/Engine/Shaders/Private/Substrate/SubstrateLegacyConversion.ush b/Engine/Shaders/Private/Substrate/SubstrateLegacyConversion.ush
--- a/Engine/Shaders/Private/Substrate/SubstrateLegacyConversion.ush
+++ b/Engine/Shaders/Private/Substrate/SubstrateLegacyConversion.ush
@@ -121,6 +121,7 @@ FSubstrateData CreateLegacySlab(
 		In.Roughness,										// Roughness
 		In.Anisotropy,										// Anisotropy
 		In.SSSProfileID,									// SSSProfileID
+		true,												// bSupportDefaultSSSProfile
 		In.SSSMFP,											// SSSMFP
 		In.SSSMFPScale,										// SSSMFPScale
 		In.SSSPhaseAniso,									// SSSPhaseAniso
diff --git a/Engine/Shaders/Private/Substrate/SubstrateRead.ush b/Engine/Shaders/Private/Substrate/SubstrateRead.ush
--- a/Engine/Shaders/Private/Substrate/SubstrateRead.ush
+++ b/Engine/Shaders/Private/Substrate/SubstrateRead.ush
@@ -165,7 +165,7 @@ FSubstrateGBufferBSDF SubstrateReadGBufferBSDF(in FScreenSpaceData ScreenSpaceDa
 				GBuffer.WorldNormal,
 				DiffuseAlbedo, F0, F90,
 				GBuffer.Roughness, GBuffer.Anisotropy,
-				SSSProfileID, SSSMFP, SSSMFPScale, SSSPhaseAnisotropy, SSSType,
+				SSSProfileID, true /*bSupportDefaultSSSProfile*/, SSSMFP, SSSMFPScale, SSSPhaseAnisotropy, SSSType,
 				Emissive,
 				SecondRoughness, SecondRoughnessWeight, SecondRoughnessAsSimpleClearCoat,
 				ClearCoatUseSecondNormal, ClearCoatBottomNormal,
@@ -181,7 +181,7 @@ FSubstrateGBufferBSDF SubstrateReadGBufferBSDF(in FScreenSpaceData ScreenSpaceDa
 			SLAB_SSSMFP(BSDF) = SanitizeSSSMFP(SLAB_SSSMFP(BSDF));
 		}
 
-		#if SUBSTRATE_INLINE_SHADING
+		#if SUBSTRATE_INLINE_SHADING && !RAYTRACINGSHADER
 		BSDF.Coverage = 1;
 		BSDF.Emissive = 0;
 		BSDF.ThicknessCm = (SSSType == SSS_TYPE_WRAP || SSSType == SSS_TYPE_TWO_SIDED_WRAP) ? SUBSTRATE_SIMPLEVOLUME_THICKNESS_CM : SUBSTRATE_LAYER_DEFAULT_THICKNESS_CM; // Identity rescaling in InternalGetSubstrateSlabBSDF 
diff --git a/Engine/Source/Runtime/Core/Private/UObject/FortniteMainBranchObjectVersion.cpp b/Engine/Source/Runtime/Core/Private/UObject/FortniteMainBranchObjectVersion.cpp
--- a/Engine/Source/Runtime/Core/Private/UObject/FortniteMainBranchObjectVersion.cpp
+++ b/Engine/Source/Runtime/Core/Private/UObject/FortniteMainBranchObjectVersion.cpp
@@ -13,7 +13,7 @@ TMap<FGuid, FGuid> FFortniteMainBranchObjectVersion::GetSystemGuids()
 	SystemGuids.Add(DevGuids.MATERIALSHADERMAP_DERIVEDDATA_VER, FGuid("F051877974B44D409435F321BDB39DBB"));
 	SystemGuids.Add(DevGuids.NIAGARASHADERMAP_DERIVEDDATA_VER, FGuid("6360A977062842A29ED30E3A7ACB0E64"));
 	SystemGuids.Add(DevGuids.Niagara_LatestScriptCompileVersion, FGuid("783FF3700BB641C799B2D43552FBF91E"));
-	SystemGuids.Add(DevGuids.MaterialTranslationDDCVersion, FGuid("D3CEB448229345B0892BF18F40A6129E"));
+	SystemGuids.Add(DevGuids.MaterialTranslationDDCVersion, FGuid("92421BCD0BFF479F90F4DB0B76C18C55"));
 
 	// These GUIDs are in AppendRestrictedFortniteMainStreamObjectVersionGuids as they should rarely be changed.
 	//SystemGuids.Add(DevGuids.NANITE_DERIVEDDATA_VER, FGuid("ACF98184978742A592537DA223A1E6BE"));
diff --git a/Engine/Source/Runtime/Engine/Private/Materials/HLSLMaterialTranslator.cpp b/Engine/Source/Runtime/Engine/Private/Materials/HLSLMaterialTranslator.cpp
--- a/Engine/Source/Runtime/Engine/Private/Materials/HLSLMaterialTranslator.cpp
+++ b/Engine/Source/Runtime/Engine/Private/Materials/HLSLMaterialTranslator.cpp
@@ -14234,6 +14234,7 @@ int32 FHLSLMaterialTranslator::SubstrateSlabBSDF(
 				/*Roughness*/                        %s,\n\
 				/*Anisotropy*/                       %s,\n\
 				/*SSSProfileId*/                     ExtractSubsurfaceProfileInt(%s),\n\
+				/*bSupportDefaultSSSProfile*/		 false,\n\
 				/*SSSMFP*/                           %s,\n\
 				/*SSSMFPScale*/                      %s,\n\
 				/*SSSPhaseAniso*/                    %s,\n\
@@ -14304,6 +14305,7 @@ int32 FHLSLMaterialTranslator::SubstrateSlabBSDF(
 			/*Roughness*/                        %s,\n\
 			/*Anisotropy*/                       %s,\n\
 			/*SSSProfileId*/                     ExtractSubsurfaceProfileInt(%s),\n\
+			/*bSupportDefaultSSSProfile*/		 false,\n\
 			/*SSSMFP*/                           %s,\n\
 			/*SSSMFPScale*/                      %s,\n\
 			/*SSSPhaseAniso*/                    %s,\n\
