* Fix incorrect fallback for pre-integrated SSS onto SSS profile when using Format=0
* Fix typo for defautl SSS profile

diff --git a/Engine/Shaders/Private/Substrate/Substrate.ush b/Engine/Shaders/Private/Substrate/Substrate.ush
--- a/Engine/Shaders/Private/Substrate/Substrate.ush
+++ b/Engine/Shaders/Private/Substrate/Substrate.ush
@@ -1550,7 +1550,7 @@ FSubstrateData InternalGetSubstrateSlabBSDF(
 	FSubstrateData SubstrateData = GetInitialisedSubstrateData();
 
 	SSSMFPScale = saturate(SSSMFPScale);
-	const bool bSSSProfilePlugged = SSSProfileID > (bSupportDefaultSSSProfile ? 0u : 1u);
+	const bool bSSSProfilePlugged = SSSProfileID >= (bSupportDefaultSSSProfile ? 0u : 1u);
 	const bool bMFPPlugged = any(SSSMFP >= 0);
 	const bool bIsThin = bIsAtTheBottomOfTopology && (bMFPPlugged && bIsThinSurface); // Enable thin only if MFP is plugged, as otherwise there will be no back-tramission
 
diff --git a/Engine/Shaders/Private/Substrate/SubstrateRead.ush b/Engine/Shaders/Private/Substrate/SubstrateRead.ush
--- a/Engine/Shaders/Private/Substrate/SubstrateRead.ush
+++ b/Engine/Shaders/Private/Substrate/SubstrateRead.ush
@@ -117,7 +117,7 @@ FSubstrateGBufferBSDF SubstrateReadGBufferBSDF(in FScreenSpaceData ScreenSpaceDa
 				FuzzAmount = GBuffer.CustomData.a;
 				FuzzRoughness = GBuffer.Roughness;
 			}
-			else if (GBuffer.ShadingModelID == SHADINGMODELID_SUBSURFACE_PROFILE)
+			else if (GBuffer.ShadingModelID == SHADINGMODELID_SUBSURFACE_PROFILE || GBuffer.ShadingModelID == SHADINGMODELID_PREINTEGRATED_SKIN)
 			{
 				SSSType = SSS_TYPE_DIFFUSION;	//  SSS_TYPE_DIFFUSION_PROFILE will be selected in InternalGetSubstrateSlabBSDF based on SSSProfileID and SSSMFPScale.
 				SSSMFP = 0;
