Betterer reservoir sampling maths and comment.
#rnx

diff --git a/Engine/Shaders/Private/VirtualTextureCommon.ush b/Engine/Shaders/Private/VirtualTextureCommon.ush
--- a/Engine/Shaders/Private/VirtualTextureCommon.ush
+++ b/Engine/Shaders/Private/VirtualTextureCommon.ush
@@ -3,7 +3,6 @@
 #pragma once
 
 #include "Common.ush"
-#include "Hash.ush"
 #include "Random.ush"
 #include "GammaCorrectionCommon.ush"
 #ifndef VT_DISABLE_VIEW_UNIFORM_BUFFER
@@ -63,9 +62,14 @@ void StoreVirtualTextureFeedback(in out FVirtualTextureFeedbackParams Params, ui
 {
 #if VT_FEEDBACK_RESERVOIR_SAMPLE
 	// Store the feedback request from a single sample selected by reservoir sampling.
+	// We select the sample if: rand() % (1 + Count) == 0
+	// Or using floating point: frand() <= 1.f / (1.f + Count)
+	// We pre-increment and multiply through by Params.Count to avoid a divide.
+	// We take frand from upper 24 bits of Params.Random to avoid floating point precision issues.
+	// We use a cheap LCG to generate the random value sequence. Offline testing shows this gives sufficient randomness in the final selected sample indices.
 	Params.Count++;
-	Params.Request = (float)Params.Random / (float)(~0u) <= 1.0f / (float)Params.Count ? Request : Params.Request;
-	Params.Random = PCGHash(Params.Random);
+	Params.Request = ((float)(Params.Random >> 8) / 16777215.0f) * (float)Params.Count <= 1.0f ? Request : Params.Request;
+	Params.Random = 1664525u * Params.Random + 1013904223u;
 #elif (NUM_VIRTUALTEXTURE_FEEDBACK_REQUESTS + LIGHTMAP_VT_ENABLED) > 1
 	// Only store the feedback request at the requested sample index.
 	Params.Request = Params.Count == Params.RequestIndex ? Request : Params.Request;
