Fix INF lighting output when using Substrate with material having 0 roughness.

diff --git a/Engine/Shaders/Private/Substrate/SubstrateEvaluation.ush b/Engine/Shaders/Private/Substrate/SubstrateEvaluation.ush
--- a/Engine/Shaders/Private/Substrate/SubstrateEvaluation.ush
+++ b/Engine/Shaders/Private/Substrate/SubstrateEvaluation.ush
@@ -466,7 +466,7 @@ FSubstrateEvaluateResult SubstrateEvaluateBSDFCommon(FSubstrateBSDFContext BSDFC
 			float3 DiffuseColor			= SLAB_DIFFUSEALBEDO(BSDFContext.BSDF);
 			float3 F0					= SLAB_F0(BSDFContext.BSDF);
 			float3 F90					= SLAB_F90(BSDFContext.BSDF);
-			const float SafeRoughness	= MakeRoughnessSafe(SLAB_ROUGHNESS(BSDFContext.BSDF));
+			const float SafeRoughness	= MakeRoughnessSafe(SLAB_ROUGHNESS(BSDFContext.BSDF), SUBSTRATE_MIN_GGX_ROUGHNESS);
 			const bool bHasAnisotropy	= BSDF_GETHASANISOTROPY(BSDFContext.BSDF) && !bIsRectLight && SUPPORTS_ANISOTROPIC_MATERIALS;
 			const bool bHaziness		= BSDF_GETHASHAZINESS(BSDFContext.BSDF);
 
@@ -691,7 +691,7 @@ FSubstrateEvaluateResult SubstrateEvaluateBSDFCommon(FSubstrateBSDFContext BSDFC
 			
 				const FHaziness Haziness = UnpackHaziness(SLAB_HAZINESS(BSDFContext.BSDF));
 				const float HazeWeight = Haziness.Weight;
-				const float HazeSafeRoughness = MakeRoughnessSafe(Haziness.Roughness);
+				const float HazeSafeRoughness = MakeRoughnessSafe(Haziness.Roughness, SUBSTRATE_MIN_GGX_ROUGHNESS);
 				const bool bHazeAsSimpleClearCoat = Haziness.bSimpleClearCoat;
 
 				if (bHazeAsSimpleClearCoat)
@@ -1160,7 +1160,7 @@ FSubstrateEvaluateResult SubstrateEvaluateBSDFCommon(FSubstrateBSDFContext BSDFC
 			float3 DiffuseColor	= EYE_DIFFUSEALBEDO(BSDFContext.BSDF);
 			float3 F0			= EYE_F0(BSDFContext.BSDF);
 			float3 F90			= EYE_F90(BSDFContext.BSDF);
-			const float SafeRoughness	= MakeRoughnessSafe(EYE_ROUGHNESS(BSDFContext.BSDF));
+			const float SafeRoughness	= MakeRoughnessSafe(EYE_ROUGHNESS(BSDFContext.BSDF), SUBSTRATE_MIN_GGX_ROUGHNESS);
 
 			// Specular occlusion is only used here once to affect the F90 source parameters. F0 will decrease naturally to 0.
 			F90 *= F0RGBToMicroOcclusion(F0);
diff --git a/Engine/Shaders/Shared/SubstrateDefinitions.h b/Engine/Shaders/Shared/SubstrateDefinitions.h
--- a/Engine/Shaders/Shared/SubstrateDefinitions.h
+++ b/Engine/Shaders/Shared/SubstrateDefinitions.h
@@ -45,7 +45,9 @@
 // The default thickness of a layer is considered to be 0.01 centimeter = 0.1 millimeter
 #define SUBSTRATE_LAYER_DEFAULT_THICKNESS_CM			0.01f
 
-// Min Fuzz Roughness to avoid numerical issue
+// Min GGX/Fuzz Roughness to avoid numerical issue
+// Avoid INF value when computing GGX's D term (similar to IntegrateBxDF() logic)
+#define SUBSTRATE_MIN_GGX_ROUGHNESS						0.02f
 #define SUBSTRATE_MIN_FUZZ_ROUGHNESS					0.02f
 
 #define SUBSTRATE_BASE_PASS_MRT_OUTPUT_COUNT			3
