Path Tracer: Implement debug view modes specific to the eye shading model

Visualization of these custom signals, for payloads tagged with the eye shading model

#rnx

diff --git a/Engine/Shaders/Private/PathTracing/PathTracingDebug.usf b/Engine/Shaders/Private/PathTracing/PathTracingDebug.usf
--- a/Engine/Shaders/Private/PathTracing/PathTracingDebug.usf
+++ b/Engine/Shaders/Private/PathTracing/PathTracingDebug.usf
@@ -1098,6 +1098,9 @@ FPathTracingResult PathTracingDebugCore(FRayDesc Ray, uint2 LaunchIndex)
 				case PATH_TRACER_DEBUG_VIZ_FUZZ_ROUGHNESS: 	Color = (HitPayload.ShadingModelID == SHADINGMODELID_CLOTH) ? HitPayload.Roughness        : 0.0; break;
 				case PATH_TRACER_DEBUG_VIZ_FUZZ_AMOUNT:    	Color = (HitPayload.ShadingModelID == SHADINGMODELID_CLOTH) ? HitPayload.GetClothAmount() : 0.0; break;
 #endif
+				case PATH_TRACER_DEBUG_VIZ_EYE_IRIS_MASK:      Color = (HitPayload.ShadingModelID == SHADINGMODELID_EYE) ?             HitPayload.GetEyeIrisMask()      : 0.0; break;
+				case PATH_TRACER_DEBUG_VIZ_EYE_IRIS_NORMAL:    Color = (HitPayload.ShadingModelID == SHADINGMODELID_EYE) ? 0.5 + 0.5 * HitPayload.GetEyeIrisNormal()    : 0.0; break;
+				case PATH_TRACER_DEBUG_VIZ_EYE_CAUSTIC_NORMAL: Color = (HitPayload.ShadingModelID == SHADINGMODELID_EYE) ? 0.5 + 0.5 * HitPayload.GetEyeCausticNormal() : 0.0; break;
 				default:
 				{
 					break;
diff --git a/Engine/Shaders/Shared/PathTracingDefinitions.h b/Engine/Shaders/Shared/PathTracingDefinitions.h
--- a/Engine/Shaders/Shared/PathTracingDefinitions.h
+++ b/Engine/Shaders/Shared/PathTracingDefinitions.h
@@ -118,3 +118,6 @@
 #define PATH_TRACER_DEBUG_VIZ_SUBSTRATE_TRANSMITTANCE_N 36
 #define PATH_TRACER_DEBUG_VIZ_SUBSTRATE_COVERAGE_ABOVE_N 37
 #define PATH_TRACER_DEBUG_VIZ_SUBSTRATE_F90 			38
+#define PATH_TRACER_DEBUG_VIZ_EYE_IRIS_MASK             39
+#define PATH_TRACER_DEBUG_VIZ_EYE_CAUSTIC_NORMAL        40
+#define PATH_TRACER_DEBUG_VIZ_EYE_IRIS_NORMAL           41
\ No newline at end of file
diff --git a/Engine/Source/Runtime/Renderer/Private/PathTracing.cpp b/Engine/Source/Runtime/Renderer/Private/PathTracing.cpp
--- a/Engine/Source/Runtime/Renderer/Private/PathTracing.cpp
+++ b/Engine/Source/Runtime/Renderer/Private/PathTracing.cpp
@@ -100,7 +100,9 @@ struct FPathTracingDebugModeInfo
 		REGISTER_VIZ_TYPE(TEXT("SubstrateTransmittanceN"), PATH_TRACER_DEBUG_VIZ_SUBSTRATE_TRANSMITTANCE_N);
 		REGISTER_VIZ_TYPE(TEXT("SubstrateCoverageAboveN"), PATH_TRACER_DEBUG_VIZ_SUBSTRATE_COVERAGE_ABOVE_N);
 		REGISTER_VIZ_TYPE(TEXT("SubstrateF90")           , PATH_TRACER_DEBUG_VIZ_SUBSTRATE_F90			);
-
+		REGISTER_VIZ_TYPE(TEXT("EyeIrisMask")           , PATH_TRACER_DEBUG_VIZ_EYE_IRIS_MASK           );
+		REGISTER_VIZ_TYPE(TEXT("EyeIrisNormal")         , PATH_TRACER_DEBUG_VIZ_EYE_IRIS_NORMAL         );
+		REGISTER_VIZ_TYPE(TEXT("EyeCausticNormal")      , PATH_TRACER_DEBUG_VIZ_EYE_CAUSTIC_NORMAL      );
 #undef REGISTER_VIZ_TYPE
 	}
 
