Improve Motion Vector World Offset invalidation with fallback to motion vector offset in the current frame. This is more friendly to motion blur.
1. Reprojected to outside of border, reuses the current frame pixel's offset and velocity.
2. Reprojected to region without velocity, use it instead of fully reject when updating velocity, let the temporal accumulator to reject so that we can maintain a similar blur when compared to the per-vertex.

#tests Done on the test map
#DisabledValidations Preflight

diff --git a/Engine/Shaders/Private/VelocityUpdate.usf b/Engine/Shaders/Private/VelocityUpdate.usf
--- a/Engine/Shaders/Private/VelocityUpdate.usf
+++ b/Engine/Shaders/Private/VelocityUpdate.usf
@@ -22,7 +22,8 @@ float3 DecodeVelocityTexture(Texture2D<ENCODED_VELOCITY_TYPE> Velocity, float2 P
 {
 	int2 IntPosition = int2(Position);
 	//TODO: sub-pixel hole filing.
-	return DecodeVelocityFromTexture(Velocity[IntPosition + int2(0, 0)]);
+	ENCODED_VELOCITY_TYPE EncodedVelocity = Velocity[IntPosition + int2(0, 0)];
+	return EncodedVelocity.x > 0 ? DecodeVelocityFromTexture(EncodedVelocity) : 0;
 }
 
 float CalcDeviceZ(float2 Position)
@@ -59,28 +60,23 @@ void MainCS(in uint3 ThreadID : SV_DispatchThreadID)
 	float3 DecodedVectorOffset = DecodeVelocityFromTexture(EncodedVectorOffset);
 	bool bHasVelocity = EncodedVectorOffset.x > 0;
 	float2 OffsetTexelPosition = TexelPosition - (ScreenPosToViewportUV(DecodedVectorOffset.xy) - 0.5f/*remove bias*/) * ResolvedView.ViewSizeAndInvSize.xy;// non-integer velocity;
-	
+	bHasVelocity &= any(OffsetTexelPosition != float2(TexelPosition));
+
 	ENCODED_VELOCITY_TYPE EncodedMergedVelocity;
-	if (bHasVelocity && PointInView(OffsetTexelPosition) && any(OffsetTexelPosition != float2(TexelPosition)))
+	if (bHasVelocity)
 	{
-		float3 DecodedTargetVelocity = DecodeVelocityTexture(VelocityTexture, OffsetTexelPosition);
+		// Velocity at border that projects to outside of the view uses the current frame's offset and velocity.
+		float3 DecodedTargetVelocity = DecodeVelocityTexture(VelocityTexture, PointInView(OffsetTexelPosition) ? OffsetTexelPosition : TexelPosition);
 		float2 TargetTexelPosition = OffsetTexelPosition - (ScreenPosToViewportUV(DecodedTargetVelocity.xy) - 0.5f)*ResolvedView.ViewSizeAndInvSize.xy;
 	
 		// Construct velocity offset from Position to Target Texel Position.
-		if (PointInView(TargetTexelPosition))
-		{
-			float SourceDeviceZ = DepthTexture.Load(int3(TexelPosition, 0)).r;
-			float OffsetDeviceZ = CalcDeviceZ(OffsetTexelPosition);
-			float3 Velocity = float3(DecodedVectorOffset.xy, SourceDeviceZ - OffsetDeviceZ) + DecodedTargetVelocity;
+		float SourceDeviceZ = DepthTexture.Load(int3(TexelPosition, 0)).r;
+		float OffsetDeviceZ = CalcDeviceZ(OffsetTexelPosition);
+		float3 Velocity = float3(DecodedVectorOffset.xy, SourceDeviceZ - OffsetDeviceZ) + DecodedTargetVelocity;
 
-			//TODO: check if the TargetTexelPosition is within bound for rigid body transform to reject history that is too far/close.
-			//  |QtQt+1| \in (|PtPt+1| +/- (s+1)/s * |Pt+1Qt+1|, where s is the scaling.
-			EncodedMergedVelocity = EncodeVelocityToTexture(Velocity,TemporalResponsivenessMask);
-		}
-		else
-		{
-			EncodedMergedVelocity = VelocityTexture[TexelPosition];
-		}
+		//TODO: check if the TargetTexelPosition is within bound for rigid body transform to reject history that is too far/close.
+		//  |QtQt+1| \in (|PtPt+1| +/- (s+1)/s * |Pt+1Qt+1|, where s is the scaling.
+		EncodedMergedVelocity = EncodeVelocityToTexture(Velocity,TemporalResponsivenessMask);
 	}
 	else
 	{
