Fix Ambient cubemap with Substrate Format=0
Fix pre-exposure being applied with Substrate but not in the legacy path.

Now results matches perfectly

diff --git a/Engine/Shaders/Private/AmbientCubemapComposite.usf b/Engine/Shaders/Private/AmbientCubemapComposite.usf
--- a/Engine/Shaders/Private/AmbientCubemapComposite.usf
+++ b/Engine/Shaders/Private/AmbientCubemapComposite.usf
@@ -4,6 +4,13 @@
 	PostprocessAmbient.usf: To apply a ambient cubemap as a postprocess
 =============================================================================*/
 
+#define APPLY_SKY_SHADOWING 0
+#define USE_DEFAULT_ENV_LIGHTING_INPUT 0
+#define USE_SUBSTRATE_ENV_LIGHTING_COMMON 1
+//#define ENABLE_DYNAMIC_SKY_LIGHT 1
+#define ENABLE_DYNAMIC_SKY_LIGHT 0
+#define ENABLE_SIMPLE_ENV_LIGHT 1
+
 #include "Common.ush"
 #include "PostProcessCommon.ush"
 #include "DeferredShadingCommon.ush"
@@ -43,14 +50,12 @@ float3 GetEnvSpecularLighting(
 	return TextureCubeSampleLevel(AmbientCubemap, AmbientCubemapSampler, SpecularDirection, SpecularMip).rgb * AmbientCubemapColor.rgb;
 }
 
-#define APPLY_SKY_SHADOWING 0
-#define USE_DEFAULT_ENV_LIGHTING_INPUT 0
-#define USE_SUBSTRATE_ENV_LIGHTING_COMMON 1
-
 #include "ReflectionEnvironmentShared.ush"
+#include "SkyLightingDiffuseShared.ush"
 #include "Substrate/Substrate.ush"
 #include "Substrate/SubstrateEvaluation.ush"
 #include "Substrate/SubstrateLightingCommon.ush"
+#include "Substrate/SubstrateRead.ush"
 
 void MainPS(
 	in float4 SvPosition : SV_POSITION,
@@ -71,10 +76,18 @@ void MainPS(
 	OutSubSurfaceSceneColor = 0.0f;
 #endif
 
+	#if SUBSTRATE_LOAD_FROM_MATERIALCONTAINER
 	FSubstrateAddressing SubstrateAddressing = GetSubstratePixelDataByteOffset(PixelPos, uint2(View.BufferSizeAndInvSize.xy), Substrate.MaxBytesPerPixel);
 	FSubstratePixelHeader SubstratePixelHeader = UnpackSubstrateHeaderIn(Substrate.MaterialTextureArray, SubstrateAddressing, Substrate.TopLayerTexture);
+	const bool bMaterialIsLit = SubstratePixelHeader.ClosureCount > 0;	// This test is also enough to exclude sky pixels
+	#else
+	FScreenSpaceData ScreenSpaceData = GetScreenSpaceData(BufferUV, true);
+	const FGBufferData GBuffer = ScreenSpaceData.GBuffer;
+	const bool bMaterialIsLit = GBuffer.ShadingModelID != SHADINGMODELID_UNLIT;
+	#endif
+
 	BRANCH
-	if (SubstratePixelHeader.ClosureCount > 0) // This test is also enough to exclude sky pixels
+	if (bMaterialIsLit)
 	{
 		float DeviceZ = SampleDeviceZFromSceneTextures(BufferUV);
 		float SceneDepth = ConvertFromDeviceZ(DeviceZ);
@@ -89,17 +102,28 @@ void MainPS(
 		const FSubstrateIntegrationSettings Settings = InitSubstrateIntegrationSettings(false /*bForceFullyRough*/, Substrate.bRoughDiffuse, Substrate.PeelLayersAboveDepth, Substrate.bRoughnessTracking);
 		FSubstrateDeferredLighting SubstrateLighting = GetInitialisedSubstrateDeferredLighting();
 
-		const float CombinedScreenAndMaterialAO = SubstrateGetAO(SubstratePixelHeader) * AmbientOcclusion;
-
+		#if SUBSTRATE_LOAD_FROM_MATERIALCONTAINER
 		Substrate_for (uint ClosureIndex = 0, ClosureIndex < SubstratePixelHeader.ClosureCount, ++ClosureIndex)
+		#else
+		uint ClosureIndex = 0;
+		#endif
 		{
+			#if SUBSTRATE_LOAD_FROM_MATERIALCONTAINER
 			FSubstrateBSDF BSDF = UnpackSubstrateBSDF(Substrate.MaterialTextureArray, SubstrateAddressing, SubstratePixelHeader);
 			FSubstrateBSDFContext SubstrateBSDFContext = SubstrateCreateBSDFContext(SubstratePixelHeader, BSDF, SubstrateAddressing, V);
-			const float3 BSDFThroughput = LuminanceWeight(SubstrateBSDFContext, BSDF); // Use the reflected direction
+			float CombinedScreenAndMaterialAO = SubstrateGetAO(SubstratePixelHeader) * AmbientOcclusion;
+			FSubstrateIrradianceAndOcclusion SubstrateIrradianceAndOcclusion = SubstrateGetIrradianceAndAO(SubstratePixelHeader);
+			#else
+			FSubstrateGBufferBSDF GBufferBSDF = SubstrateReadGBufferBSDF(ScreenSpaceData);
+			FSubstrateBSDFContext SubstrateBSDFContext = SubstrateCreateBSDFContext(GBufferBSDF.BSDFTangentBasis, GBufferBSDF.BSDF, V, V/*L*/, false/*bHasValidL*/);
+			float CombinedScreenAndMaterialAO = GBufferBSDF.BSDFAO;
+			FSubstrateIrradianceAndOcclusion SubstrateIrradianceAndOcclusion = SubstrateGetIrradianceAndAO(ScreenSpaceData.GBuffer);
+			#endif
+
+			const float3 BSDFThroughput = LuminanceWeight(SubstrateBSDFContext, SubstrateBSDFContext.BSDF); // Use the reflected direction
 
 			// Evaluate environment lighting
 			FSubstrateEnvLightResult SubstrateEnvLight = SubstrateEvaluateForEnvLight(SubstrateBSDFContext, true /*bEnableSpecular*/, Settings);
-			FSubstrateIrradianceAndOcclusion SubstrateIrradianceAndOcclusion = SubstrateGetIrradianceAndAO(SubstratePixelHeader);
 
 			float3 DiffuseLighting = 0;
 			float3 SpecularLighting = 0;
@@ -107,7 +131,7 @@ void MainPS(
 				SubstrateEnvLight,
 				SubstrateIrradianceAndOcclusion,
 				SubstrateBSDFContext,
-				BSDF,
+				SubstrateBSDFContext.BSDF,
 				SubstrateEnvLight.DiffuseNormal,
 				BSDFThroughput,
 				0 /*CaptureDataStartIndex*/,
@@ -121,8 +145,8 @@ void MainPS(
 				SpecularLighting);
 
 			FLightAccumulator Out = (FLightAccumulator)0;
-			LightAccumulator_AddSplit(Out, DiffuseLighting, SpecularLighting, DiffuseLighting, View.PreExposure, SubstrateEnvLight.bPostProcessSubsurface);
-			AccumulateSubstrateDeferredLighting(SubstrateLighting, Out, SubstrateEnvLight.bPostProcessSubsurface, BSDF_GETISTOPLAYER(BSDF));
+			LightAccumulator_AddSplit(Out, DiffuseLighting, SpecularLighting, DiffuseLighting, 1.f/*Legacy does not apply View.PreExposure*/, SubstrateEnvLight.bPostProcessSubsurface);
+			AccumulateSubstrateDeferredLighting(SubstrateLighting, Out, SubstrateEnvLight.bPostProcessSubsurface, BSDF_GETISTOPLAYER(SubstrateBSDFContext.BSDF));
 		}
 
 		OutColor += SubstrateLighting.SceneColor;
@@ -133,7 +157,7 @@ void MainPS(
 	}
 }
 
-#else // SUBSTRATE_GBUFFER_FORMAT==1
+#else // SUBSTRATE_ENABLED
 
 float3 DiffuseIBL( uint2 Random, float3 DiffuseColor, float Roughness, float3 N, float3 V )
 {
@@ -599,4 +623,4 @@ void MainPS(
 
 	OutColor *= AmbientOcclusion;
 }
-#endif // SUBSTRATE_GBUFFER_FORMAT==1
\ No newline at end of file
+#endif // SUBSTRATE_ENABLED
\ No newline at end of file
diff --git a/Engine/Shaders/Private/Substrate/SubstrateLightingCommon.ush b/Engine/Shaders/Private/Substrate/SubstrateLightingCommon.ush
--- a/Engine/Shaders/Private/Substrate/SubstrateLightingCommon.ush
+++ b/Engine/Shaders/Private/Substrate/SubstrateLightingCommon.ush
@@ -229,13 +229,20 @@ void SubstrateEnvLightingCommon(
 
 	// Diffuse component
 	DiffuseLighting = 0;
-#if ENABLE_DYNAMIC_SKY_LIGHT
+#if ENABLE_DYNAMIC_SKY_LIGHT || ENABLE_SIMPLE_ENV_LIGHT
 	const bool bProcessFrontFaceDiffuse = any(SubstrateEnvLight.DiffuseWeight > 0.0f);
 	const bool bProcessBackFaceDiffuse = any(SubstrateEnvLight.DiffuseBackFaceWeight > 0.0f);
 	if (bProcessFrontFaceDiffuse || bProcessBackFaceDiffuse)
 	{
 		// Compute the common sky visibility factors
+		#if ENABLE_DYNAMIC_SKY_LIGHT
 		FSkyLightVisibilityData SkyVisData = GetSkyLightVisibilityData(SubstrateBSDFContext.N, SubstrateBSDFContext.N, SubstrateIrradianceAndOcclusion.MaterialAO, ScreenAmbientOcclusion, BentNormal);
+		#else
+		FSkyLightVisibilityData SkyVisData = (FSkyLightVisibilityData)0;
+		SkyVisData.SkyDiffuseLookUpMul		= CombinedScreenAndMaterialAO;
+		SkyVisData.SkyDiffuseLookUpAdd		= 0.f;
+		SkyVisData.SkyDiffuseLookUpNormal	= BentNormal;
+		#endif
 
 		if (bProcessFrontFaceDiffuse)
 		{
diff --git a/Engine/Source/Runtime/Renderer/Private/IndirectLightRendering.cpp b/Engine/Source/Runtime/Renderer/Private/IndirectLightRendering.cpp
--- a/Engine/Source/Runtime/Renderer/Private/IndirectLightRendering.cpp
+++ b/Engine/Source/Runtime/Renderer/Private/IndirectLightRendering.cpp
@@ -247,6 +247,7 @@ class FAmbientCubemapCompositePS : public FGlobalShader
 		SHADER_PARAMETER_SAMPLER(SamplerState,  AmbientOcclusionSampler)
 
 		SHADER_PARAMETER_RDG_UNIFORM_BUFFER(FSubstrateGlobalUniformParameters, Substrate)
+		SHADER_PARAMETER_RDG_UNIFORM_BUFFER(FSceneTextureUniformParameters, SceneTexturesStruct)
 		SHADER_PARAMETER_STRUCT_INCLUDE(FAmbientCubemapParameters, AmbientCubemap)
 		SHADER_PARAMETER_STRUCT_INCLUDE(FSceneTextureParameters, SceneTextures)
 		SHADER_PARAMETER_STRUCT_REF(FViewUniformShaderParameters, ViewUniformBuffer)
@@ -1623,6 +1624,7 @@ void FDeferredShadingSceneRenderer::RenderDiffuseIndirectAndAmbientOcclusion(
 
 				PSParameters->Substrate = Substrate::BindSubstrateGlobalUniformParameters(View);
 				PSParameters->SceneTextures = SceneTextureParameters;
+				PSParameters->SceneTexturesStruct = SceneTextures.UniformBuffer;
 				PSParameters->ViewUniformBuffer = View.ViewUniformBuffer;
 
 				PSParameters->RenderTargets[0] = FRenderTargetBinding(SceneColorTexture, ERenderTargetLoadAction::ELoad);
