Substrate - fixed mobile shader compilation.

#tests pc editor dx12
#rnx
#DisabledValidations Preflight

diff --git a/Engine/Shaders/Private/Substrate/SubstrateMobileForwardLighting.ush b/Engine/Shaders/Private/Substrate/SubstrateMobileForwardLighting.ush
--- a/Engine/Shaders/Private/Substrate/SubstrateMobileForwardLighting.ush
+++ b/Engine/Shaders/Private/Substrate/SubstrateMobileForwardLighting.ush
@@ -218,7 +218,7 @@ float3 SubstrateMobileForwardLighting(
 				if (bHasDirectionalLight)
 				{
 					half DynamicShadowing = 1.0f;
-					FDeferredLightData DirLightData = GetDirectionalLightData(MaterialParameters.ScreenPosition, ScreenUV, DynamicShadowing, EyeIndex);
+					FDeferredLightData DirLightData = GetDirectionalLightData(MaterialParameters.ScreenPosition, DynamicShadowing);
 					half4 DirLightAttenuation = half4(1, 1, DynamicShadowing, DynamicShadowing);
 					uint DirLightingChannelMask = LIGHTING_CHANNEL_MASK;
 					float3 ToDirLight = DirLightData.Direction;	// Already normalized
