Path Tracer: Various fixes for Substrate support

Simplify how legacy materials are handled. They now go through the same codepath as legacy materials and the minor differences in payload are handled explicitly. This includes both blendable and adaptive g-buffer to minimize the chance of divergence between these two modes.

Simplify handling of optimized unlit materials to reduce nested preprocesor conditionals. Now we just remove the main BSDF loop for these cases.

Fix MeanFreePath having no effect in regular substrate slabs

Fix missing SSS and eye caustics in EYE shadingmodel support.

Fix GPU Lightmass shaders for invalid substrate materials

Fix missing energy conservation/preservation when substrate is enabled (the ifdef logic did not consider the path tracing case when using SUBSTRATE_GBUFFER_FORMAT=0

Add a utility function to ensure conservation of energy in legacy SSS and Two-Sided Foliage shading models (the logic was previously duplicated in a few spots).

#rnx

diff --git a/Engine/Shaders/Private/PathTracing/Material/PathTracingMaterialCommon.ush b/Engine/Shaders/Private/PathTracing/Material/PathTracingMaterialCommon.ush
--- a/Engine/Shaders/Private/PathTracing/Material/PathTracingMaterialCommon.ush
+++ b/Engine/Shaders/Private/PathTracing/Material/PathTracingMaterialCommon.ush
@@ -198,3 +198,17 @@ float3 GetPathTracingDiffuseModel(float3 DiffuseColor, float Roughness, float No
 	return DiffuseColor;
 #endif
 }
+
+// Given two diffuse colors as found on either SHADINGMODELID_TWOSIDED_FOLIAGE or SHADINGMODELID_SUBSURFACE, balance them out to ensure they don't reflect more than 100% energy
+void AdjustPathTracingTwoSidedColorBalance(inout float3 FrontColor, inout float3 BackColor)
+{
+	float3 Sum = FrontColor + BackColor;
+	float S = max(Sum.x, max(Sum.y, Sum.z));
+	if (S > 1)
+	{
+		// nudge the material back to something physically plausible
+		// NOTE: we ignore spec here since it should be accounted for by the brdf model itself
+		FrontColor /= S;
+		BackColor /= S;
+	}
+}
\ No newline at end of file
diff --git a/Engine/Shaders/Private/PathTracing/Material/PathTracingMaterialSampling.ush b/Engine/Shaders/Private/PathTracing/Material/PathTracingMaterialSampling.ush
--- a/Engine/Shaders/Private/PathTracing/Material/PathTracingMaterialSampling.ush
+++ b/Engine/Shaders/Private/PathTracing/Material/PathTracingMaterialSampling.ush
@@ -367,10 +367,13 @@ void AdjustPayloadAfterUnpack(inout FPathTracingPayload HitPayload, bool bApplyO
 		}
 		case SHADINGMODELID_TWOSIDED_FOLIAGE:
 		{
+			float3 TransmissionColor = HitPayload.GetFoliageTransmissionColor();
 			if (bApplyOverrides)
 			{
-				HitPayload.SetFoliageTransmissionColor(HitPayload.GetFoliageTransmissionColor() * View.DiffuseOverrideParameter.w + View.DiffuseOverrideParameter.xyz);
+				TransmissionColor = TransmissionColor * View.DiffuseOverrideParameter.w + View.DiffuseOverrideParameter.xyz;
 			}
+			AdjustPathTracingTwoSidedColorBalance(HitPayload.DiffuseColor, TransmissionColor);
+			HitPayload.SetFoliageTransmissionColor(TransmissionColor);
 			break;
 		}
 		case SHADINGMODELID_CLEAR_COAT:
@@ -387,6 +390,7 @@ void AdjustPayloadAfterUnpack(inout FPathTracingPayload HitPayload, bool bApplyO
 			float IrisMask = HitPayload.GetEyeIrisMask();
 			HitPayload.SubsurfaceColor = HitPayload.DiffuseColor - IrisMask * HitPayload.DiffuseColor;
 			HitPayload.DiffuseColor *= IrisMask;
+			AdjustDiffuseSSSContribution(HitPayload.DiffuseColor, HitPayload.SubsurfaceColor, HitPayload.MeanFreePath);
 			break;
 		}
 	}
@@ -410,25 +414,15 @@ void AdjustPayloadAfterUnpack(inout FPathTracingPayload HitPayload, bool bApplyO
 		case SHADINGMODELID_TWOSIDED_FOLIAGE:
 		{
 			HitPayload.SubsurfaceColor = HitPayload.GetSubsurfaceColor();
-			float3 Front = HitPayload.DiffuseColor;
-			float3 Back = HitPayload.SubsurfaceColor;
-			float3 Sum = Front + Back;
-			float S = max(Sum.x, max(Sum.y, Sum.z));
-			if (S > 1)
+			if (bApplyOverrides)
 			{
-				// nudge the material back to something physically plausible
-				// NOTE: we ignore spec here since it should be accounted for by the brdf model itself
-				HitPayload.DiffuseColor = Front / S;
-				HitPayload.SubsurfaceColor = Back / S;
+				HitPayload.SubsurfaceColor = HitPayload.SubsurfaceColor * View.DiffuseOverrideParameter.w + View.DiffuseOverrideParameter.xyz;
 			}
+			AdjustPathTracingTwoSidedColorBalance(HitPayload.DiffuseColor, HitPayload.SubsurfaceColor);
 			if (HitPayload.ShadingModelID != SHADINGMODELID_TWOSIDED_FOLIAGE)
 			{
 				AdjustDiffuseSSSContribution(HitPayload.DiffuseColor, HitPayload.SubsurfaceColor, HitPayload.SubsurfaceRadius);
 			}
-			if (bApplyOverrides)
-			{
-				HitPayload.SubsurfaceColor = HitPayload.SubsurfaceColor * View.DiffuseOverrideParameter.w + View.DiffuseOverrideParameter.xyz;
-			}
 			break;
 		}
 		case SHADINGMODELID_SUBSURFACE_PROFILE:
diff --git a/Engine/Shaders/Private/PathTracing/PathTracingCommon.ush b/Engine/Shaders/Private/PathTracing/PathTracingCommon.ush
--- a/Engine/Shaders/Private/PathTracing/PathTracingCommon.ush
+++ b/Engine/Shaders/Private/PathTracing/PathTracingCommon.ush
@@ -817,6 +817,7 @@ struct FPathTracingPayload : FMinimalPayload
 	float2 GetHairPrimitiveUV() { return SpecularEdgeColor.xy; }
 	void SetHairPrimitiveUV(float2 UV) { SpecularEdgeColor.xy = UV; }
 
+	// SHADINGMODELID_SOLID_GLASS
 	float3 GetExtinction() { return rcp(max(MeanFreePath, 1e-5)); }
 
 	// SHADINGMODELID_THIN_TRANSLUCENT
@@ -884,7 +885,7 @@ struct FPathTracingPayload : FMinimalPayload
 	float3 GetTransmittanceColor() { return CustomData0.xyz; }
 	void SetTransmittanceColor(float3 TransmittanceColor) { CustomData0.xyz = TransmittanceColor; }
 
-	// SHADINGMODELID_DEFAULT_LIT
+	// SHADINGMODELID_SOLID_GLASS
 	float3 GetExtinction() { return CustomData0.xyz; }
 	void SetExtinction(float3 SigmaT) { CustomData0.xyz = SigmaT; }
 
diff --git a/Engine/Shaders/Private/PathTracing/PathTracingGPULightmassMaterialHitShader.usf b/Engine/Shaders/Private/PathTracing/PathTracingGPULightmassMaterialHitShader.usf
--- a/Engine/Shaders/Private/PathTracing/PathTracingGPULightmassMaterialHitShader.usf
+++ b/Engine/Shaders/Private/PathTracing/PathTracingGPULightmassMaterialHitShader.usf
@@ -426,7 +426,9 @@ RAY_TRACING_ENTRY_ANY_HIT(GPULightmassMaterialAHS,
 
 	float3 Transparency = saturate((1.0f - Coverage) + Coverage * TransmittancePreCoverage);
 #else
-#error "Shouldn't need to compile AHS for non-substrate material"
+	// Invalid material
+	// SUBSTRATE_TODO: Investigate which materials land in this case, and if there is any need to process them more fully
+	float3 Transparency = 0;
 #endif
 
 #else // SUBSTRATE_ENABLED
diff --git a/Engine/Shaders/Private/PathTracing/PathTracingMaterialHitShader.usf b/Engine/Shaders/Private/PathTracing/PathTracingMaterialHitShader.usf
--- a/Engine/Shaders/Private/PathTracing/PathTracingMaterialHitShader.usf
+++ b/Engine/Shaders/Private/PathTracing/PathTracingMaterialHitShader.usf
@@ -21,6 +21,9 @@
 // Ensure that SSS albedo comes through in the material
 #define SUBSTRATE_SSS_MATERIAL_OVERRIDE 0
 
+// This define happens when we compile material with BlendableGBuffer (legacy) and no substrate nodes are used in a materials.
+// As oftoday, we are then going to compile the material with Substrate disabled. But the payload is still a Substrate payload!
+// The implementation here will take the legacy code path, but simply encoding the information in the new payload format as needed.
 #if (SUBSTRATE_ENABLED && !TEMPLATE_USES_SUBSTRATE) // == TEMPLATE_USES_LEGACYGBUFFERDATA from MaterialTemplate, needs to happen here before other #include in this file.
 #define SUBSTRATE_INLINE_SHADING 1
 #endif
@@ -44,15 +47,19 @@
 #include "/Engine/Private/PathTracing/PathTracingCommon.ush"
 #include "/Engine/Private/PathTracing/Material/PathTracingFresnel.ush"
 
-// This define happens when we compile material with BlendableGBuffer (legacy) data and no substrate nodes are used or hidden conversion happening.
-// As of today, we are then going to compile the material with Substrate disabled. But the payload is still a Substrate payload!
-// So to quickly and simply convert from LEgacy material into a Substrate payload, We create a slab from SubstrateConvertLegacyMaterialStatic
-// that is then exported to a path tracer material using the recular code path. Only sampling from a Single BSDF instead of the substrate tree of BSDFs.
-#if TEMPLATE_USES_LEGACYGBUFFERDATA
-#include "/Engine/Private/Substrate/SubstrateLegacyConversion.ush"
+#if SUBSTRATE_ENABLED && (TEMPLATE_USES_LEGACYGBUFFERDATA || MATERIAL_IS_SUBSTRATEHIDDENCONVERSION)
+// Either the material was explicitly setup to use legacy data (blendable g-buffer case)
+// Or the material is a hidden conversion into a substrate material (adaptive g-buffer case)
+// In both cases, for the path tracer we prefer to work with the original data directly as this allows taking the
+// well tested legacy path and simply handle the differences in payload encoding locally rather than
+// go through the substrate slab structures and reverse engineer the parameters (which is not always possible, particularly for
+// cases that require promiting transparency into refraction).
+#define USE_LEGACY_CODEPATH_FOR_SUBSTRATE 1
+#else
+#define USE_LEGACY_CODEPATH_FOR_SUBSTRATE 0
 #endif
 
-#if (SUBSTRATE_ENABLED && TEMPLATE_USES_SUBSTRATE) || TEMPLATE_USES_LEGACYGBUFFERDATA
+#if (SUBSTRATE_ENABLED && TEMPLATE_USES_SUBSTRATE) || USE_LEGACY_CODEPATH_FOR_SUBSTRATE
 #define MATERIAL_SUBSTRATE_OPAQUE_PRECOMPUTED_LIGHTING 0
 #include "/Engine/Private/Substrate/SubstrateExport.ush"
 #include "/Engine/Private/PathTracing/Material/PathTracingSubstrateCommon.ush"
@@ -218,28 +225,12 @@ RAY_TRACING_ENTRY_CLOSEST_HIT(PathTracingMaterialCHS,
 	}
 
 	const float3 V_World = MaterialParameters.CameraVector;
-#if (SUBSTRATE_ENABLED && TEMPLATE_USES_SUBSTRATE)
+#if SUBSTRATE_ENABLED && !USE_LEGACY_CODEPATH_FOR_SUBSTRATE
 
 #if MATERIAL_IS_SUBSTRATE
 	FSubstratePixelHeader SubstratePixelHeader = MaterialParameters.GetFrontSubstrateHeader();
 	FSubstrateData SubstrateData = PixelMaterialInputs.GetFrontSubstrateData();
 
-#if SUBSTRATE_OPTIMIZED_UNLIT
-	// Unlit BSDF goes through the SubstrateTree to support weighting operations. Technically, layering and mixing could also be supported.
-	float3	UnlitSurfaceLuminancePostCoverage = 0.0f;
-	float	UnlitSurfaceCoverage = 0.0f;
-	float3	UnlitSurfaceTransmittancePreCoverage = 0.0f;
-	float3	UnlitSurfaceNormal = 0.0f;
-	SubstratePixelHeader.SubstrateUpdateTreeUnlit(
-		uint2(SvPosition.xy),
-		MaterialParameters.CameraVector,
-		SubstrateData,
-		UnlitSurfaceLuminancePostCoverage,
-		UnlitSurfaceCoverage,
-		UnlitSurfaceTransmittancePreCoverage,
-		UnlitSurfaceNormal);
-	const float3 Transparency = saturate(lerp(1.0, UnlitSurfaceTransmittancePreCoverage, UnlitSurfaceCoverage));
-#else
 	// SUBSTRATE_TODO: These should actually be initialized via the corresponding cvars. To do this properly would require binding FSubstrateSceneData in the path tracing shaders
 	// For now, simply match the default substrate settings
 	FSubstrateIntegrationSettings IntegrationSettings = InitSubstrateIntegrationSettings(
@@ -264,49 +255,16 @@ RAY_TRACING_ENTRY_CLOSEST_HIT(PathTracingMaterialCHS,
 	// SUBSTRATE_TODO: Thin case does not seem to be properly handled here, so the TotalTransmittancePreCoverage is recomputed later
 	const float3 Transparency = saturate(lerp(1.0, TotalTransmittancePreCoverage, TotalCoverage));
 #endif // PATH_TRACER_REFRACTION_ENABLED
-#endif // SUBSTRATE_OPTIMIZED_UNLIT
 
 #else // MATERIAL_IS_SUBSTRATE
 	// This only happens for invalid materials
 	const float3 Transparency = 0;
 #endif
 
-#elif TEMPLATE_USES_LEGACYGBUFFERDATA
-
-#if SUBSTRATE_BLENDING_COLOREDTRANSMITTANCEONLY
-	const uint BSDFCount = 0;
-#else
-	const uint BSDFCount = 1;
-#endif
-
-	float TotalCoverage = 1;
-	float3 TotalTransmittancePreCoverage = 0;
-#if SUBSTRATE_BLENDING_COLOREDTRANSMITTANCEONLY || SUBSTRATE_BLENDING_TRANSLUCENT_COLOREDTRANSMITTANCE || SUBSTRATE_BLENDING_TRANSLUCENT_GREYTRANSMITTANCE
-
-	TotalCoverage = GetMaterialOpacity(PixelMaterialInputs);
-#if MATERIAL_SHADINGMODEL_THIN_TRANSLUCENT
-	//if (Payload.ShadingModelID == SHADINGMODELID_THIN_TRANSLUCENT)
-	{
-		TotalTransmittancePreCoverage = GetThinTranslucentMaterialOutput0(MaterialParameters);
-		TotalCoverage = GetThinTranslucentMaterialOutput1(MaterialParameters);
-	}
-#endif
-#endif
-
-#if REFRACTION_USE_INDEX_OF_REFRACTION
-	// transparency is only used for coverage, transmittance happens through refraction (except for shadows potentially, which is handled below)
-	float3 Transparency = 1 - TotalCoverage;
-#else
-	// transparency is used for both coverage and transmittance
-	// SUBSTRATE_TODO: Thin case does not seem to be properly handled here, so the TotalTransmittancePreCoverage is recomputed later
-	const float3 Transparency = saturate(lerp(1.0, TotalTransmittancePreCoverage, TotalCoverage));
-#endif
-
-#endif // SUBSTRATE_ENABLED && TEMPLATE_USES_SUBSTRATE
+#endif // SUBSTRATE_ENABLED && !USE_LEGACY_CODEPATH_FOR_SUBSTRATE
 
 	FPathTracingPayload Payload = (FPathTracingPayload)0;
 
-#if !(SUBSTRATE_ENABLED && TEMPLATE_USES_SUBSTRATE) || TEMPLATE_USES_LEGACYGBUFFERDATA
 	/**
 	 * Set common material attributes for both full and simplified materials
 	 **/
@@ -357,8 +315,6 @@ RAY_TRACING_ENTRY_CLOSEST_HIT(PathTracingMaterialCHS,
 #error Unknown material blending mode
 #endif
 
-#endif //  !(SUBSTRATE_ENABLED && TEMPLATE_USES_SUBSTRATE) || TEMPLATE_USES_LEGACYGBUFFERDATA
-
 	// fetch primitive flags only once
 	// TODO: would be nice to keep this inside MaterialParameters as it is also needed there as well
 	const uint PrimitiveFlags = GetPrimitiveData(MaterialParameters.PrimitiveId).Flags;
@@ -396,283 +352,36 @@ RAY_TRACING_ENTRY_CLOSEST_HIT(PathTracingMaterialCHS,
 	Payload.WorldSmoothNormal = AdjustShadingNormal(Payload.WorldSmoothNormal, Payload.WorldGeoNormal, WorldRayDirection());
 #endif
 
-#if SUBSTRATE_ENABLED && TEMPLATE_USES_SUBSTRATE  || TEMPLATE_USES_LEGACYGBUFFERDATA
+#if SUBSTRATE_ENABLED && !USE_LEGACY_CODEPATH_FOR_SUBSTRATE
 	// unless proven otherwise, we are a basic unlit material
 	Payload.ShadingModelID = SHADINGMODELID_UNLIT;
 
-#if SUBSTRATE_OPTIMIZED_UNLIT
+	float BaseBSDFOpacity = 1; //Payload.BSDFOpacity;
 
-	Payload.Radiance = UnlitSurfaceLuminancePostCoverage;
-#if SUBSTRATE_OPAQUE_MATERIAL
-	Payload.BSDFOpacity = 1;
-	Payload.TransparencyColor = 0;
-#else
-	Payload.BSDFOpacity = UnlitSurfaceCoverage;
-	Payload.TransparencyColor = lerp(1.0, UnlitSurfaceTransmittancePreCoverage, UnlitSurfaceCoverage);
-#endif
-
-#else // SUBSTRATE_OPTIMIZED_UNLIT
-
-#if !TEMPLATE_USES_LEGACYGBUFFERDATA // In this case we must keep BSDFOpacity and TransparencyColor evaluated above.
-	Payload.BSDFOpacity = 1;
 #if !SUBSTRATE_OPAQUE_MATERIAL
 	Payload.TransparencyColor = Transparency;
 #endif
-#endif
-
-
-#if TEMPLATE_USES_LEGACYGBUFFERDATA
-
-	Payload.ShadingModelID = GetMaterialShadingModel(PixelMaterialInputs);
-
-	FSubstratePixelFootprint PixelFootprint = SubstrateGetPixelFootprint(1.0f, 1.0f, 0.0f /*GBuffer.Curvature*/);
-
-	half3 BaseColor = GetMaterialBaseColor(PixelMaterialInputs);
-	half  Metallic = GetMaterialMetallic(PixelMaterialInputs);
-	half  Specular = GetMaterialSpecular(PixelMaterialInputs);
-	half Roughness = GetMaterialRoughness(PixelMaterialInputs);
-
-	uint SharedLocalBasisTypes = SUBSTRATE_BASIS_TYPE_NORMAL;
-	float3 WorldNormal = MaterialParameters.WorldNormal;
-	float3 WorldTangent = MaterialParameters.TangentToWorld[2];
-	float  Anisotropy = 0;
-#if MATERIAL_USES_ANISOTROPY
-	SharedLocalBasisTypes = SUBSTRATE_BASIS_TYPE_TANGENT;
-		WorldTangent = CalculateAnisotropyTangent(MaterialParameters, PixelMaterialInputs);
-	Anisotropy = GetMaterialAnisotropy(PixelMaterialInputs);
-#endif	
-
-	uint ProfileId = 0;
-	if (Payload.ShadingModelID == SHADINGMODELID_SUBSURFACE_PROFILE || Payload.ShadingModelID == SHADINGMODELID_EYE)
-	{
-		ProfileId = ExtractSubsurfaceProfileInt(GetMaterialSubsurfaceDataRaw(PixelMaterialInputs).a);
-	}
-
-
-	const float3 Emissive = GetMaterialEmissive(PixelMaterialInputs);
-	const float Opacity = GetMaterialOpacity(PixelMaterialInputs);
-
-	float3 SubSurfaceColor = 0.0f;
-
-	float ClearCoat = 0.0f;
-	float ClearCoatRoughness = 0.0f;
-	float3 ClearCoatWorldBottomNormal = WorldNormal;
-
-	float3 ThinTranslucentTransmittanceColor = 1.0f;
-	float ThinTranslucentSurfaceCoverage = 0.0f;
-
-	float3 WaterScatteringCoefficients = 0.0f;
-	float3 WaterAbsorptionCoefficients = 0.0f;
-	float WaterPhaseG = 0.0f;
-	float3 ColorScaleBehindWater = 1.0f;
-
-	float3 IrisNormal = WorldNormal;
-	float3 IrisPlaneNormal = WorldNormal;
-	float IrisMask = 0;
-	float IrisDistance = 0;
-
-#if MATERIAL_SHADINGMODEL_TWOSIDED_FOLIAGE || MATERIAL_SHADINGMODEL_SUBSURFACE
-	if (Payload.ShadingModelID == SHADINGMODELID_SUBSURFACE ||
-		Payload.ShadingModelID == SHADINGMODELID_TWOSIDED_FOLIAGE)
-	{
-		SubSurfaceColor = GetMaterialSubsurfaceData(PixelMaterialInputs).rgb;
-	}
-	else
-#endif
-
-#if MATERIAL_SHADINGMODEL_CLEAR_COAT
-	if (Payload.ShadingModelID == SHADINGMODELID_CLEAR_COAT)
-	{
-		ClearCoat = GetMaterialCustomData0(PixelMaterialInputs);
-		ClearCoatRoughness = GetMaterialCustomData1(PixelMaterialInputs);
-
-	#if CLEAR_COAT_BOTTOM_NORMAL
-	#if NUM_MATERIAL_OUTPUTS_CLEARCOATBOTTOMNORMAL > 0
-	#if MATERIAL_TANGENTSPACENORMAL
-		ClearCoatWorldBottomNormal = normalize(TransformTangentVectorToWorld(MaterialParameters.TangentToWorld, ClearCoatBottomNormal0(MaterialParameters)));
-	#else
-		ClearCoatWorldBottomNormal = ClearCoatBottomNormal0(MaterialParameters);
-	#endif
-
-		// if we got a custom normal for the clearcoat, adjust it now
-		ClearCoatWorldBottomNormal = AdjustShadingNormal(ClearCoatWorldBottomNormal, Payload.WorldGeoNormal, WorldRayDirection());
-	#endif
-	#endif
-	}
-	else
-#endif
-
-#if MATERIAL_SHADINGMODEL_CLOTH
-	if (Payload.ShadingModelID == SHADINGMODELID_CLOTH)
-	{
-		const float3 ClothColor = GetMaterialSubsurfaceData(PixelMaterialInputs).rgb;
-		const float FuzzAmount = saturate(GetMaterialCustomData0(PixelMaterialInputs));
-		ClearCoat = FuzzAmount;
-		SubSurfaceColor = ClothColor;
-	}
-	else
-#endif
-
-#if MATERIAL_SHADINGMODEL_THIN_TRANSLUCENT
-	if (Payload.ShadingModelID == SHADINGMODELID_THIN_TRANSLUCENT)
-	{
-		ThinTranslucentTransmittanceColor = GetThinTranslucentMaterialOutput0(MaterialParameters);
-		ThinTranslucentSurfaceCoverage = GetThinTranslucentMaterialOutput1(MaterialParameters);
-	}
-	else
-#endif
-
-#if MATERIAL_SHADINGMODEL_SINGLELAYERWATER
-	if (Payload.ShadingModelID == SHADINGMODELID_SINGLELAYERWATER)
-	{
-		// For single layer water we always want to treat it as front facing when evaluating parameters 
-		// as hitting back face is treated as being underwater and we get completely different parameters.
-		FMaterialPixelParameters WaterMaterialParameters = MaterialParameters;
-		WaterMaterialParameters.TwoSidedSign = 1;
-
-		WaterScatteringCoefficients = max(0.0f, DFDemote(GetSingleLayerWaterMaterialOutput0(WaterMaterialParameters)));
-		WaterAbsorptionCoefficients = max(0.0f, DFDemote(GetSingleLayerWaterMaterialOutput1(WaterMaterialParameters)));
-		WaterPhaseG = clamp(DFDemote(GetSingleLayerWaterMaterialOutput2(WaterMaterialParameters).x), -0.99f, 0.99f);
-	}
-	else
-#endif
-
-#if MATERIAL_SHADINGMODEL_EYE
-	if (Payload.ShadingModelID == SHADINGMODELID_EYE)
-	{
-		// This is all based on logic from ShadingModelsMaterial.ush
-		float IrisMask     = saturate(GetMaterialCustomData0(PixelMaterialInputs));
-		float IrisDistance = saturate(GetMaterialCustomData1(PixelMaterialInputs));
-	#if NUM_MATERIAL_OUTPUTS_GETTANGENTOUTPUT > 0
-		// Blend in the negative intersection normal to create some concavity
-		// Not great as it ties the concavity to the convexity of the cornea surface
-		// No good justification for that. On the other hand, if we're just looking to
-		// introduce some concavity, this does the job.
-		const float3 PlaneNormal = normalize(GetTangentOutput0(MaterialParameters));
-		const float3 CausticNormal = normalize(lerp(PlaneNormal, -Payload.WorldNormal, IrisMask * IrisDistance));
-		IrisPlaneNormal = CausticNormal;
-		// NOTE: calling AdjustShadingNormal on the custom normal does not seem necessary since this is only used for a diffuse calculation
-	#else
-		const float3 PlaneNormal = Payload.WorldNormal;
-		IrisPlaneNormal = Payload.WorldNormal;
-	#endif
-
-	#if IRIS_NORMAL // on
-		Specular = 0.25;
-
-	#if NUM_MATERIAL_OUTPUTS_CLEARCOATBOTTOMNORMAL > 0
-		IrisPlaneNormal = normalize( ClearCoatBottomNormal0(MaterialParameters) );
-		#if MATERIAL_TANGENTSPACENORMAL
-		IrisPlaneNormal = normalize( TransformTangentVectorToWorld( MaterialParameters.TangentToWorld, IrisNormal ) );
-		#endif
-	#else
-		IrisPlaneNormal = PlaneNormal;
-	#endif
-	#endif
-	}
-	else
-#endif // MATERIAL_SHADINGMODEL_EYE
-	{
-		// Last empty if/else block.
-	}
-
-	FSubstrateData SubstrateData = SubstrateConvertLegacyMaterialStatic(
-		PixelFootprint,
-		BaseColor, Specular, Metallic,
-		Roughness, Anisotropy,
-		SubSurfaceColor, ProfileId,
-		ClearCoat, ClearCoatRoughness,
-		Emissive,
-		Opacity,
-		ThinTranslucentTransmittanceColor,
-		ThinTranslucentSurfaceCoverage,
-		WaterScatteringCoefficients, WaterAbsorptionCoefficients, WaterPhaseG, ColorScaleBehindWater,
-		Payload.ShadingModelID,
-		WorldNormal,
-		WorldTangent,
-		ClearCoatWorldBottomNormal,
-		IrisPlaneNormal,
-		0,//dummy uint SharedLocalBasisIndex,
-		1,//dummy uint ClearCoatBottomNormal_SharedLocalBasisIndex,
-		SharedLocalBasisTypes);
-
-#endif // TEMPLATE_USES_LEGACYGBUFFERDATA
 
-
-
-
-
-
-	if (BSDFCount > 0 || TEMPLATE_USES_LEGACYGBUFFERDATA)
+	if (BSDFCount > 0)
 	{
-#if TEMPLATE_USES_LEGACYGBUFFERDATA
-		{
-	#if HAIR_STRAND_MESH_FACTORY
-			Payload.WorldSmoothNormal = Payload.WorldGeoNormal;
-	#else
-			WorldNormal = normalize(WorldNormal);
-		#if !MATERIAL_TANGENTSPACENORMAL
-			// already flipped if the material was in tangent space, so add the flip if it wasn't
-			WorldNormal *= MaterialParameters.TwoSidedSign;
-		#endif
-			WorldNormal = AdjustShadingNormal(WorldNormal, Payload.WorldGeoNormal, WorldRayDirection());
-	#endif
 
-		}
-#else // TEMPLATE_USES_LEGACYGBUFFERDATA
-		{
-	#if HAIR_STRAND_MESH_FACTORY
-			Payload.WorldSmoothNormal = Payload.WorldGeoNormal;
-	#else
-			// Nudge normal to avoid dark edges in bump
-			// Not that we must recompute tangent vector to keep the basis frame orthogonal
-			for (uint i = 0; i < SubstratePixelHeader.SharedLocalBases.Count; ++i)
-			{
-				float3 N = normalize(SubstratePixelHeader.SharedLocalBases.Normals[i]);
-	#if !MATERIAL_TANGENTSPACENORMAL
-                // already flipped if the material was in tangent space, so add the flip if it wasn't
-				N *= MaterialParameters.TwoSidedSign;
-	#endif
-				N = AdjustShadingNormal(N, Payload.WorldGeoNormal, WorldRayDirection());
-				SubstratePixelHeader.SharedLocalBases.Normals[i]  = N;
-			}
-	#endif
-		}
-#endif // TEMPLATE_USES_LEGACYGBUFFERDATA
 
 		int BSDFIdx = 0;
-#if TEMPLATE_USES_LEGACYGBUFFERDATA
-#define BSDF SubstrateData.InlinedBSDF
-		BSDF.SubstrateSanitizeBSDF();
-		BSDF.LuminanceWeightV = 1;
-		BSDF.CoverageAboveAlongN = 0;
-		BSDF.TransmittanceAboveAlongN = 1;
-		BSDF.bIsBottom = 1;
-		BSDF.bIsTop = 1;
-	#if SUBSTRATE_INLINE_SHADING
-		Payload.Radiance = BSDF.Emissive * Payload.BSDFOpacity;
-	#endif
-#else
 #define BSDF SubstratePixelHeader.SubstrateTree.BSDFs[BSDFIdx]
 
-	#if !TEMPLATE_USES_LEGACYGBUFFERDATA
 		// Accumulate radiance across all BSDFs
 		Payload.Radiance = 0.0;
-		for (BSDFIdx = 0; BSDFIdx < BSDFCount; ++BSDFIdx)
+		Substrate_for_unroll(BSDFIdx = 0, BSDFIdx < BSDFCount, ++BSDFIdx)
 		{
 			Payload.Radiance += BSDF.LuminanceWeightV * BSDF_GETEMISSIVE(BSDF);
 		}
-	#endif
-#endif
 
-#if SUBSTRATE_CLAMPED_CLOSURE_COUNT > 1 && !TEMPLATE_USES_LEGACYGBUFFERDATA
-		// The material might be using more than one slab -- insert code to stochastically pick one 
+#if !SUBSTRATE_OPTIMIZED_UNLIT 
+		// The slab sampler lets us stochastically choose between the different slab models the path tracer supports
+		// NOTE: Even if the material has a single substrate slab, we might still need this because refraction is handled seperately from SSS
 		FRISContext SlabSampler = InitRISContext(PackedPayload.GetStochasticSlabRand());
-		for (BSDFIdx = 0; BSDFIdx < BSDFCount; ++BSDFIdx)
-#else
-		// Just one slab, don't bother inserting logic for lobe selection, we just have one
-		BSDFIdx = 0;
-#endif
+
+		Substrate_for_unroll(BSDFIdx = 0, BSDFIdx < BSDFCount, ++BSDFIdx)
 		{
 			switch (BSDF_GETTYPE(BSDF))
 			{
@@ -681,16 +390,12 @@ RAY_TRACING_ENTRY_CLOSEST_HIT(PathTracingMaterialCHS,
 					const float3 WeightV = BSDF.LuminanceWeightV;
 					const float3 TransmittanceN = BSDF.TransmittanceAboveAlongN;
 					const float CoverageAboveAlongN = BSDF.CoverageAboveAlongN;
-					const float3 MaxLobeWeight = WeightV * lerp(1.0f, TransmittanceN, CoverageAboveAlongN); // largest value LobeWeight() could return
+					const float3 MaxLobeWeight = BaseBSDFOpacity * WeightV * lerp(1.0f, TransmittanceN, CoverageAboveAlongN); // largest value LobeWeight() could return
 
 					// Is the slab even visible at all?
 					if (any(MaxLobeWeight > 0.0))
 					{
-					#if TEMPLATE_USES_LEGACYGBUFFERDATA
-						const float3 N = WorldNormal;
-					#else
 						const float3 N = SubstratePixelHeader.SharedLocalBases.Normals[BSDF_GETSHAREDLOCALBASISID(BSDF)];
-					#endif
 						const float NoV = saturate(dot(N, V_World));
 
 						float3 DiffuseColor = SLAB_DIFFUSEALBEDO(BSDF);
@@ -718,13 +423,25 @@ RAY_TRACING_ENTRY_CLOSEST_HIT(PathTracingMaterialCHS,
 								Spec1E = ComputeGGXSpecEnergyTermsRGB(RoughnessData.y, NoV, F0, F90).E;
 							}
 						}
+						float3 MeanFreePath = 0;
+						float PhaseG = 0;
 						float3 LegacySubsurfaceColor = 0;
 						float LegacyOpacity = 1;
-						if (BSDF_GETSSSTYPE(BSDF) == SSS_TYPE_SIMPLEVOLUME)
+						if (BSDF_GETSSSTYPE(BSDF) == SSS_TYPE_SIMPLEVOLUME && !BSDF.bIsBottom)
 						{
 							// TODO: This should really be represented by a different BSDF, but for now we are approximating it with a diffuse lobe
 							DiffuseColor = GetSimpleVolumeDiffuseColor(DiffuseColor, SLAB_SSSMFP(BSDF));
 						}
+						else if (BSDF_GETSSSTYPE(BSDF) == SSS_TYPE_DIFFUSION_PROFILE)
+						{
+							const uint ProfileId = SLAB_SSSPROFILEID(BSDF);
+							const float Scale = SLAB_SSSPROFILERADIUSSCALE(BSDF);
+							MeanFreePath = DecodeSSSProfileRadius(ProfileId, Payload.DiffuseColor, Scale);
+							#if USE_SSS_PROFILE_ANISOTROPY
+							PhaseG = DecodeSSSProfileScatteringDistribution(ProfileId);
+							MeanFreePath *= 1 - PhaseG; // Adjust MFP to preserve appearance
+							#endif
+						}
 						else if (BSDF_GETSSSTYPE(BSDF) == SSS_TYPE_WRAP || BSDF_GETSSSTYPE(BSDF) == SSS_TYPE_TWO_SIDED_WRAP)
 						{
 							// For the old MATERIAL_SHADINGMODEL_SUBSURFACE and SHADINGMODELID_TWOSIDED_FOLIAGE models, we previously had logic to normalize the DiffuseColor and SubsurfaceColor to ensure the sum is < 1
@@ -734,17 +451,13 @@ RAY_TRACING_ENTRY_CLOSEST_HIT(PathTracingMaterialCHS,
 							LegacySubsurfaceColor = MeanFreePathToTransmittance(MFPInCm * CENTIMETER_TO_METER, SUBSTRATE_SIMPLEVOLUME_THICKNESS_CM * CENTIMETER_TO_METER);
 							LegacyOpacity = 1.f - abs(SLAB_SSSPHASEANISOTROPY(BSDF));
 
-							float3 Front = DiffuseColor;
-							float3 Back = LegacySubsurfaceColor;
-							float3 Sum = Front + Back;
-							float S = max3(Sum.x, Sum.y, Sum.z);
-							if (S > 1)
-							{
-								// nudge the material back to something physically plausible
-								// NOTE: we ignore spec here since it should be accounted for by the brdf model itself
-								DiffuseColor = Front / S;
-								LegacySubsurfaceColor = Back / S;
-							}
+							AdjustPathTracingTwoSidedColorBalance(DiffuseColor, LegacySubsurfaceColor);
+						}
+						else
+						{
+							// Basic substrate slab, legacy preintegrated skin or SIMPLEVOLUME on the bottom slab, just use MFP directly for the random walk
+							MeanFreePath = SLAB_SSSMFP(BSDF);
+							PhaseG       = SLAB_SSSPHASEANISOTROPY(BSDF);
 						}
 						// SUBSTRATE_TODO: Other SSSTYPE cases may result in some portion of a "glass" BSDF
 						float FuzzAmount = 0.0;
@@ -777,35 +490,40 @@ RAY_TRACING_ENTRY_CLOSEST_HIT(PathTracingMaterialCHS,
 							Spec0Albedo *= 1.0 - RoughnessData.z;
 							Spec1Albedo *= RoughnessData.z;
 						}
+
+						float Ior = 1; // for refracted lobe only
+						float3 Transmittance = 0;
+						float SlabMix = 1; // percentage mix against glass lobe
+						#if MATERIALBLENDING_TRANSLUCENT && (SUBSTRATE_BLENDING_TRANSLUCENT_GREYTRANSMITTANCE || SUBSTRATE_BLENDING_TRANSLUCENT_COLOREDTRANSMITTANCE)
+						if (BSDF.bIsBottom)
+						{
+							#if PATH_TRACER_REFRACTION_ENABLED
+							Ior = GetRefractionIor(PixelMaterialInputs, F0);
+							#endif
+							Transmittance = MeanFreePathToTransmittance(MeanFreePath * CENTIMETER_TO_METER, SUBSTRATE_SIMPLEVOLUME_THICKNESS_CM * CENTIMETER_TO_METER);
+							SlabMix = LobeSelectionProb(DiffuseColor, Transmittance);
+						}
+						#endif
+
 						// SUBSTRATE_TODO: Need to account for glass case here for lobe selection ....
 						const float3 FuzzAlbedo = FuzzAmount * FuzzE * FuzzColor;
 						const float3 SlabAlbedo = MaxLobeWeight * (DiffuseWeight * (DiffuseColor + LegacySubsurfaceColor) + Spec0Albedo + Spec1Albedo + FuzzAlbedo);
-						#if SUBSTRATE_CLAMPED_CLOSURE_COUNT > 1
+						const float3 RefrAlbedo = MaxLobeWeight * (FuzzAttenuation * 1.0 + FuzzAlbedo); // glass has unit albedo
 						// If we are dealing witha multi-BSDF material, stochastically decide if we need to write to the payload
 						// This means estimating the overal albedo of the slab and selectively writing to the payload
 						// NOTE: Technically, some of the albedo calculation are not needed if we aren't picking between slabs
 						// We rely on the compiler to get rid of dead code here to avoid diverging the logic too far between single and multi-lobe cases
-						float Contrib = LobeColorToWeight(SlabAlbedo);
-						if (SlabSampler.Accept(Contrib))
-						#endif
+						float ContribSlab = LobeColorToWeight(SlabAlbedo);
+						if (SlabSampler.Accept(SlabMix * ContribSlab))
 						{
 							// We've selected this particular slab - now write it into the payload
 							// NOTE: This could happen more than once, so be careful to write to all fields, even the unused ones!
-							#if SUBSTRATE_CLAMPED_CLOSURE_COUNT > 1
-							Payload.BSDFOpacity = 1.0 / Contrib;
-							#endif
+							Payload.BSDFOpacity = BaseBSDFOpacity / ContribSlab;
 							Payload.WorldNormal = N;
-						#if TEMPLATE_USES_LEGACYGBUFFERDATA
-							if (SharedLocalBasisTypes == SUBSTRATE_BASIS_TYPE_TANGENT)
-							{
-								Payload.WorldTangent = WorldTangent;
-							}
-						#else
 							if (SubstrateGetSharedLocalBasisType(SubstratePixelHeader.SharedLocalBases.Types, BSDFIdx) == SUBSTRATE_BASIS_TYPE_TANGENT)
 							{
 								Payload.WorldTangent = SubstratePixelHeader.SharedLocalBases.Tangents[BSDF_GETSHAREDLOCALBASISID(BSDF)];
 							}
-						#endif
 							Payload.WeightV               = BSDF.LuminanceWeightV;
 							Payload.TransmittanceN        = BSDF.TransmittanceAboveAlongN;
 							Payload.CoverageAboveAlongN   = BSDF.CoverageAboveAlongN;
@@ -824,51 +542,31 @@ RAY_TRACING_ENTRY_CLOSEST_HIT(PathTracingMaterialCHS,
 							#endif
 							}
 							Payload.Anisotropy = BSDF_GETHASANISOTROPY(BSDF) ? SLAB_ANISOTROPY(BSDF) : 0.0;
-							if (BSDF_GETHASFUZZ(BSDF))
-							{
-								Payload.FuzzAmount    = FuzzAmount;
-								Payload.FuzzColor     = FuzzColor;
-								Payload.FuzzRoughness = FuzzRoughness;
-							}
-							else
-							{
-								Payload.FuzzAmount    = 0.0;
-								Payload.FuzzColor     = 0.0;
-								Payload.FuzzRoughness = 0.0;
-							}
-							Payload.MeanFreePath = 0;
-							Payload.PhaseG       = 0;
-							if (BSDF_GETSSSTYPE(BSDF) == SSS_TYPE_SIMPLEVOLUME)
-							{
-								// TODO: this case should be represented by a different BSDF
-								// For now, we already adjusted DiffuseColor according to MFP above, so nothing extra to do
-							}
-							else if (BSDF_GETSSSTYPE(BSDF) == SSS_TYPE_DIFFUSION)
-							{
-								Payload.MeanFreePath = SLAB_SSSMFP(BSDF);
-								Payload.PhaseG       = SLAB_SSSPHASEANISOTROPY(BSDF);
-							}
-							else if (BSDF_GETSSSTYPE(BSDF) == SSS_TYPE_DIFFUSION_PROFILE)
-							{
-								const uint ProfileId = SLAB_SSSPROFILEID(BSDF);
-								const float Scale = SLAB_SSSPROFILERADIUSSCALE(BSDF);
-								Payload.MeanFreePath = DecodeSSSProfileRadius(ProfileId, Payload.DiffuseColor, Scale);
-								#if USE_SSS_PROFILE_ANISOTROPY
-								const float G = DecodeSSSProfileScatteringDistribution(ProfileId);
-								Payload.PhaseG = G;
-								Payload.MeanFreePath *= 1 - G; // Adjust MFP to preserve appearance
-								#endif
-							}
-							else if (BSDF_GETSSSTYPE(BSDF) == SSS_TYPE_WRAP)
+							Payload.FuzzAmount    = FuzzAmount;
+							Payload.FuzzColor     = FuzzColor;
+							Payload.FuzzRoughness = FuzzRoughness;
+							if (BSDF_GETSSSTYPE(BSDF) == SSS_TYPE_WRAP)
 							{
 								// Legacy support for SHADINGMODELID_SUBSURFACE
-								const float SSSRadius = (1 - LegacyOpacity) * 10.0; // simple formula, up to 10cm radius
+
 								// The legacy behavior has both diffuse+sss responses separately, but we can combine them to avoid needing a custom lobe
 								// In practice the diffuse and SSS responses are similar, so the difference should be minimal (and is an approximation of the raster behavior anyway)
 								// SUBSTRATE_TODO: If we implement stochastic lobe picking _within_ slabs, we could fine tune this to get back to the original behavior exactly
 								Payload.DiffuseColor += LegacySubsurfaceColor;
+
+							#if HAVE_GetSubsurfaceMediumMaterialOutput0
+								Payload.MeanFreePath = GetSubsurfaceMediumMaterialOutput0(MaterialParameters);
+							#else
+								const float SSSRadius = (1 - LegacyOpacity) * 10.0; // simple formula, up to 10cm radius
 								Payload.MeanFreePath = SSSRadius;
+							#endif
+
+							#if HAVE_GetSubsurfaceMediumMaterialOutput1
+								float ScatteringDistribution = clamp(GetSubsurfaceMediumMaterialOutput1(MaterialParameters), -0.99f, 0.99f);
+								Payload.PhaseG = ScatteringDistribution;
+							#else
 								Payload.PhaseG = 0;
+							#endif
 							}
 							else if (BSDF_GETSSSTYPE(BSDF) == SSS_TYPE_TWO_SIDED_WRAP)
 							{
@@ -876,47 +574,16 @@ RAY_TRACING_ENTRY_CLOSEST_HIT(PathTracingMaterialCHS,
 								Payload.ShadingModelID = SHADINGMODELID_TWOSIDED_FOLIAGE;
 								Payload.SetFoliageTransmissionColor(LegacySubsurfaceColor);
 							}
-
-							if (BSDF.bIsBottom)
+							else
 							{
-								#if PATH_TRACER_REFRACTION_ENABLED
-								float Ior = GetRefractionIor(PixelMaterialInputs, F0);
-								#endif
-
-								#if 1 && MATERIALBLENDING_TRANSLUCENT && (SUBSTRATE_BLENDING_TRANSLUCENT_GREYTRANSMITTANCE || SUBSTRATE_BLENDING_TRANSLUCENT_COLOREDTRANSMITTANCE) && REFRACTION_USE_INDEX_OF_REFRACTION
-								// Disable refraction handling for now -- need to come up with a better way to know the amount of "glass" we have separate from the concept of opacity
-								// In legacy slabs, the old Opacity plug should blend between DefaultLit behavior and Glass, while in new Substrate slabs this blend should depend on the amount of albedo and MFP.
-
-								// NOTE: only the bottom most layer can refract, coating slabs should never include a refraction lobe
-								// We group this with the SSS conditions because it involves MFP
-
-								float3 MFP = SLAB_SSSMFP(BSDF);
-								if (Ior > 0 && any(MFP > 0.0) && Payload.ShadingModelID == SHADINGMODELID_SUBSTRATE)
-								{
-									// SUBSTRATE_TODO: This isn't quite correct - we want to incorporate a proper mixing for the diffuse/sss part of the material here, not just make pure glass
-									#if MATERIAL_ISTHINSURFACE || MATERIAL_SHADINGMODEL_THIN_TRANSLUCENT // SUBSTRATE_TODO: Should THIN_TRANSLUCENT imply ISTHIN_SURFACE?
-									Payload.ShadingModelID = SHADINGMODELID_THIN_TRANSLUCENT;
-									#else
-									Payload.ShadingModelID = SHADINGMODELID_SOLID_GLASS;
-									#endif
-									Payload.MeanFreePath   = MFP;
-									Payload.Ior            = Ior;
-								}
-								#endif
-
-								#if (MATERIAL_ISTHINSURFACE || MATERIAL_SHADINGMODEL_THIN_TRANSLUCENT) && MATERIALBLENDING_TRANSLUCENT && (SUBSTRATE_BLENDING_TRANSLUCENT_GREYTRANSMITTANCE || SUBSTRATE_BLENDING_TRANSLUCENT_COLOREDTRANSMITTANCE)
-								// Material is configured as a thin surface, and refraction is disabled
-								if (!(Payload.Ior > 0.0))
-								{
-									// SUBSTRATE_TODO: Do we really want to keep this? Should maybe always use refraction path?
-									Payload.MeanFreePath = SLAB_SSSMFP(BSDF);
-									float3 Transmission = Payload.GetTransmittanceColor();
-									const float NoV = saturate(dot(Payload.WorldNormal, V_World));
-									float3 Tr = ComputeThinSlabWeights(Transmission, NoV, Payload.Ior, F0RGBToF0(Payload.SpecularColor)).Transmitted;
-									Payload.TransparencyColor += TotalCoverage * Tr;
-								}
-								#endif
-
+							#if (MATERIAL_ISTHINSURFACE || MATERIAL_SHADINGMODEL_THIN_TRANSLUCENT)
+								// material is thin, make sure we don't use random walk
+								// SUBSTRATE_TODO: derive the proper BSDF for the internal scattering layer inside the slab
+							#else
+								// Setup material for random walk
+								Payload.MeanFreePath = MeanFreePath;
+								Payload.PhaseG       = PhaseG;
+							#endif
 							}
 
 							// Handle special-case specular features
@@ -939,29 +606,75 @@ RAY_TRACING_ENTRY_CLOSEST_HIT(PathTracingMaterialCHS,
 							}
 							#endif
 						}
+						// now test for the presence of a refracted lobe
+						float ContribRefr = LobeColorToWeight(RefrAlbedo);
+						if (SlabSampler.Accept((1 - SlabMix) * ContribRefr))
+						{
+							// We've selected this particular slab - now write it into the payload
+							// NOTE: This could happen more than once, so be careful to write to all fields, even the unused ones!
+							Payload.BSDFOpacity = BaseBSDFOpacity / ContribRefr;
+							Payload.WorldNormal = N;
+							if (SubstrateGetSharedLocalBasisType(SubstratePixelHeader.SharedLocalBases.Types, BSDFIdx) == SUBSTRATE_BASIS_TYPE_TANGENT)
+							{
+								Payload.WorldTangent = SubstratePixelHeader.SharedLocalBases.Tangents[BSDF_GETSHAREDLOCALBASISID(BSDF)];
+							}
+							Payload.WeightV               = BSDF.LuminanceWeightV;
+							Payload.TransmittanceN        = BSDF.TransmittanceAboveAlongN;
+							Payload.CoverageAboveAlongN   = BSDF.CoverageAboveAlongN;
+							#if MATERIAL_ISTHINSURFACE || MATERIAL_SHADINGMODEL_THIN_TRANSLUCENT // SUBSTRATE_TODO: Should THIN_TRANSLUCENT imply ISTHIN_SURFACE?
+							Payload.ShadingModelID = SHADINGMODELID_THIN_TRANSLUCENT;
+							#else
+							Payload.ShadingModelID = SHADINGMODELID_SOLID_GLASS;
+							#endif
+							Payload.Ior                   = 0.0;
+							Payload.DiffuseColor          = 0.0;
+							Payload.SpecularColor         = F0;
+							Payload.SpecularEdgeColor     = F90;
+							Payload.RoughnessData         = RoughnessData;
+							Payload.BottomNormalOct16bits = 0;
+							Payload.Anisotropy = BSDF_GETHASANISOTROPY(BSDF) ? SLAB_ANISOTROPY(BSDF) : 0.0;
+							Payload.FuzzAmount    = FuzzAmount;
+							Payload.FuzzColor     = FuzzColor;
+							Payload.FuzzRoughness = FuzzRoughness;
+							Payload.MeanFreePath  = MeanFreePath;
+							Payload.PhaseG        = 0;
+							Payload.Ior = Ior;
+							// SUBSTRATE_TODO: should these be supported in the Glass case?
+							Payload.SpecularProfileId = 0;
+							Payload.GlintValue = 1.f;
+
+
+							#if (MATERIAL_ISTHINSURFACE || MATERIAL_SHADINGMODEL_THIN_TRANSLUCENT) && MATERIALBLENDING_TRANSLUCENT && (SUBSTRATE_BLENDING_TRANSLUCENT_GREYTRANSMITTANCE || SUBSTRATE_BLENDING_TRANSLUCENT_COLOREDTRANSMITTANCE)
+							// Material is configured as a thin surface, and refraction is disabled
+							if (!(Payload.Ior > 0.0))
+							{
+								// SUBSTRATE_TODO: Do we really want to keep this? Should maybe always use refraction path?
+								float3 Transmission = Payload.GetTransmittanceColor();
+								const float NoV = saturate(dot(Payload.WorldNormal, V_World));
+								float3 Tr = ComputeThinSlabWeights(Transmission, NoV, 0.0, F0RGBToF0(Payload.SpecularColor)).Transmitted;
+								Payload.TransparencyColor += TotalCoverage * Tr;
+							}
+							#endif
+						}
 					}
 					break;
 				}
 				// NOTE: The cases below are for back-compatibility and cannot be layered
 				case SUBSTRATE_BSDF_TYPE_EYE:
 				{
-				#if TEMPLATE_USES_LEGACYGBUFFERDATA
-					Payload.WorldNormal = WorldNormal;
-				#else
 					Payload.WorldNormal = SubstratePixelHeader.SharedLocalBases.Normals[BSDF_GETSHAREDLOCALBASISID(BSDF)];
-				#endif
 					Payload.DiffuseColor = EYE_DIFFUSEALBEDO(BSDF);
 					Payload.SpecularColor = EYE_F0(BSDF);
 					// NOTE: EYE_F90 is always 1.0 (does not exist in legacy model), so don't bother with it
 					Payload.SetEyeRoughness(AdjustMaterialRoughness(EYE_ROUGHNESS(BSDF), PackedPayload.GetPathRoughness()));
 
+					Payload.MeanFreePath = 0;
+					Payload.PhaseG = 0;
 					if (BSDF_GETSSSTYPE(BSDF) == SSS_TYPE_DIFFUSION_PROFILE)
 					{
 						Payload.MeanFreePath = DecodeSSSProfileRadius(EYE_SSSPROFILEID(BSDF), Payload.DiffuseColor, 1.0);
-						Payload.PhaseG = 0.0;
 					}
 
-
 					float IrisMask = EYE_IRISMASK(BSDF);
 					float IrisDistance = EYE_IRISDISTANCE(BSDF);
 
@@ -973,10 +686,13 @@ RAY_TRACING_ENTRY_CLOSEST_HIT(PathTracingMaterialCHS,
 					Payload.SetEyeIrisMask(IrisMask);
 					Payload.SetEyeIrisNormal(normalize(EYE_IRISNORMAL(BSDF)));
 					Payload.ShadingModelID = SHADINGMODELID_EYE;
+					Payload.SpecularProfileId = 0;
+					Payload.GlintValue = 1.f;
 					break;
 				}
 				case SUBSTRATE_BSDF_TYPE_HAIR:
 				{
+					SlabSampler.Accept(1); // dummy, avoids special case below
 					// clamp the roughness to whatever came along the path
 					Payload.SetBaseColor(HAIR_BASECOLOR(BSDF));
 					Payload.SetHairLongitudinalRoughness(AdjustMaterialRoughness(HAIR_ROUGHNESS(BSDF), PackedPayload.GetPathRoughness()));
@@ -991,15 +707,13 @@ RAY_TRACING_ENTRY_CLOSEST_HIT(PathTracingMaterialCHS,
 					Payload.SetHairPrimitiveUV(0.5);
 					#endif
 					Payload.ShadingModelID = SHADINGMODELID_HAIR;
+					Payload.SpecularProfileId = 0;
+					Payload.GlintValue = 1.f;
 					break;
 				}
 				case SUBSTRATE_BSDF_TYPE_SINGLELAYERWATER:
 				{
-				#if TEMPLATE_USES_LEGACYGBUFFERDATA
-					Payload.WorldNormal = WorldNormal;
-				#else
 					Payload.WorldNormal = SubstratePixelHeader.SharedLocalBases.Normals[BSDF_GETSHAREDLOCALBASISID(BSDF)];
-				#endif
 					Payload.BSDFOpacity = 1;
 					Payload.ShadingModelID = SHADINGMODELID_SOLID_GLASS;
 					float3 BaseColor = SLW_BASECOLOR(BSDF);
@@ -1028,40 +742,32 @@ RAY_TRACING_ENTRY_CLOSEST_HIT(PathTracingMaterialCHS,
 					Payload.FuzzAmount = 0;
 					Payload.FuzzRoughness = 0;
 					Payload.FuzzColor = 0;
-					Payload.GlintValue = 1;
+					Payload.SpecularProfileId = 0;
+					Payload.GlintValue = 1.f;
 					break;
 				}
 			}  // switch on slab tyype
 		} // for loop over BSDFs
 
-		#if SUBSTRATE_CLAMPED_CLOSURE_COUNT > 1
 		// normalize the weight for picking a given slab
 		Payload.BSDFOpacity *= SlabSampler.GetNormalization();
-		#endif
-
 #undef  BSDF
-	} // if (BSDFCount > 0) 
-
+#endif // !SUBSTRATE_OPTIMIZED_UNLIT
+	} // if (BSDFCount > 0)
 
-#endif // SUBSTRATE_OPTIMIZED_UNLIT
 
+#else // SUBSTRATE_ENABLED && !USE_LEGACY_CODEPATH_FOR_SUBSTRATE
 
-#if MATERIAL_TWOSIDED == 0
-	if (MaterialParameters.TwoSidedSign < 0)
-	{
-		// when viewing the surface from "inside", don't include emission
-		Payload.Radiance = 0;
-	}
-#endif
-
-#else // SUBSTRATE_ENABLED && TEMPLATE_USES_SUBSTRATE  || TEMPLATE_USES_LEGACYGBUFFERDATA 
 	Payload.Radiance = GetMaterialEmissive(PixelMaterialInputs);
+	Payload.Radiance *= Payload.BSDFOpacity; // pre-multiply
+	// NOTE: Emissive contribution is turned off below for non-two-sided materials
 
 	Payload.WorldNormal = MaterialParameters.WorldNormal;
 #if !MATERIAL_TANGENTSPACENORMAL
 	// already flipped if the material was in tangent space, so add the flip if it wasn't
 	Payload.WorldNormal *= MaterialParameters.TwoSidedSign;
 #endif
+
 #if (TRANSLUCENCY_LIGHTING_VOLUMETRIC_PERVERTEX_NONDIRECTIONAL || TRANSLUCENCY_LIGHTING_VOLUMETRIC_PERVERTEX_DIRECTIONAL || TRANSLUCENCY_LIGHTING_VOLUMETRIC_NONDIRECTIONAL || TRANSLUCENCY_LIGHTING_VOLUMETRIC_DIRECTIONAL)
 	// volumetric shading doesn't need to adjust normal
 #else
@@ -1080,21 +786,23 @@ RAY_TRACING_ENTRY_CLOSEST_HIT(PathTracingMaterialCHS,
 	Ior = GetRefractionIor(PixelMaterialInputs);
 #endif
 
-
-	Payload.Radiance *= Payload.BSDFOpacity; // pre-multiply
-
-
-#if MATERIAL_TWOSIDED == 0
-	if (MaterialParameters.TwoSidedSign < 0)
-	{
-		// when viewing the surface from "inside", don't include emission
-		Payload.Radiance = 0;
-	}
-#endif
+#if USE_LEGACY_CODEPATH_FOR_SUBSTRATE
+	Payload.DiffuseColor = ComputeDiffuseAlbedo(BaseColor, Metallic);
+	Payload.SpecularColor = ComputeF0(Specular, BaseColor, Metallic);
+	Payload.RoughnessData = float3(Roughness, Roughness, 0);
+	// reset unused features
+	Payload.SpecularEdgeColor = 1;
+	Payload.WeightV = 1;
+	Payload.CoverageAboveAlongN = 0;
+	Payload.TransmittanceN = 1;
+	Payload.GlintValue = 1;
+	Payload.SpecularProfileId = 0;
+#else
 	Payload.BaseColor = BaseColor;
 	Payload.Specular = Specular;
 	Payload.Metallic = Metallic;
 	Payload.Roughness = Roughness;
+#endif
 
 #if MATERIAL_USES_ANISOTROPY
 	Payload.WorldTangent = CalculateAnisotropyTangent(MaterialParameters, PixelMaterialInputs);
@@ -1113,25 +821,38 @@ RAY_TRACING_ENTRY_CLOSEST_HIT(PathTracingMaterialCHS,
 	if (Payload.ShadingModelID == SHADINGMODELID_UNLIT)
 	{
 		// encode the directionality using the specular field
-#if TRANSLUCENCY_LIGHTING_VOLUMETRIC_PERVERTEX_DIRECTIONAL || TRANSLUCENCY_LIGHTING_VOLUMETRIC_DIRECTIONAL
+	#if TRANSLUCENCY_LIGHTING_VOLUMETRIC_PERVERTEX_DIRECTIONAL || TRANSLUCENCY_LIGHTING_VOLUMETRIC_DIRECTIONAL
 		Payload.Anisotropy = GetMaterialTranslucencyDirectionalLightingIntensity();
-#else
+	#else
 		Payload.Anisotropy = 0;
-#endif
+	#endif
+
+	#if USE_LEGACY_CODEPATH_FOR_SUBSTRATE
+		Payload.DiffuseColor = BaseColor;
+		Payload.SpecularColor = 0;
+		Payload.RoughnessData = 1;
+	#else
 		Payload.Specular = 0;
 		Payload.Metallic = 0;
 		Payload.Roughness = 1;
+	#endif
 	}
 	else
 #elif MATERIAL_SHADINGMODEL_UNLIT
 	if (Payload.ShadingModelID == SHADINGMODELID_UNLIT)
 	{
 		// Regular unlit material, make sure the volumetric component is zero'ed out
+	#if USE_LEGACY_CODEPATH_FOR_SUBSTRATE
+		Payload.DiffuseColor = 0;
+		Payload.SpecularColor = 0;
+	#else
 		Payload.BaseColor = 0;
 		Payload.Specular = 0;
+		Payload.Metallic = 0;
+	#endif
 	}
 	else
-#endif
+#endif // MATERIAL_SHADINGMODEL_UNLIT
 
 
 #if MATERIAL_SHADINGMODEL_CLEAR_COAT
@@ -1139,75 +860,114 @@ RAY_TRACING_ENTRY_CLOSEST_HIT(PathTracingMaterialCHS,
 	{
 		const float ClearCoat = GetMaterialCustomData0(PixelMaterialInputs);
 		const float ClearCoatRoughness = GetMaterialCustomData1(PixelMaterialInputs);
+
+	#if CLEAR_COAT_BOTTOM_NORMAL
+		#if NUM_MATERIAL_OUTPUTS_CLEARCOATBOTTOMNORMAL > 0
+			#if MATERIAL_TANGENTSPACENORMAL
+				float3 BottomNormal = normalize(TransformTangentVectorToWorld(MaterialParameters.TangentToWorld, ClearCoatBottomNormal0(MaterialParameters)));
+			#else
+				float3 BottomNormal = ClearCoatBottomNormal0(MaterialParameters);
+			#endif
+		
+			// if we got a custom normal for the clearcoat, adjust it now
+			BottomNormal = AdjustShadingNormal(BottomNormal, Payload.WorldGeoNormal, WorldRayDirection());
+		#else
+			float3 BottomNormal = Payload.WorldNormal;
+		#endif
+	#else
+			float3 BottomNormal = Payload.WorldNormal;
+	#endif
+
+#if USE_LEGACY_CODEPATH_FOR_SUBSTRATE
+		Payload.RoughnessData.y = ClearCoatRoughness; // adjusted later
+		Payload.RoughnessData.z = ClearCoat;
+
+		float2 oct2 = UnitVectorToOctahedron(Payload.WorldNormal);
+		float2 oct1 = UnitVectorToOctahedron(BottomNormal);
+		float2 BottomNormalOct = ((oct1 - oct2) * 0.25) + (128.0 / 255.0);
+		Payload.BottomNormalOct16bits = PackR8(BottomNormalOct.x) | (PackR8(BottomNormalOct.y) << 8);
+#else
 		Payload.SetClearCoat(ClearCoat);
 		Payload.SetClearCoatRoughness(AdjustMaterialRoughness(ClearCoatRoughness, PackedPayload.GetPathRoughness()));
-#if CLEAR_COAT_BOTTOM_NORMAL
-#if NUM_MATERIAL_OUTPUTS_CLEARCOATBOTTOMNORMAL > 0
-#if MATERIAL_TANGENTSPACENORMAL
-		float3 BottomNormal = normalize(TransformTangentVectorToWorld(MaterialParameters.TangentToWorld, ClearCoatBottomNormal0(MaterialParameters)));
-#else
-		float3 BottomNormal = ClearCoatBottomNormal0(MaterialParameters);
-#endif
-		
-		// if we got a custom normal for the clearcoat, adjust it now
-		BottomNormal = AdjustShadingNormal(BottomNormal, Payload.WorldGeoNormal, WorldRayDirection());
 		Payload.SetClearCoatBottomNormal(BottomNormal);
-#else
-		Payload.SetClearCoatBottomNormal(Payload.WorldNormal);
-#endif
-#else
-		Payload.SetClearCoatBottomNormal(Payload.WorldNormal);
 #endif
-		
 	}
 	else
-#endif
+#endif // MATERIAL_SHADINGMODEL_CLEAR_COAT
 
 #if MATERIAL_SHADINGMODEL_CLOTH
 	if (Payload.ShadingModelID == SHADINGMODELID_CLOTH)
 	{
 		const float3 ClothColor = GetMaterialSubsurfaceData(PixelMaterialInputs).rgb;
 		const float Fuzz = saturate(GetMaterialCustomData0(PixelMaterialInputs));
+	#if USE_LEGACY_CODEPATH_FOR_SUBSTRATE
+		// the default substrate model includes a fuzz layer, so we can do it this way
+		// Additionally, tune the look like in SubstrateConvertLegacyMaterialDynamic
+		Payload.SpecularColor *= 1 - Fuzz;
+		Payload.FuzzColor = ClothColor;
+		Payload.FuzzAmount = Fuzz;
+		Payload.FuzzRoughness = 0.5 * Roughness + 0.5; // (adjust range because the legacy model has a very different response)
+		Payload.ShadingModelID = SHADINGMODELID_SUBSTRATE;
+	#else
 		Payload.SetClothColor(ClothColor);
 		Payload.SetClothAmount(Fuzz);
+	#endif
 	}
 	else
-#endif
+#endif // MATERIAL_SHADINGMODEL_CLOTH
 
 #if MATERIAL_SHADINGMODEL_TWOSIDED_FOLIAGE
 	if (Payload.ShadingModelID == SHADINGMODELID_TWOSIDED_FOLIAGE)
 	{
 		const float3 SubsurfaceColor = GetMaterialSubsurfaceData(PixelMaterialInputs).rgb;
+	#if USE_LEGACY_CODEPATH_FOR_SUBSTRATE
+		Payload.SetFoliageTransmissionColor(SubsurfaceColor);
+	#else
 		Payload.SetSubsurfaceColor(SubsurfaceColor);
+	#endif
 	}
 	else
-#endif
+#endif // MATERIAL_SHADINGMODEL_TWOSIDED_FOLIAGE
 
 #if MATERIAL_SHADINGMODEL_SUBSURFACE || MATERIAL_SHADINGMODEL_PREINTEGRATED_SKIN
 	if (Payload.ShadingModelID == SHADINGMODELID_SUBSURFACE ||
 		Payload.ShadingModelID == SHADINGMODELID_PREINTEGRATED_SKIN)
 	{
-		const float3 SubsurfaceColor = GetMaterialSubsurfaceData(PixelMaterialInputs).rgb;
-		Payload.SetSubsurfaceColor(SubsurfaceColor);
+		float3 SubsurfaceColor = GetMaterialSubsurfaceData(PixelMaterialInputs).rgb;
 
-#if HAVE_GetSubsurfaceMediumMaterialOutput0
-		Payload.SetSubsurfaceRadius(GetSubsurfaceMediumMaterialOutput0(MaterialParameters));
-#else
+	#if HAVE_GetSubsurfaceMediumMaterialOutput0
+		float3 SSSRadius = GetSubsurfaceMediumMaterialOutput0(MaterialParameters);
+	#else
 		// TODO: is this accurate enough? hard to measure exactly since the raster path uses a very different approach
 		const float Opacity = GetMaterialOpacity(PixelMaterialInputs);
 		const float SSSRadius = (1 - Opacity) * 10.0; // simple formula, up to 10cm radius
-		Payload.SetSubsurfaceRadius(SSSRadius);
-#endif
+	#endif
 
-#if HAVE_GetSubsurfaceMediumMaterialOutput1
-		float ScatteringDistribution = clamp(GetSubsurfaceMediumMaterialOutput1(MaterialParameters), -0.99f, 0.99f);
-		Payload.SetSubsurfacePhaseFunction(ScatteringDistribution);
-#else
-		Payload.SetSubsurfacePhaseFunction(0.0);
-#endif
+		
+
+	#if HAVE_GetSubsurfaceMediumMaterialOutput1
+		float PhaseG = clamp(GetSubsurfaceMediumMaterialOutput1(MaterialParameters), -0.99f, 0.99f);
+	#else
+		float PhaseG = 0.0;
+	#endif
+
+	#if USE_LEGACY_CODEPATH_FOR_SUBSTRATE
+		AdjustPathTracingTwoSidedColorBalance(Payload.DiffuseColor, SubsurfaceColor);
+		// default substrate slab includes SSS
+		Payload.ShadingModelID = SHADINGMODELID_SUBSTRATE;
+		Payload.DiffuseColor += SubsurfaceColor; // no distinction between diffuse portion and SSS portion (could try using stochastic blending here instead?)
+		Payload.MeanFreePath = SSSRadius;
+		Payload.PhaseG = PhaseG;
+	#else
+		Payload.ShadingModelID = SHADINGMODELID_SUBSURFACE;
+		Payload.SetSubsurfaceColor(SubsurfaceColor);
+		Payload.SetSubsurfaceRadius(SSSRadius);
+		Payload.SetSubsurfacePhaseFunction(PhaseG);
+	#endif
 	}
 	else
-#endif
+#endif // MATERIAL_SHADINGMODEL_SUBSURFACE || MATERIAL_SHADINGMODEL_PREINTEGRATED_SKIN
+
 #if MATERIAL_SHADINGMODEL_SUBSURFACE_PROFILE
 	if (Payload.ShadingModelID == SHADINGMODELID_SUBSURFACE_PROFILE)
 	{
@@ -1219,42 +979,49 @@ RAY_TRACING_ENTRY_CLOSEST_HIT(PathTracingMaterialCHS,
 		float LobeMix = 0;
 
 		// NOTE: we pass 1.0 for opacity because we don't really want to fade-out the dual roughness as opacity is lowered
-		GetSubsurfaceProfileDualSpecular(ProfileId, Payload.Roughness, 1.0, Lobe0Roughness, Lobe1Roughness, LobeMix);
+		GetSubsurfaceProfileDualSpecular(ProfileId, Roughness, 1.0, Lobe0Roughness, Lobe1Roughness, LobeMix);
 
 		Lobe0Roughness = AdjustMaterialRoughness(Lobe0Roughness, PackedPayload.GetPathRoughness());
 		Lobe1Roughness = AdjustMaterialRoughness(Lobe1Roughness, PackedPayload.GetPathRoughness());
-		Payload.SetDualRoughnessSpecular(Lobe0Roughness, Lobe1Roughness, LobeMix);
 
-#if HAVE_GetSubsurfaceMediumMaterialOutput0
-		Payload.SetSubsurfaceRadius(GetSubsurfaceMediumMaterialOutput0(MaterialParameters));
-#else
+	#if HAVE_GetSubsurfaceMediumMaterialOutput0
+		float3 SSSRadius = GetSubsurfaceMediumMaterialOutput0(MaterialParameters);
+	#else
+		float3 SSSRadius = 0;
 		if (GetSubsurfaceProfileUseBurley(ProfileId))
 		{
 			// Decode SSS radius
 			const float Opacity = GetMaterialOpacity(PixelMaterialInputs);
-			const float3 DiffuseColor = Payload.BaseColor * (1.0 - Payload.Metallic);
-			Payload.SetSubsurfaceRadius(DecodeSSSProfileRadius(ProfileId, DiffuseColor, Opacity));
+			const float3 DiffuseColor = ComputeDiffuseAlbedo(BaseColor, Metallic);
+			SSSRadius = DecodeSSSProfileRadius(ProfileId, DiffuseColor, Opacity);
 		}
-		else
-		{
-			Payload.SetSubsurfaceRadius(0.0);
-		}
-#endif
+	#endif
 
-#if HAVE_GetSubsurfaceMediumMaterialOutput1
-		float ScatteringDistribution = clamp(GetSubsurfaceMediumMaterialOutput1(MaterialParameters), -0.99f, 0.99f);
-		Payload.SetSubsurfacePhaseFunction(ScatteringDistribution);
-#elif USE_SSS_PROFILE_ANISOTROPY
+	#if HAVE_GetSubsurfaceMediumMaterialOutput1
+		float PhaseG = clamp(GetSubsurfaceMediumMaterialOutput1(MaterialParameters), -0.99f, 0.99f);
+	#elif USE_SSS_PROFILE_ANISOTROPY
 		// NOTE: this is not currently used
-		float G = DecodeSSSProfileScatteringDistribution(ProfileId);
-		Payload.SetSubsurfacePhaseFunction(G);
-		Payload.SetSubsurfaceRadius(Payload.GetSubsurfaceRadius() * (1 - G));
-#else
-		Payload.SetSubsurfacePhaseFunction(0.0);
-#endif
+		float PhaseG = DecodeSSSProfileScatteringDistribution(ProfileId);
+		SSSRadius *= 1 - PhaseG;
+	#else
+		float PhaseG = 0;
+	#endif
+
+	#if USE_LEGACY_CODEPATH_FOR_SUBSTRATE
+		// substrate slab includes SSS and dual roughness already
+		Payload.ShadingModelID = SHADINGMODELID_SUBSTRATE;
+		Payload.MeanFreePath = SSSRadius;
+		Payload.PhaseG = PhaseG;
+		Payload.RoughnessData = float3(Lobe0Roughness, Lobe1Roughness, LobeMix);
+	#else
+		Payload.SetSubsurfaceRadius(SSSRadius);
+		Payload.SetSubsurfacePhaseFunction(PhaseG);
+		Payload.SetDualRoughnessSpecular(Lobe0Roughness, Lobe1Roughness, LobeMix);
+	#endif
 	}
 	else
-#endif
+#endif // MATERIAL_SHADINGMODEL_SUBSURFACE_PROFILE
+
 #if MATERIAL_SHADINGMODEL_SINGLELAYERWATER
 	if (Payload.ShadingModelID == SHADINGMODELID_SINGLELAYERWATER)
 	{
@@ -1271,68 +1038,85 @@ RAY_TRACING_ENTRY_CLOSEST_HIT(PathTracingMaterialCHS,
 		const float3 WaterAlbedo = select(WaterExtinction > 0.0, ScatteringCoeff / WaterExtinction, float3(0.0, 0.0, 0.0));
 
 		Payload.ShadingModelID = SHADINGMODELID_SOLID_GLASS; // replace shading model to avoid extra cases in raygen
+	#if USE_LEGACY_CODEPATH_FOR_SUBSTRATE
+		Payload.DiffuseColor = 0; // not used
+		Payload.MeanFreePath = rcp(max(1e-5, WaterExtinction));
+	#else
 		Payload.BaseColor = WaterAlbedo; // TODO: unused at the moment
 		Payload.SetExtinction(WaterExtinction);
+	#endif
 		Payload.Ior = DielectricF0ToIor(DielectricSpecularToF0(Specular));
 		// we are using actual refraction, so don't apply any transparency
 		Payload.TransparencyColor = 0.0;
 		Payload.BSDFOpacity = 1.0;
 	}
 	else
-#endif
+#endif // MATERIAL_SHADINGMODEL_SINGLELAYERWATER
 
 #if MATERIAL_SHADINGMODEL_THIN_TRANSLUCENT
 	if (Payload.ShadingModelID == SHADINGMODELID_THIN_TRANSLUCENT)
 	{
 		const float3 TransmittanceColor = GetThinTranslucentMaterialOutput0(MaterialParameters);
-		Payload.SetTransmittanceColor(TransmittanceColor);
 
 		// The material is a blend of glass and default_lit response, model it stochastically
 		if (Payload.BSDFOpacity < PackedPayload.GetStochasticSlabRand())
 		{
 			// Pick glass lobe - shade as glass (leave SHADINGMODELID_THIN_TRANSLUCENT on)
 			Payload.Ior = Ior;
+		#if USE_LEGACY_CODEPATH_FOR_SUBSTRATE
+			Payload.MeanFreePath = TransmittanceToMeanFreePath(TransmittanceColor, SUBSTRATE_SIMPLEVOLUME_THICKNESS_CM * CENTIMETER_TO_METER) * METER_TO_CENTIMETER;
+			Payload.DiffuseColor = 0;
+		#else
+			Payload.SetTransmittanceColor(TransmittanceColor);
+		#endif
 		}
 		else
 		{
 			// Pick solid lobe -- toggle back to default_lit response
 			// Ior stays 0.0 since this mode does not include any refraction
+		#if USE_LEGACY_CODEPATH_FOR_SUBSTRATE
+			Payload.ShadingModelID = SHADINGMODELID_SUBSTRATE;
+		#else
 			Payload.ShadingModelID = SHADINGMODELID_DEFAULT_LIT;
+		#endif
 		}
 		
 		float SurfaceCoverage = GetThinTranslucentMaterialOutput1(MaterialParameters);
+	#if PATH_TRACER_REFRACTION_ENABLED
 		if (Ior > 0.0)
 		{
 			// In this mode we get rough refraction, so any transparency is only due to coverage (affecting both solid and glass lobes)
 			Payload.TransparencyColor = 1 - SurfaceCoverage;
 		}
 		else
+	#endif
 		{
 			// In this mode we just get straight transparency
 			const float3 V = WorldRayDirection();
 			const float3 N = normalize(Payload.WorldNormal);
 			const float VoN = abs(dot(V, N));
-			const float3 BaseColor = Payload.BaseColor;
-			const float  Metallic = Payload.Metallic;
-			const float  Specular = Payload.Specular;
 			const float  F0 = F0RGBToF0(ComputeF0(Specular, BaseColor, Metallic));
 			// This is the portion of transparency due to glass. We always compute it because we don't need to model the amount
 			// of transparency stochastically.
 			/// BSDFOpacity here still represents the blend between solid and glass behaviors.
 			Payload.TransparencyColor = (1 - Payload.BSDFOpacity) * ComputeThinSlabWeights(TransmittanceColor, VoN, 0.0, F0).Transmitted;
 			// Now account for overall coverage (allows the combination of solid/glass to become transparent)
-			Payload.TransparencyColor = lerp(1.0.xxx, Payload.TransparencyColor, SurfaceCoverage);
+			Payload.TransparencyColor = lerp(1.0, Payload.TransparencyColor, SurfaceCoverage);
 		}
 		// Now we can replace the BSDFOpacity with coverage
 		Payload.BSDFOpacity = SurfaceCoverage;
 	}
 	else
-#endif
+#endif // MATERIAL_SHADINGMODEL_THIN_TRANSLUCENT
 
 #if MATERIAL_SHADINGMODEL_HAIR
 	if (Payload.ShadingModelID == SHADINGMODELID_HAIR)
 	{
-		// TODO: encode hair specific info
+		Payload.SetBaseColor(BaseColor);
+		Payload.SetHairLongitudinalRoughness(AdjustMaterialRoughness(Roughness, PackedPayload.GetPathRoughness()));
+		Payload.SetHairAzimuthalRoughness(Metallic); // SUBSTRATE_TODO: Map from scatter param to azimuthal roughness here instead of in the raygen
+		Payload.SetHairSpecular(Specular);
+
 #if HAIR_STRAND_MESH_FACTORY
 		Payload.SetHairPrimitiveUV(MaterialParameters.HairPrimitiveUV);
 #else
@@ -1340,7 +1124,7 @@ RAY_TRACING_ENTRY_CLOSEST_HIT(PathTracingMaterialCHS,
 #endif
 	}
 	else
-#endif
+#endif // MATERIAL_SHADINGMODEL_HAIR
 
 #if MATERIAL_SHADINGMODEL_EYE
 	if (Payload.ShadingModelID == SHADINGMODELID_EYE)
@@ -1349,8 +1133,13 @@ RAY_TRACING_ENTRY_CLOSEST_HIT(PathTracingMaterialCHS,
 		const float IrisMask     = saturate(GetMaterialCustomData0(PixelMaterialInputs));
 		const float IrisDistance = saturate(GetMaterialCustomData1(PixelMaterialInputs));
 		Payload.SetEyeIrisMask(IrisMask);
+	#if USE_LEGACY_CODEPATH_FOR_SUBSTRATE
+		Payload.SpecularColor = ComputeF0(Specular, BaseColor, 0);
+	#else
 		Payload.Metallic = 0.0;
-#if NUM_MATERIAL_OUTPUTS_GETTANGENTOUTPUT > 0
+	#endif
+
+	#if NUM_MATERIAL_OUTPUTS_GETTANGENTOUTPUT > 0
 		// Blend in the negative intersection normal to create some concavity
 		// Not great as it ties the concavity to the convexity of the cornea surface
 		// No good justification for that. On the other hand, if we're just looking to
@@ -1359,66 +1148,73 @@ RAY_TRACING_ENTRY_CLOSEST_HIT(PathTracingMaterialCHS,
 		const float3 CausticNormal = normalize(lerp(PlaneNormal, -Payload.WorldNormal, IrisMask * IrisDistance));
 		Payload.SetEyeCausticNormal(CausticNormal);
 		// NOTE: calling AdjustShadingNormal on the custom normal does not seem necessary since this is only used for a diffuse calculation
-#else
+	#else
 		const float3 PlaneNormal = Payload.WorldNormal;
 		Payload.SetEyeCausticNormal(Payload.WorldNormal);
-#endif
+	#endif
 
-#if IRIS_NORMAL // on
-		Payload.Specular = 0.25;
+	#if IRIS_NORMAL // on
+		// Match BasePassPixelShader logic, Specular part of g-buffer is re-use, so hard code the term
+		#if SUBSTRATE_CONVERT_FROM_LEGACY
+			Payload.SpecularColor = DielectricSpecularToF0(0.25);
+		#else
+			Payload.Specular = 0.25;
+		#endif
 
-	#if NUM_MATERIAL_OUTPUTS_CLEARCOATBOTTOMNORMAL > 0
-		float3 IrisNormal = normalize( ClearCoatBottomNormal0(MaterialParameters) );
-		#if MATERIAL_TANGENTSPACENORMAL
-		IrisNormal = normalize( TransformTangentVectorToWorld( MaterialParameters.TangentToWorld, IrisNormal ) );
+		#if NUM_MATERIAL_OUTPUTS_CLEARCOATBOTTOMNORMAL > 0
+			float3 IrisNormal = normalize( ClearCoatBottomNormal0(MaterialParameters) );
+			#if MATERIAL_TANGENTSPACENORMAL
+				IrisNormal = normalize( TransformTangentVectorToWorld( MaterialParameters.TangentToWorld, IrisNormal ) );
+			#endif
+		#else
+			float3 IrisNormal = PlaneNormal;
 		#endif
-	#else
-		float3 IrisNormal = PlaneNormal;
-	#endif
 		Payload.SetEyeIrisNormal(IrisNormal);
-#else
-		// IRIS_NORMAL off
+	#else // IRIS_NORMAL = off
 		Payload.SetEyeIrisNormal(PlaneNormal);
-#endif
+	#endif
 
-#if HAVE_GetSubsurfaceMediumMaterialOutput0
-		Payload.SetSubsurfaceRadius(GetSubsurfaceMediumMaterialOutput0(MaterialParameters));
-#else		
+		
+	#if HAVE_GetSubsurfaceMediumMaterialOutput0
+		float3 SSSRadius = GetSubsurfaceMediumMaterialOutput0(MaterialParameters);
+	#else
+		float3 SSSRadius = 0;
 		const uint ProfileId = ExtractSubsurfaceProfileInt(GetMaterialSubsurfaceDataRaw(PixelMaterialInputs).a);
 		if (GetSubsurfaceProfileUseBurley(ProfileId))
 		{
 			// Decode SSS radius
 			const float Opacity = GetMaterialOpacity(PixelMaterialInputs);
-			const float3 DiffuseColor = Payload.BaseColor * (1.0 - Payload.Metallic);
-			Payload.SetSubsurfaceRadius(DecodeSSSProfileRadius(ProfileId, DiffuseColor, Opacity));
-		}
-		else
-		{
-			Payload.SetSubsurfaceRadius(0.0);
+			const float3 DiffuseColor = ComputeDiffuseAlbedo(BaseColor, Metallic);
+			SSSRadius = DecodeSSSProfileRadius(ProfileId, DiffuseColor, Opacity);
 		}
-#endif
+	#endif
 
-#if HAVE_GetSubsurfaceMediumMaterialOutput1
-		float ScatteringDistribution = clamp(GetSubsurfaceMediumMaterialOutput1(MaterialParameters), -0.99f, 0.99f);
-		Payload.SetSubsurfacePhaseFunction(ScatteringDistribution);
-#else
-		Payload.SetSubsurfacePhaseFunction(0.0);
-#endif
+	#if HAVE_GetSubsurfaceMediumMaterialOutput1
+		float PhaseG = clamp(GetSubsurfaceMediumMaterialOutput1(MaterialParameters), -0.99f, 0.99f);
+	#else
+		float PhaseG = 0;
+	#endif
 
+	#if USE_LEGACY_CODEPATH_FOR_SUBSTRATE
+		Payload.MeanFreePath = SSSRadius;
+		Payload.PhaseG = PhaseG;
+	#else
+		Payload.SetSubsurfaceRadius(SSSRadius);
+		Payload.SetSubsurfacePhaseFunction(PhaseG);
+	#endif
 	}
 	else
-#endif
+#endif // MATERIAL_SHADINGMODEL_EYE
 
 #if MATERIAL_SHADINGMODEL_DEFAULT_LIT
 	if (Payload.ShadingModelID == SHADINGMODELID_DEFAULT_LIT)
 	{
+		#if USE_LEGACY_CODEPATH_FOR_SUBSTRATE
+		Payload.ShadingModelID = SHADINGMODELID_SUBSTRATE;
+		#endif
+
+	#if PATH_TRACER_REFRACTION_ENABLED
 		// only allow refraction for default lit materials since we need space for the absorption amount
-#if HAVE_GetAbsorptionMediumMaterialOutput0
-		Payload.SetExtinction(PathTracingGlassTransmittanceToExtinction(GetAbsorptionMediumMaterialOutput0(MaterialParameters)));
-#else
-		// Make the glass totally clear if no custom medium is set
-		Payload.SetExtinction(0.0);
-#endif
 		if (Ior > 0.0)
 		{
 			// Our material is a blend between glass and default_lit -- use stochastic blending between the two
@@ -1433,20 +1229,47 @@ RAY_TRACING_ENTRY_CLOSEST_HIT(PathTracingMaterialCHS,
 				// we are using actual refraction, so disable transparency
 				Payload.Ior = Ior;
 				Payload.ShadingModelID = SHADINGMODELID_SOLID_GLASS;
+
+			#if HAVE_GetAbsorptionMediumMaterialOutput0
+				float3 Extinction = PathTracingGlassTransmittanceToExtinction(GetAbsorptionMediumMaterialOutput0(MaterialParameters));
+			#else
+				// Make the glass totally clear if no custom medium is set
+				float Extinction = 0;
+			#endif
+
+			#if USE_LEGACY_CODEPATH_FOR_SUBSTRATE
+				Payload.MeanFreePath = rcp(max(Extinction, 1e-5f));
+			#else
+				Payload.SetExtinction(Extinction);
+			#endif
 			}
 			Payload.BSDFOpacity = 1.0; // cancels out with the probability of selecting the lobe
 		}
+	#endif
 	}
 	else
-#endif
+#endif // MATERIAL_SHADINGMODEL_DEFAULT_LIT
 	{
 		// terminal case for the conditionals above
 	}
 
 	// adjust after everything else (because SSS profile case needs to decode the dual spec info from the unmodified roughness)
+#if USE_LEGACY_CODEPATH_FOR_SUBSTRATE
+	Payload.RoughnessData.x = AdjustMaterialRoughness(Payload.RoughnessData.x, PackedPayload.GetPathRoughness());
+	Payload.RoughnessData.y = AdjustMaterialRoughness(Payload.RoughnessData.y, PackedPayload.GetPathRoughness());
+#else
 	Payload.Roughness = AdjustMaterialRoughness(Payload.Roughness, PackedPayload.GetPathRoughness());
+#endif
+
+#endif // SUBSTRATE_ENABLED && !USE_LEGACY_CODEPATH_FOR_SUBSTRATE
 
-#endif // SUBSTRATE_ENABLED && TEMPLATE_USES_SUBSTRATE 
+#if MATERIAL_TWOSIDED == 0
+	if (MaterialParameters.TwoSidedSign < 0)
+	{
+		// when viewing the surface from "inside", don't include emission
+		Payload.Radiance = 0;
+	}
+#endif
 
 	// Only opaque surfaces can receive decals. This is both for compatibility with raster and improving performance when we have stacks of transparent surfaces
 #if MATERIALBLENDING_SOLID || MATERIALBLENDING_MASKED
@@ -1588,7 +1411,7 @@ RAY_TRACING_ENTRY_ANY_HIT(PathTracingMaterialAHS,
 
 	{ // TODO: Cleanup this indentation level
 
-#if SUBSTRATE_ENABLED && TEMPLATE_USES_SUBSTRATE
+#if SUBSTRATE_ENABLED && !USE_LEGACY_CODEPATH_FOR_SUBSTRATE
 
 #if MATERIAL_IS_SUBSTRATE 
 		FSubstratePixelHeader SubstratePixelHeader = MaterialParameters.GetFrontSubstrateHeader();
@@ -1620,7 +1443,8 @@ RAY_TRACING_ENTRY_ANY_HIT(PathTracingMaterialAHS,
 		float3 RoughnessData = 0.0;
 		float NoV = 0.0f;
 		float3 LocalSigmaT = 0.0f;
-		float Ior = 0.0;
+		float Ior = 1.0;
+		float RefrProb  = 0.0;
 		Substrate_for_unroll(int BSDFIdx = 0, BSDFIdx < SubstratePixelHeader.SubstrateTree.BSDFCount, ++BSDFIdx)
 		{
 #define BSDF SubstratePixelHeader.SubstrateTree.BSDFs[BSDFIdx]
@@ -1643,13 +1467,18 @@ RAY_TRACING_ENTRY_ANY_HIT(PathTracingMaterialAHS,
 							RoughnessData.z = Haziness.Weight;
 						}
 
-						if (BSDF_GETSSSTYPE(BSDF) == SSS_TYPE_NONE)
+						if (BSDF_GETSSSTYPE(BSDF) == SSS_TYPE_DIFFUSION)
 						{
 							// only non-SSS slabs 
 							LocalSigmaT = rcp(max(SLAB_SSSMFP(BSDF), 0.0001));
-#if MATERIALBLENDING_TRANSLUCENT && (SUBSTRATE_BLENDING_TRANSLUCENT_GREYTRANSMITTANCE || SUBSTRATE_BLENDING_TRANSLUCENT_COLOREDTRANSMITTANCE) && PATH_TRACER_REFRACTION_ENABLED
+					#if MATERIALBLENDING_TRANSLUCENT && (SUBSTRATE_BLENDING_TRANSLUCENT_GREYTRANSMITTANCE || SUBSTRATE_BLENDING_TRANSLUCENT_COLOREDTRANSMITTANCE)
+						#if PATH_TRACER_REFRACTION_ENABLED
 							Ior = GetRefractionIor(PixelMaterialInputs, F0);
-#endif
+						#endif
+							float SlabMix = LobeSelectionProb(SLAB_DIFFUSEALBEDO(BSDF), exp(-SUBSTRATE_SIMPLEVOLUME_THICKNESS_CM * LocalSigmaT));
+							RefrProb = 1 - SlabMix;
+
+					#endif
 						}
 						break;
 					}
@@ -1665,9 +1494,6 @@ RAY_TRACING_ENTRY_ANY_HIT(PathTracingMaterialAHS,
 						break;
 					}
 				}
-			#if SUBSTRATE_CLAMPED_CLOSURE_COUNT > 1 // Only break when Substrate_for_unroll is a for loop.
-				break; // Found BSDF bottom, exit loop
-			#endif
 			}
 #undef BSDF
 		}
@@ -1678,13 +1504,13 @@ RAY_TRACING_ENTRY_ANY_HIT(PathTracingMaterialAHS,
 
 		// SUBSTRATE_TODO: Revisit this -- TotalTransmittancePreCoverage should probably account for this in the thin case
 		float3 Transmission = exp(-LocalSigmaT * SUBSTRATE_SIMPLEVOLUME_THICKNESS_CM);
-		float3 Tr = ComputeThinSlabWeights(Transmission, NoV, Ior, F0).Transmitted;
+		float3 Tr = RefrProb * ComputeThinSlabWeights(Transmission, NoV, Ior, F0).Transmitted;
 #else
 		HLSL_STATIC_ASSERT(PATH_TRACER_REFRACTION_ENABLED == 1, "Expect refraction to be true here");
 		// Solid Glass Fake Caustics
-		float Tr = FakeCaustics(F0, Ior, NoV);
+		float Tr = RefrProb * FakeCaustics(F0, Ior, NoV);
 		// Account for transmission 
-		PackedPayload.SetTau(PackedPayload.GetTau() + TotalCoverage * ComputeOpticalDepthForHit(LocalSigmaT, bIsFrontFace));
+		PackedPayload.SetTau(PackedPayload.GetTau() + TotalCoverage * RefrProb * ComputeOpticalDepthForHit(LocalSigmaT, bIsFrontFace));
 #endif
 
 		// If material is doing refraction, apply an extra factor to account for ApproximateCaustics
@@ -1708,7 +1534,7 @@ RAY_TRACING_ENTRY_ANY_HIT(PathTracingMaterialAHS,
 		const float3 Transparency = 0.0;
 #endif // MATERIAL_IS_SUBSTRATE
 
-#else // SUBSTRATE_ENABLED && TEMPLATE_USES_SUBSTRATE
+#else // SUBSTRATE_ENABLED && !USE_LEGACY_CODEPATH_FOR_SUBSTRATE
 
 #if MATERIALBLENDING_MODULATE
 		const float3 Transparency = GetMaterialEmissive(PixelMaterialInputs);
@@ -1812,7 +1638,7 @@ RAY_TRACING_ENTRY_ANY_HIT(PathTracingMaterialAHS,
 #error Unhandled blending mode!
 #endif
 
-#endif // SUBSTRATE_ENABLED && TEMPLATE_USES_SUBSTRATE
+#endif // SUBSTRATE_ENABLED && !USE_LEGACY_CODEPATH_FOR_SUBSTRATE
 
 		// Update the ray throughput (it is packed simply into the payload since we don't need to carry any other information across hits)
 		float3 RayThroughput = PackedPayload.GetRayThroughput();
diff --git a/Engine/Shaders/Private/ShadingEnergyConservation.ush b/Engine/Shaders/Private/ShadingEnergyConservation.ush
--- a/Engine/Shaders/Private/ShadingEnergyConservation.ush
+++ b/Engine/Shaders/Private/ShadingEnergyConservation.ush
@@ -6,8 +6,16 @@
 // * If Substrate is enabled: ON by default
 // * If Substrate is enabled with legacy blendable GBuffer: OFF by default for performance parity with legacy on low end platforms.
 // * If Substrate is disabled: OFF by default
-#if (SUBSTRATE_ENABLED && SUBSTRATE_GBUFFER_FORMAT>0) || LEGACY_MATERIAL_ENERGYCONSERVATION
-	#undef USE_ENERGY_CONSERVATION
+#ifdef USE_ENERGY_CONSERVATION
+#error "USE_ENERGY_CONSERVATION is defined by this file, it should not be defined prior to including it"
+#endif
+
+#if PATH_TRACING && (SUBSTRATE_ENABLED || LEGACY_MATERIAL_ENERGYCONSERVATION)
+	// When using the path tracer, we use energy conservation either when:
+	//   * Substrate is enabled for the project
+	//   * It is enabled in the project for legacy materials
+	#define USE_ENERGY_CONSERVATION 1
+#elif (SUBSTRATE_ENABLED && SUBSTRATE_GBUFFER_FORMAT>0) || LEGACY_MATERIAL_ENERGYCONSERVATION
 	#if SUPPORTS_INDEPENDENT_SAMPLERS 
 		#define USE_ENERGY_CONSERVATION 1
 	#else
@@ -16,7 +24,7 @@
 		#define USE_ENERGY_CONSERVATION 0
 	#endif
 #else
-	#undef USE_ENERGY_CONSERVATION
+	// SUBSTRATE_GBUFFER_FORMAT == 0 || LEGACY_MATERIAL_ENERGYCONSERVATION == 0
 	#define USE_ENERGY_CONSERVATION 0
 #endif
 
diff --git a/Engine/Source/Runtime/Core/Private/UObject/FortniteMainBranchObjectVersion.cpp b/Engine/Source/Runtime/Core/Private/UObject/FortniteMainBranchObjectVersion.cpp
--- a/Engine/Source/Runtime/Core/Private/UObject/FortniteMainBranchObjectVersion.cpp
+++ b/Engine/Source/Runtime/Core/Private/UObject/FortniteMainBranchObjectVersion.cpp
@@ -13,7 +13,7 @@ TMap<FGuid, FGuid> FFortniteMainBranchObjectVersion::GetSystemGuids()
 	SystemGuids.Add(DevGuids.MATERIALSHADERMAP_DERIVEDDATA_VER, FGuid("F051877974B44D409435F321BDB39DBB"));
 	SystemGuids.Add(DevGuids.NIAGARASHADERMAP_DERIVEDDATA_VER, FGuid("6360A977062842A29ED30E3A7ACB0E64"));
 	SystemGuids.Add(DevGuids.Niagara_LatestScriptCompileVersion, FGuid("783FF3700BB641C799B2D43552FBF91E"));
-	SystemGuids.Add(DevGuids.MaterialTranslationDDCVersion, FGuid("F9BB134100414BE4AA156DEF0C01C4DF"));
+	SystemGuids.Add(DevGuids.MaterialTranslationDDCVersion, FGuid("5609B29D40D44A89967F303EAEDED2C0"));
 
 	// These GUIDs are in AppendRestrictedFortniteMainStreamObjectVersionGuids as they should rarely be changed.
 	//SystemGuids.Add(DevGuids.NANITE_DERIVEDDATA_VER, FGuid("ACF98184978742A592537DA223A1E6BE"));
diff --git a/Engine/Source/Runtime/Engine/Private/Materials/HLSLMaterialTranslator.cpp b/Engine/Source/Runtime/Engine/Private/Materials/HLSLMaterialTranslator.cpp
--- a/Engine/Source/Runtime/Engine/Private/Materials/HLSLMaterialTranslator.cpp
+++ b/Engine/Source/Runtime/Engine/Private/Materials/HLSLMaterialTranslator.cpp
@@ -1810,8 +1810,7 @@ void FHLSLMaterialTranslator::TranslateMaterial()
 		NormalCodeChunkEnd = SharedPropertyCodeChunks[NormalShaderFrequency].Num();
 	}
 
-	if (!bSubstrateEnabled || !bMaterialUsesRootNodeToSubstrateHiddenConversion
-		|| GetSubstrateMaterialExportType() != ESubstrateMaterialExport::SME_None)	// Needed for material preview on node or any other export.
+	// Output all chunks types for legacy materials (in case shader code wants to work with legacy attributes directly)
 	{
 		// Rest of properties
 		Chunk[MP_EmissiveColor]	= Material->CompilePropertyAndSetMaterialProperty(MP_EmissiveColor,	this);
@@ -1961,7 +1960,6 @@ void FHLSLMaterialTranslator::TranslateMaterial()
 		CHECK_DDC_QUERY_FINISHED_ELSE_RETURN();
 	}
 
-	if (!bSubstrateEnabled || !bMaterialUsesRootNodeToSubstrateHiddenConversion)
 	{
 		Chunk[MP_CustomData0] = Material->CompilePropertyAndSetMaterialProperty(MP_CustomData0, this);
 		Chunk[MP_CustomData1] = Material->CompilePropertyAndSetMaterialProperty(MP_CustomData1, this);
