Update shader assertion macros to try to make FXC happy when optimizations are disabled

diff --git a/Engine/Shaders/Private/AssertionMacros.ush b/Engine/Shaders/Private/AssertionMacros.ush
--- a/Engine/Shaders/Private/AssertionMacros.ush
+++ b/Engine/Shaders/Private/AssertionMacros.ush
@@ -58,10 +58,12 @@
 	}
 
 #else // !DO_CHECK
-	#define check(expr) { }
+	void CheckNop() {}
 
-	#define checkBufferAccess(buffer, index) { }
-	#define checkStructuredBufferAccess(buffer, index) { }
+	#define check(expr) CheckNop()
+
+	#define checkBufferAccess(buffer, index) CheckNop()
+	#define checkStructuredBufferAccess(buffer, index) CheckNop()
 #endif // DO_CHECK
 
 #if DO_GUARD_SLOW
@@ -71,8 +73,10 @@
 	#define checkStructuredBufferAccessSlow(buffer, index)	checkStructuredBufferAccess(buffer, index)
 
 #else // DO_GUARD_SLOW
-	#define checkSlow(expr) { }
+	void GuardSlowNop() {}
+
+	#define checkSlow(expr) GuardSlowNop()
 
-	#define checkBufferAccessSlow(buffer, index) { }
-	#define checkStructuredBufferAccessSlow(buffer, index) { }
+	#define checkBufferAccessSlow(buffer, index) GuardSlowNop()
+	#define checkStructuredBufferAccessSlow(buffer, index) GuardSlowNop()
 #endif // DO_GUARD_SLOW
