Fix Furnace test Analytic for Substrate Format=0.
Change FurnaceTest shader to only be compiled if enabled.

diff --git a/Engine/Shaders/Private/ShadingFurnaceTest.usf b/Engine/Shaders/Private/ShadingFurnaceTest.usf
--- a/Engine/Shaders/Private/ShadingFurnaceTest.usf
+++ b/Engine/Shaders/Private/ShadingFurnaceTest.usf
@@ -4,11 +4,29 @@
 
 #include "LargeWorldCoordinates.ush" 
 
+#if SUBSTRATE_ENABLED
+// In order to avoid numerical issue with integrating low roughness value, incease the min. GGX roughness
+#include "/Engine/Shared/SubstrateDefinitions.h"
+#ifdef SUBSTRATE_MIN_GGX_ROUGHNESS
+#undef SUBSTRATE_MIN_GGX_ROUGHNESS
+#define SUBSTRATE_MIN_GGX_ROUGHNESS 0.04
+#endif
+#endif
+
 #define SUPPORT_CONTACT_SHADOWS 0
 #define SHADING_PATH_DEFERRED 1
 #define USE_SOURCE_TEXTURE 0
 #define DEBUG_ENABLE 0
 #define SUBSTRATE_SSS_TRANSMISSION USE_TRANSMISSION
+#if SUBSTRATE_GBUFFER_FORMAT==0
+#define SUBSTRATE_LOAD_FROM_MATERIALCONTAINER 0	// Load from a specified BSDF, not from MAterialBuffer.
+#define SUBSTRATE_INLINE_SHADING 1				// Inline shading for normal setup if needed.
+#define SUBSTRATE_SSS_MATERIAL_OVERRIDE 1		// Need for SSS dither.
+#include "Substrate/SubstrateRead.ush"
+#else 
+#define SUBSTRATE_LOAD_FROM_MATERIALCONTAINER 1
+#endif
+
 #include "Substrate/Substrate.ush"
 #include "DeferredShadingCommon.ush"
 #include "DeferredLightingCommon.ush"
@@ -21,6 +39,7 @@
 
 #include "Substrate/SubstrateEvaluation.ush"
 #include "Substrate/SubstrateDeferredLighting.ush"
+#include "Substrate/SubstrateRead.ush"
 
 #if SHADER_FURNACE_ANALYTIC
 
@@ -44,23 +63,43 @@ void MainPS(
 	float ClearCoatRoughness = 0.5f;
 	float3 N = float3(0, 0, 1);
 	float SceneDepth = 0;
-#if SUBSTRATE_ENABLED && SUBSTRATE_GBUFFER_FORMAT==1
+#if SUBSTRATE_ENABLED
 	{
+		#if SUBSTRATE_LOAD_FROM_MATERIALCONTAINER
 		// Fetch the first layer of roughness & normal only. This fine for simple material but will not work for clear coat.
 		FSubstrateAddressing SubstrateAddressing = GetSubstratePixelDataByteOffset(PixelPos, uint2(View.BufferSizeAndInvSize.xy), Substrate.MaxBytesPerPixel);
 		FSubstratePixelHeader SubstratePixelHeader = UnpackSubstrateHeaderIn(Substrate.MaterialTextureArray, SubstrateAddressing, Substrate.TopLayerTexture);
 		bIsValid = SubstratePixelHeader.ClosureCount > 0;
+		#else
+		FScreenSpaceData ScreenSpaceData = GetScreenSpaceData(ScreenUV, true);
+		bIsValid = ScreenSpaceData.GBuffer.ShadingModelID != SHADINGMODELID_UNLIT;
+		#endif
 		if (bIsValid)
 		{
 			const float3 DummyV = float3(0, 0, 1);
 			const float3 DummyL = float3(0, 0, 1);
+			#if SUBSTRATE_LOAD_FROM_MATERIALCONTAINER
 			FSubstrateBSDF BSDF = UnpackSubstrateBSDFIn(Substrate.MaterialTextureArray, SubstrateAddressing, SubstratePixelHeader);
 			FSubstrateBSDFContext BSDFContext = SubstrateCreateBSDFContext(SubstratePixelHeader, BSDF, SubstrateAddressing, DummyV, DummyL);
+			#else
+			FSubstrateGBufferBSDF GBufferBSDF = SubstrateReadGBufferBSDF(ScreenSpaceData);
+			FSubstrateBSDFContext BSDFContext = SubstrateCreateBSDFContext(GBufferBSDF.BSDF, GBufferBSDF.BSDFTangentBasis, DummyV, DummyL);
+			#endif
 
 			SceneDepth = CalcSceneDepth(ScreenUV);
 			N = BSDFContext.N;
-			Roughness = SubstrateGetBSDFRoughness(BSDFContext.BSDF);
+			Roughness = MakeRoughnessSafe(SubstrateGetBSDFRoughness(BSDFContext.BSDF), SUBSTRATE_MIN_GGX_ROUGHNESS);
+
 			ClearCoatRoughness = Roughness;
+			if (BSDF_GETTYPE(BSDFContext.BSDF) == SUBSTRATE_BSDF_TYPE_SLAB && BSDF_GETHASHAZINESS(BSDFContext.BSDF))
+			{
+				const FHaziness Haziness = UnpackHaziness(SLAB_HAZINESS(BSDFContext.BSDF));
+				BRANCH
+				if (Haziness.bSimpleClearCoat)
+				{
+					ClearCoatRoughness = MakeRoughnessSafe(Haziness.Roughness, SUBSTRATE_MIN_GGX_ROUGHNESS);
+				}
+			}
 		}
 	}
 #else
@@ -159,7 +198,7 @@ void MainPS(
 				{
 					InvWeight += Square(PDFs[j] * NumSamples[j]);
 				}
-				float Weight = (InvWeight > 0.f ? rcp(InvWeight) : 0.f) * PDFs[Set] * NumSamples[Set];
+				float Weight = (InvWeight > 1e-5f ? rcp(InvWeight) : 0.f) * PDFs[Set] * NumSamples[Set];
 
 				#if DEBUG_ENABLE
 				if (bDebugEnabled)
@@ -219,30 +258,53 @@ void MainPS(
 						uint2(PixelPos), 
 						SurfaceShadow).xyz;
 				}
-			#elif 0
+			#else
 				{
+
+				#if SUBSTRATE_LOAD_FROM_MATERIALCONTAINER
 					FSubstrateAddressing SubstrateAddressing = GetSubstratePixelDataByteOffset(PixelPos, uint2(View.BufferSizeAndInvSize.xy), Substrate.MaxBytesPerPixel);
 					FSubstratePixelHeader SubstratePixelHeader = UnpackSubstrateHeaderIn(Substrate.MaterialTextureArray, SubstrateAddressing, Substrate.TopLayerTexture);
+					const bool bMaterialIsLit = SubstratePixelHeader.ClosureCount > 0;	// This test is also enough to exclude sky pixels
+					#if SUBSTRATE_INLINE_SHADING
+					HEADER_SETIRRADIANCE_AO(SubstratePixelHeader.State, SubstratePackIrradianceAndOcclusion(InitIrradianceAndOcclusion(1.f));
+					#endif
+				#else
+					FScreenSpaceData ScreenSpaceData = GetScreenSpaceData(ScreenUV, true);
+					ScreenSpaceData.AmbientOcclusion = 1.f;
+					const FGBufferData GBuffer = ScreenSpaceData.GBuffer;
+					const bool bMaterialIsLit = GBuffer.ShadingModelID != SHADINGMODELID_UNLIT;
+				#endif
+
 					const float LightMask = 1;
 
 					BRANCH
-					if (SubstratePixelHeader.BSDFCount > 0)
+					if (bMaterialIsLit)
 					{
-						HEADER_SETIRRADIANCE_AO(SubstratePixelHeader.State, SubstratePackIrradianceAndOcclusion(1.0, 1.0));
+
+						#if SUBSTRATE_GBUFFER_FORMAT==0
+						FSubstrateGBufferBSDF GBufferBSDF = SubstrateReadGBufferBSDF(ScreenSpaceData);
+						#endif
 
 						FSubstrateShadowTermInputParameters SubstrateShadowTermInputParameters	= GetInitialisedSubstrateShadowTermInputParameters();
 						SubstrateShadowTermInputParameters.bEvaluateShadowTerm					= false;
-						SampleRadiance = SubstrateDeferredLighting(
+						FSubstrateDeferredLighting SampleLighting = SubstrateDeferredLighting(
 							LightData,
 							V,
 							L,
 							L,
 							LightMask,
 							SubstrateShadowTermInputParameters,
-							RectTexture,
+							#if SUBSTRATE_LOAD_FROM_MATERIALCONTAINER
 							Substrate.MaterialTextureArray,
 							SubstrateAddressing,
-							SubstratePixelHeader).xyz;
+							SubstratePixelHeader
+							#else
+							GBufferBSDF.BSDF,
+							GBufferBSDF.BSDFTangentBasis,
+							GBufferBSDF.BSDFAO
+							#endif
+						);
+						SampleRadiance.xyz = SampleLighting.TotalDiffuseLighting + SampleLighting.TotalSpecularLighting;
 					}
 				}
 			#endif
diff --git a/Engine/Source/Runtime/Renderer/Private/ShadingEnergyConservation.cpp b/Engine/Source/Runtime/Renderer/Private/ShadingEnergyConservation.cpp
--- a/Engine/Source/Runtime/Renderer/Private/ShadingEnergyConservation.cpp
+++ b/Engine/Source/Runtime/Renderer/Private/ShadingEnergyConservation.cpp
@@ -36,6 +36,12 @@ static TAutoConsoleVariable<int32> CVarShadingFurnaceTest(
 	TEXT("Enable/disable furnace for shading validation."),
 	ECVF_RenderThreadSafe);
 
+static TAutoConsoleVariable<int32> CVarShadingFurnaceTest_EnableForProject(
+	TEXT("r.Shading.FurnaceTest.EnableForProject"),
+	0,
+	TEXT("Enable FurnaceTest rendering and shader code."),
+	ECVF_ReadOnly | ECVF_RenderThreadSafe);
+
 static TAutoConsoleVariable<int32> CVarShadingFurnaceTest_SampleCount(
 	TEXT("r.Shading.FurnaceTest.SampleCount"),
 	64,
@@ -290,14 +296,19 @@ class FShadingFurnaceTestPassPS : public FGlobalShader
 		RENDER_TARGET_BINDING_SLOTS()
 	END_SHADER_PARAMETER_STRUCT()
 
-		static FPermutationDomain RemapPermutation(FPermutationDomain PermutationVector)
+	static FPermutationDomain RemapPermutation(FPermutationDomain PermutationVector)
 	{
 		return PermutationVector;
 	}
 
+	static bool IsSupported(const EShaderPlatform& InPlaform)
+	{
+		return CVarShadingFurnaceTest_EnableForProject.GetValueOnAnyThread() && GetMaxSupportedFeatureLevel(InPlaform) >= ERHIFeatureLevel::SM5;
+	}
+
 	static bool ShouldCompilePermutation(const FGlobalShaderPermutationParameters& Parameters)
 	{
-		return GetMaxSupportedFeatureLevel(Parameters.Platform) >= ERHIFeatureLevel::SM5;
+		return IsSupported(Parameters.Platform);
 	}
 
 	static void ModifyCompilationEnvironment(const FGlobalShaderPermutationParameters& Parameters, FShaderCompilerEnvironment& OutEnvironment)
@@ -315,6 +326,11 @@ static void AddShadingFurnacePass(
 	TRDGUniformBufferRef<FSceneTextureUniformParameters>& SceneTexturesUniformBuffer,
 	FRDGTextureRef OutTexture)
 {
+	if (!FShadingFurnaceTestPassPS::IsSupported(View.GetShaderPlatform()))
+	{
+		return;
+	}
+
 	TShaderMapRef<FShadingFurnaceTestPassPS> PixelShader(View.ShaderMap);
 	FShadingFurnaceTestPassPS::FParameters* Parameters = GraphBuilder.AllocParameters<FShadingFurnaceTestPassPS::FParameters>();
 	Parameters->ViewUniformBuffer				= View.ViewUniformBuffer;
