Fix light complexity debug mode with Substrate.

diff --git a/Engine/Shaders/Private/DeferredLightPixelShaders.usf b/Engine/Shaders/Private/DeferredLightPixelShaders.usf
--- a/Engine/Shaders/Private/DeferredLightPixelShaders.usf
+++ b/Engine/Shaders/Private/DeferredLightPixelShaders.usf
@@ -363,7 +363,11 @@ void DeferredLightPixelMain(
 	}
 
 #if VISUALIZE_LIGHT_CULLING == 1
-	OutColor = ConvertEstimatedCostToColor(EstimatedCost);
+	// Convert estimated cost into a color code.
+	// Legacy stored the color code into both diffuse and specular term, and added. 
+	// Doing the same to match the legacy output.
+	const float4 CostWithColorCode = ConvertEstimatedCostToColor(EstimatedCost);
+	OutColor = CostWithColorCode + CostWithColorCode;
 	#if SUBSTRATE_OPAQUE_ROUGH_REFRACTION_ENABLED
 	OutOpaqueRoughRefractionSceneColor = 0.0f;
 	OutSubSurfaceSceneColor = 0.0f;
