Fix capsule area lighting with anisotropy material when using Substrate.

diff --git a/Engine/Shaders/Private/Substrate/SubstrateEvaluation.ush b/Engine/Shaders/Private/Substrate/SubstrateEvaluation.ush
--- a/Engine/Shaders/Private/Substrate/SubstrateEvaluation.ush
+++ b/Engine/Shaders/Private/Substrate/SubstrateEvaluation.ush
@@ -580,7 +580,7 @@ FSubstrateEvaluateResult SubstrateEvaluateBSDFCommon(FSubstrateBSDFContext BSDFC
 							D = Substrate_D_Glint(BSDFContext, SafeRoughness, D);
 						}
 
-						Vis = Substrate_Vis_GGX_Aniso(AlphaXY.x, AlphaXY.y, BSDFContext.Context.NoV, BSDFContext.SatNoL, BSDFContext.Context.XoV, BSDFContext.Context.XoL, BSDFContext.Context.YoV, BSDFContext.Context.YoL);
+						Vis = Substrate_Vis_GGX_Aniso(AlphaXY.x, AlphaXY.y, BSDFContext.Context.NoV, AreaLight.NoL, BSDFContext.Context.XoV, BSDFContext.Context.XoL, BSDFContext.Context.YoV, BSDFContext.Context.YoL);
 						const float H_PDF = VisibleGGXPDF_aniso(BSDFContext.TangentV, BSDFContext.TangentH, AlphaXY);
 						PDF = RayPDFToReflectionRayPDF(BSDFContext.Context.VoH, H_PDF);
 					}
