Fix eye shading differences with legacy.

This fixes the failing eye tests on EngineTests.

diff --git a/Engine/Shaders/Private/Substrate/SubstrateEvaluation.ush b/Engine/Shaders/Private/Substrate/SubstrateEvaluation.ush
--- a/Engine/Shaders/Private/Substrate/SubstrateEvaluation.ush
+++ b/Engine/Shaders/Private/Substrate/SubstrateEvaluation.ush
@@ -1174,7 +1174,7 @@ FSubstrateEvaluateResult SubstrateEvaluateBSDFCommon(FSubstrateBSDFContext BSDFC
 			float Alpha2 = Pow4(SafeRoughness);
 
 			Init(BSDFContext.Context, BSDFContext.N, BSDFContext.V, AreaLight.SpecularL);
-			SphereMaxNoH(BSDFContext.Context, AreaLight.SphereSinAlpha, true);
+			SphereMaxNoH(BSDFContext.Context, AreaLight.SphereSinAlpha, false);
 
 			BSDFContext.Context.NoV = saturate(max(abs(BSDFContext.Context.NoV), SUBSTRATE_EPSILON));
 
diff --git a/Engine/Shaders/Private/Substrate/SubstrateRead.ush b/Engine/Shaders/Private/Substrate/SubstrateRead.ush
--- a/Engine/Shaders/Private/Substrate/SubstrateRead.ush
+++ b/Engine/Shaders/Private/Substrate/SubstrateRead.ush
@@ -105,6 +105,7 @@ FSubstrateGBufferBSDF SubstrateReadGBufferBSDF(in FScreenSpaceData ScreenSpaceDa
 			IrisDistance = 1.0; // Set at 1 in order to let's the mask drive the interpolation. See SubstrateEvaluation.ush.
 		#else
 			IrisNormal = OctahedronToUnitVector(GBuffer.CustomData.yz * 2 - 1);
+			IrisPlaneNormal = IrisNormal;
 			IrisDistance = GBuffer.StoredMetallic;
 		#endif
 			BSDF = GetSubstrateEyeBSDF(DiffuseAlbedo, F0, GBuffer.Roughness, IrisMask, IrisDistance, IrisNormal, IrisPlaneNormal, SSSProfileID, Emissive, SharedLocalBasisIndex).InlinedBSDF;
