Path Tracer: Fix substrate eye and SLW cases legacy lobes

Fix tangent space normal handling and restore shading normal adjustment lost in previous refactor

#rnx

diff --git a/Engine/Shaders/Private/PathTracing/PathTracingMaterialHitShader.usf b/Engine/Shaders/Private/PathTracing/PathTracingMaterialHitShader.usf
--- a/Engine/Shaders/Private/PathTracing/PathTracingMaterialHitShader.usf
+++ b/Engine/Shaders/Private/PathTracing/PathTracingMaterialHitShader.usf
@@ -365,6 +365,15 @@ RAY_TRACING_ENTRY_CLOSEST_HIT(PathTracingMaterialCHS,
 	if (BSDFCount > 0)
 	{
 
+		for (int Index = 0, Num = SubstratePixelHeader.SharedLocalBases.Count; Index < Num; Index++)
+		{
+			float3 N = SubstratePixelHeader.SharedLocalBases.Normals[Index];
+#if !MATERIAL_TANGENTSPACENORMAL
+			N *= MaterialParameters.TwoSidedSign;
+#endif
+			SubstratePixelHeader.SharedLocalBases.Normals[Index] = AdjustShadingNormal(N, Payload.WorldGeoNormal, WorldRayDirection());
+		}
+
 
 		int BSDFIdx = 0;
 #define BSDF SubstratePixelHeader.SubstrateTree.BSDFs[BSDFIdx]
@@ -662,6 +671,7 @@ RAY_TRACING_ENTRY_CLOSEST_HIT(PathTracingMaterialCHS,
 				// NOTE: The cases below are for back-compatibility and cannot be layered
 				case SUBSTRATE_BSDF_TYPE_EYE:
 				{
+					SlabSampler.Accept(1); // dummy, avoids special case below
 					Payload.WorldNormal = SubstratePixelHeader.SharedLocalBases.Normals[BSDF_GETSHAREDLOCALBASISID(BSDF)];
 					Payload.DiffuseColor = EYE_DIFFUSEALBEDO(BSDF);
 					Payload.SpecularColor = EYE_F0(BSDF);
@@ -713,6 +723,7 @@ RAY_TRACING_ENTRY_CLOSEST_HIT(PathTracingMaterialCHS,
 				}
 				case SUBSTRATE_BSDF_TYPE_SINGLELAYERWATER:
 				{
+					SlabSampler.Accept(1); // dummy, avoids special case below
 					Payload.WorldNormal = SubstratePixelHeader.SharedLocalBases.Normals[BSDF_GETSHAREDLOCALBASISID(BSDF)];
 					Payload.BSDFOpacity = 1;
 					Payload.ShadingModelID = SHADINGMODELID_SOLID_GLASS;
