Fix WPO for mobile renderer supporting SM5 nodes.

Some meshes were rendering in the prepass using the WorldGridMaterial (which has no WPO), but we were forcing a WPO evaluation during the base pass. This discrepancy in position between the passes caused the mesh to be depth culled.

#tests On console with ODSC

diff --git a/Engine/Shaders/Private/MaterialTemplate.ush b/Engine/Shaders/Private/MaterialTemplate.ush
--- a/Engine/Shaders/Private/MaterialTemplate.ush
+++ b/Engine/Shaders/Private/MaterialTemplate.ush
@@ -3956,7 +3956,7 @@ void ApplyMaterialPreviousFirstPersonTransform(FMaterialVertexParameters Paramet
 bool ShouldEnableWorldPositionOffset(FMaterialVertexParameters Parameters)
 {
 #if USES_WORLD_POSITION_OFFSET
-	#if (FEATURE_LEVEL > FEATURE_LEVEL_ES3_1) // Non optional WPO on mobile
+	#if (FEATURE_LEVEL > FEATURE_LEVEL_ES3_1 || MOBILE_SUPPORTS_SM5_MATERIAL_NODES) // Non optional WPO on mobile not supporting sm5 material nodes
 		if (!Parameters.bEvaluateWorldPositionOffset ||
 			(GetPrimitiveData(Parameters).Flags & PRIMITIVE_SCENE_DATA_FLAG_EVALUATE_WORLD_POSITION_OFFSET) == 0)
 		{
