Fix typo causing all surfaces to have per-pixel SSS enabled with Format=1.

diff --git a/Engine/Shaders/Private/Substrate/Substrate.ush b/Engine/Shaders/Private/Substrate/Substrate.ush
--- a/Engine/Shaders/Private/Substrate/Substrate.ush
+++ b/Engine/Shaders/Private/Substrate/Substrate.ush
@@ -1551,7 +1551,7 @@ FSubstrateData InternalGetSubstrateSlabBSDF(
 
 	SSSMFPScale = saturate(SSSMFPScale);
 	const bool bSSSProfilePlugged = SSSProfileID >= (bSupportDefaultSSSProfile ? 0u : 1u);
-	const bool bMFPPlugged = any(SSSMFP >= 0);
+	const bool bMFPPlugged = any(SSSMFP > 0);
 	const bool bIsThin = bIsAtTheBottomOfTopology && (bMFPPlugged && bIsThinSurface); // Enable thin only if MFP is plugged, as otherwise there will be no back-tramission
 
 	// In forward mode (blend mode != Opaque/Masked), diffusion is not supported. 
