Fix eye shading difference due to incorrect shadow term when using Substrate.

diff --git a/Engine/Shaders/Private/Substrate/SubstrateEvaluation.ush b/Engine/Shaders/Private/Substrate/SubstrateEvaluation.ush
--- a/Engine/Shaders/Private/Substrate/SubstrateEvaluation.ush
+++ b/Engine/Shaders/Private/Substrate/SubstrateEvaluation.ush
@@ -1198,7 +1198,7 @@ FSubstrateEvaluateResult SubstrateEvaluateBSDFCommon(FSubstrateBSDFContext BSDFC
 			const float Sclera = AreaLight.NoL;
 
 			Sample.DiffusePathValue = Diffuse_Lambert(DiffuseColor) * lerp(Sclera, Iris, EYE_IRISMASK(BSDFContext.BSDF));
-			Sample.IntegratedDiffuseValue += (ShadowTerms.SurfaceShadow * /* We don't AreaLight.NoL as it is already applied through the Sclera/Iris value above */ AreaLight.Falloff) * AreaLight.FalloffColor * Sample.DiffusePathValue;
+			Sample.IntegratedDiffuseValue += (ShadowTerms.TransmissionShadow * /* We don't AreaLight.NoL as it is already applied through the Sclera/Iris value above */ AreaLight.Falloff) * AreaLight.FalloffColor * Sample.DiffusePathValue;
 
 			Sample.DiffuseColor		= DiffuseColor;
 			Sample.DiffusePDF		= BSDFContext.SatNoL / PI;
