Fix missing legacy energy conservation term on eye shading when energy conservation is disabled with Substrate.

diff --git a/Engine/Shaders/Private/Substrate/SubstrateEvaluation.ush b/Engine/Shaders/Private/Substrate/SubstrateEvaluation.ush
--- a/Engine/Shaders/Private/Substrate/SubstrateEvaluation.ush
+++ b/Engine/Shaders/Private/Substrate/SubstrateEvaluation.ush
@@ -1208,13 +1208,21 @@ FSubstrateEvaluateResult SubstrateEvaluateBSDFCommon(FSubstrateBSDFContext BSDFC
 			// Apply energy conservation on the diffuse component
 			// If the specular layer is anisotropic, the energy term is computed onto the 'main' roughness [Kulla 2019]
 			//
+			const float3 FresnelTerm = Substrate_F_GGX(F0, F90, BSDFContext.Context.VoH);
 			FBxDFEnergyType MSScale = 1;
 			float DirectionalAlbedo_SpecularTransmission = 1.0f;
+			#if USE_ENERGY_CONSERVATION
 			{
 				FBxDFEnergyTerms EnergyTerms = ComputeGGXSpecEnergyTerms(SafeRoughness, BSDFContext.Context.NoV, F0, F90);
 				DirectionalAlbedo_SpecularTransmission = ComputeEnergyPreservation(EnergyTerms);
 				MSScale = ComputeEnergyConservation(EnergyTerms);
 			}
+			#else
+			{
+				// Preserve old behavior when energy conservation is disabled, which is the case with Format=0
+				DirectionalAlbedo_SpecularTransmission = 1.f - FresnelTerm.x;
+			}
+			#endif
 			
 			Sample.DiffusePathValue *= DirectionalAlbedo_SpecularTransmission;
 			Sample.IntegratedDiffuseValue *= DirectionalAlbedo_SpecularTransmission;
@@ -1260,7 +1268,6 @@ FSubstrateEvaluateResult SubstrateEvaluateBSDFCommon(FSubstrateBSDFContext BSDFC
 				PDF = RayPDFToReflectionRayPDF(BSDFContext.Context.VoH, H_PDF);
 			}
 
-			const float3 FresnelTerm = Substrate_F_GGX(F0, F90, BSDFContext.Context.VoH);
 
 			Sample.SpecularPathProbability	= 1.0f;
 			Sample.SpecularPDF				= PDF * Sample.SpecularPathProbability;
