Fix incorrect GGX's D with area light.

diff --git a/Engine/Shaders/Private/Substrate/SubstrateEvaluation.ush b/Engine/Shaders/Private/Substrate/SubstrateEvaluation.ush
--- a/Engine/Shaders/Private/Substrate/SubstrateEvaluation.ush
+++ b/Engine/Shaders/Private/Substrate/SubstrateEvaluation.ush
@@ -629,7 +629,9 @@ FSubstrateEvaluateResult SubstrateEvaluateBSDFCommon(FSubstrateBSDFContext BSDFC
 							}
 							else
 							{
-							
+								// The energy term needs to be computed before the D evaluation as this function modifies Alpha2Spec.
+								const float Energy = EnergyNormalization(Alpha2Spec, BSDFContext.Context.VoH, AreaLight);
+
 								// GGX lobe
 								{
 									D = Substrate_D_GGX(SafeRoughness, Alpha2Spec, BSDFContext.Context.NoH);
@@ -642,7 +644,6 @@ FSubstrateEvaluateResult SubstrateEvaluateBSDFCommon(FSubstrateBSDFContext BSDFC
 
 								// Area light energy compensation
 								{
-									const float Energy = EnergyNormalization(Alpha2Spec, BSDFContext.Context.VoH, AreaLight);
 									D *= Energy;
 								}
 							}
