Add `Temporal Responsiveness` node support to the translucency clipped depth pass.

#tests On feature test map

diff --git a/Engine/Shaders/Private/VelocityShader.usf b/Engine/Shaders/Private/VelocityShader.usf
--- a/Engine/Shaders/Private/VelocityShader.usf
+++ b/Engine/Shaders/Private/VelocityShader.usf
@@ -210,10 +210,23 @@ void MainPixelShader(
 		float CurrentTranslucencyDepth = 0;
 		ApplyPixelDepthOffsetToMaterialParameters(MaterialParameters, PixelMaterialInputs, CurrentTranslucencyDepth);
 		CurrentTranslucencyDepth = ConvertFromDeviceZ(CurrentTranslucencyDepth); 
-
-		uint WriteFlag = select(CurrentTranslucencyDepth < OpaqueDepth, 1u, 0);
-		const uint2 EncodedOutputBit = uint2(0u, WriteFlag << 16u); // Mask the lowest bit of alpha.
-		ImageInterlockedOrUInt64(VelocityClippedDepth.RWVelocity, Position, PackUlongType(EncodedOutputBit));
+	
+		// Use Temporal Responsiveness node to provide more flexible control.
+		#if VELOCITY_ENCODE_TEMPORAL_RESPONSIVENESS
+			float TranslucentTemporalResponsiveness = GetMaterialTemporalResponsiveness(MaterialParameters);
+			uint TemporalResponsivenessMask = select(TranslucentTemporalResponsiveness > 0.5f, 3U, uint(TranslucentTemporalResponsiveness>0)) & TEMPORAL_RESPONSIVENESS_MASK;
+		#else
+			uint TemporalResponsivenessMask = 0;
+		#endif
+		
+		uint WriteFlag = select(CurrentTranslucencyDepth < OpaqueDepth, TemporalResponsivenessMask, 0);
+ 
+		BRANCH
+		if (WriteFlag != 0)
+		{ 
+			const uint2 EncodedOutputBit = uint2(0u, WriteFlag << 16u); // Mask the lowest bits of alpha.
+			ImageInterlockedOrUInt64(VelocityClippedDepth.RWVelocity, Position, PackUlongType(EncodedOutputBit));
+		}
 		#endif
 	#endif
 	
