Path Tracer: Fix compilation error for legacy eye materials when Substrate is enabled

Fix minor typo in comment

#rnx

diff --git a/Engine/Shaders/Private/PathTracing/PathTracingMaterialHitShader.usf b/Engine/Shaders/Private/PathTracing/PathTracingMaterialHitShader.usf
--- a/Engine/Shaders/Private/PathTracing/PathTracingMaterialHitShader.usf
+++ b/Engine/Shaders/Private/PathTracing/PathTracingMaterialHitShader.usf
@@ -387,7 +387,7 @@ RAY_TRACING_ENTRY_CLOSEST_HIT(PathTracingMaterialCHS,
 
 #if !SUBSTRATE_OPTIMIZED_UNLIT 
 		// The slab sampler lets us stochastically choose between the different slab models the path tracer supports
-		// NOTE: Even if the material has a single substrate slab, we might still need this because refraction is handled seperately from SSS
+		// NOTE: Even if the material has a single substrate slab, we might still need this because refraction is handled separately from SSS
 		FRISContext SlabSampler = InitRISContext(PackedPayload.GetStochasticSlabRand());
 
 		Substrate_for_unroll(BSDFIdx = 0, BSDFIdx < BSDFCount, ++BSDFIdx)
@@ -1166,7 +1166,7 @@ RAY_TRACING_ENTRY_CLOSEST_HIT(PathTracingMaterialCHS,
 
 	#if IRIS_NORMAL // on
 		// Match BasePassPixelShader logic, Specular part of g-buffer is re-use, so hard code the term
-		#if SUBSTRATE_CONVERT_FROM_LEGACY
+		#if USE_LEGACY_CODEPATH_FOR_SUBSTRATE
 			Payload.SpecularColor = DielectricSpecularToF0(0.25);
 		#else
 			Payload.Specular = 0.25;
