[Shaders] Add default value for `RADIAL_LIGHT_TYPE` if it's not defined in MobileDeferredShading.usf to fix CI build issue.

#rnx

diff --git a/Engine/Shaders/Private/MobileDeferredShading.usf b/Engine/Shaders/Private/MobileDeferredShading.usf
--- a/Engine/Shaders/Private/MobileDeferredShading.usf
+++ b/Engine/Shaders/Private/MobileDeferredShading.usf
@@ -42,6 +42,10 @@
 #include "DistanceFieldAOShared.ush"
 #endif
 
+#ifndef RADIAL_LIGHT_TYPE
+#define RADIAL_LIGHT_TYPE 0
+#endif
+
 float4x4 TranslatedWorldToLight;
 float2 LightFunctionParameters2;
 
