Disable dithering and darkening for debug primitives drawn with SDPG_Foreground

#tests Tested in editor with DebugPDI
#rnx

diff --git a/Engine/Shaders/Private/PostProcessCompositePrimitives.usf b/Engine/Shaders/Private/PostProcessCompositePrimitives.usf
--- a/Engine/Shaders/Private/PostProcessCompositePrimitives.usf
+++ b/Engine/Shaders/Private/PostProcessCompositePrimitives.usf
@@ -107,6 +107,9 @@ float2 DepthTextureJitter;
 uint bCameraCut;
 uint bProcessAlpha;
 
+float OccludedDithering;
+float OccludedBrightness;
+
 float4 SampleOffsetArray[MSAA_SAMPLE_COUNT];
 
 uint bOpaqueEditorGizmo;
@@ -444,11 +447,11 @@ void MainCompositePostProcessPrimitivesPS(
 	float PatternMask = ((EditorPrimitivesPixelPos.x/2 + EditorPrimitivesPixelPos.y/2) % 2) * 0.7f;
 
 	// the contants express the two opacity values we see when the primitive is hidden
-	float LowContrastPatternMask = lerp(0.2, 1, PatternMask);
+	float LowContrastPatternMask = lerp(1.0 - OccludedDithering, 1.0f, PatternMask);
 
 	LowContrastPatternMask = saturate(lerp(LowContrastPatternMask, 1, bOpaqueEditorGizmo));
 
-	float HiddenMask = lerp(0.7f, 1.0f, DepthMask);
+	float HiddenMask = lerp(OccludedBrightness, 1.0f, DepthMask);
 	float DarkenMask = lerp(LowContrastPatternMask, 1.0f, DepthMask);
 
 	// Blend editor primitives with scene color.
diff --git a/Engine/Source/Runtime/Renderer/Private/PostProcess/PostProcessCompositeDebugPrimitives.cpp b/Engine/Source/Runtime/Renderer/Private/PostProcess/PostProcessCompositeDebugPrimitives.cpp
--- a/Engine/Source/Runtime/Renderer/Private/PostProcess/PostProcessCompositeDebugPrimitives.cpp
+++ b/Engine/Source/Runtime/Renderer/Private/PostProcess/PostProcessCompositeDebugPrimitives.cpp
@@ -226,6 +226,8 @@ static void ComposeDebugPrimitives(
 	PassParameters->bCompositeAnyNonNullDepth = bCompositeAnyNonNullDepth;
 	PassParameters->DepthTextureJitter = SceneDepthJitter;
 	PassParameters->bProcessAlpha = IsPostProcessingWithAlphaChannelSupported();
+	PassParameters->OccludedDithering = 0.0f;
+	PassParameters->OccludedBrightness = 1.0f;
 
 	for (int32 i = 0; i < int32(NumMSAASamples); i++)
 	{
diff --git a/Engine/Source/Runtime/Renderer/Private/PostProcess/PostProcessCompositeEditorPrimitives.cpp b/Engine/Source/Runtime/Renderer/Private/PostProcess/PostProcessCompositeEditorPrimitives.cpp
--- a/Engine/Source/Runtime/Renderer/Private/PostProcess/PostProcessCompositeEditorPrimitives.cpp
+++ b/Engine/Source/Runtime/Renderer/Private/PostProcess/PostProcessCompositeEditorPrimitives.cpp
@@ -457,6 +457,8 @@ FScreenPassTexture AddEditorPrimitivePass(
 		PassParameters->bCompositeAnyNonNullDepth = Context.bCompositeAnyNonNullDepth;
 		PassParameters->DepthTextureJitter = SceneDepthJitter;
 		PassParameters->bProcessAlpha = IsPostProcessingWithAlphaChannelSupported();
+		PassParameters->OccludedDithering = 0.8f;
+		PassParameters->OccludedBrightness = 0.7f;
 
 		for (int32 i = 0; i < int32(NumMSAASamples); i++)
 		{
diff --git a/Engine/Source/Runtime/Renderer/Private/PostProcess/PostProcessCompositePrimitivesCommon.h b/Engine/Source/Runtime/Renderer/Private/PostProcess/PostProcessCompositePrimitivesCommon.h
--- a/Engine/Source/Runtime/Renderer/Private/PostProcess/PostProcessCompositePrimitivesCommon.h
+++ b/Engine/Source/Runtime/Renderer/Private/PostProcess/PostProcessCompositePrimitivesCommon.h
@@ -136,6 +136,8 @@ public:
 		SHADER_PARAMETER(uint32, bCompositeAnyNonNullDepth)
 		SHADER_PARAMETER(FVector2f, DepthTextureJitter)
 		SHADER_PARAMETER(uint32, bProcessAlpha)
+		SHADER_PARAMETER(float, OccludedDithering)
+		SHADER_PARAMETER(float, OccludedBrightness)
 		RENDER_TARGET_BINDING_SLOTS()
 	END_SHADER_PARAMETER_STRUCT()
 
