Replace PF_R32G32B32F with PF_R32_FLOAT for ACES gamut table

diff --git a/Engine/Shaders/Private/ACES/ACESOutputTransform.ush b/Engine/Shaders/Private/ACES/ACESOutputTransform.ush
--- a/Engine/Shaders/Private/ACES/ACESOutputTransform.ush
+++ b/Engine/Shaders/Private/ACES/ACESOutputTransform.ush
@@ -93,7 +93,7 @@ struct ODTParams
     float lowerHullGamma;
 
     Texture2D<float> reach_table;
-    Texture2D<float3> gamut_table;
+    Texture2D<float> gamut_table;
     Texture2D<float> gamma_table;
 };
 
@@ -355,7 +355,7 @@ ODTParams init_ODTParams(
     float3x3 OUTPUT_RGB_TO_XYZ,
     float3x3 OUTPUT_XYZ_TO_RGB,
     Texture2D<float> reach_table,
-    Texture2D<float3> gamut_table,
+    Texture2D<float> gamut_table,
     Texture2D<float> gamma_table
 )
 {
@@ -631,7 +631,7 @@ int midpoint(int low, int high)
     return (low + high) / 2;
 }
 
-float2 cuspFromTable(float h, Texture2D<float3> gamut_table)
+float2 cuspFromTable(float h, Texture2D<float> gamut_table)
 {
     float3 lo;
     float3 hi;
@@ -642,7 +642,7 @@ float2 cuspFromTable(float h, Texture2D<float3> gamut_table)
 
     while (low_i + 1 < high_i)
     {
-        if (h > gamut_table.Load(int3(i, 0, 0))[2])
+        if (h > gamut_table.Load(int3(2, i, 0)))
         {
             low_i = i;
         }
@@ -652,8 +652,12 @@ float2 cuspFromTable(float h, Texture2D<float3> gamut_table)
         }
         i = midpoint(low_i, high_i);
     }
-    lo = gamut_table.Load(int3(high_i - 1, 0, 0));
-    hi = gamut_table.Load(int3(high_i, 0, 0));
+    lo = float3(gamut_table.Load(int3(0, high_i - 1, 0)),
+        gamut_table.Load(int3(1, high_i - 1, 0)),
+        gamut_table.Load(int3(2, high_i - 1, 0)));
+    hi = float3(gamut_table.Load(int3(0, high_i, 0)),
+        gamut_table.Load(int3(1, high_i, 0)),
+        gamut_table.Load(int3(2, high_i, 0)));
 
     float t = (h - lo[2]) / (hi[2] - lo[2]);
     float cuspJ = lerp(lo[0], hi[0], t);
diff --git a/Engine/Shaders/Private/PostProcessCombineLUTs.usf b/Engine/Shaders/Private/PostProcessCombineLUTs.usf
--- a/Engine/Shaders/Private/PostProcessCombineLUTs.usf
+++ b/Engine/Shaders/Private/PostProcessCombineLUTs.usf
@@ -78,7 +78,7 @@ float ACESGamutCompression;
 
 // ACES 2.0
 Texture2D<float> ACESReachTable;
-Texture2D<float3> ACESGamutTable;
+Texture2D<float> ACESGamutTable;
 Texture2D<float> ACESGammaTable;
 
 static const float DefaultTemperature = 6500.0;
diff --git a/Engine/Shaders/Private/PostProcessDeviceEncodingOnly.usf b/Engine/Shaders/Private/PostProcessDeviceEncodingOnly.usf
--- a/Engine/Shaders/Private/PostProcessDeviceEncodingOnly.usf
+++ b/Engine/Shaders/Private/PostProcessDeviceEncodingOnly.usf
@@ -38,7 +38,7 @@ float ACESGamutCompression;
 
 // ACES 2.0
 Texture2D<float> ACESReachTable;
-Texture2D<float3> ACESGamutTable;
+Texture2D<float> ACESGamutTable;
 Texture2D<float> ACESGammaTable;
 
 uint GetOutputDevice()
diff --git a/Engine/Shaders/Private/TonemapCommon.ush b/Engine/Shaders/Private/TonemapCommon.ush
--- a/Engine/Shaders/Private/TonemapCommon.ush
+++ b/Engine/Shaders/Private/TonemapCommon.ush
@@ -327,8 +327,8 @@ float3 ACESOutputTransform(
 	const float3x3 WCS_2_AP0,
 	float PeakLuminance,
 	Texture2D<float> ReachTable,
-    Texture2D<float3> GamutTable,
-    Texture2D<float> GammaTable,
+	Texture2D<float> GamutTable,
+	Texture2D<float> GammaTable,
 	const bool bApplyWhiteLimiting = true
 )
 {
@@ -365,8 +365,8 @@ float3 ACESInvOutputTransform(
 	const float3x3 AP0_2_WCS,
 	float PeakLuminance,
 	Texture2D<float> ReachTable,
-    Texture2D<float3> GamutTable,
-    Texture2D<float> GammaTable,
+	Texture2D<float> GamutTable,
+	Texture2D<float> GammaTable,
 	const bool bApplyWhiteLimiting = true
 )
 {
diff --git a/Engine/Source/Runtime/Renderer/Private/PostProcess/ACESUtils.cpp b/Engine/Source/Runtime/Renderer/Private/PostProcess/ACESUtils.cpp
--- a/Engine/Source/Runtime/Renderer/Private/PostProcess/ACESUtils.cpp
+++ b/Engine/Source/Runtime/Renderer/Private/PostProcess/ACESUtils.cpp
@@ -771,24 +771,29 @@ public:
 
 	virtual ~FGamutCuspTable() = default;
 
-	void Update(FRHICommandListImmediate& RHICmdList, const FColorSpace& InLimitingColorSpace, float InPeakLuminance) override
+	virtual void Update(FRHICommandListImmediate& RHICmdList, const FColorSpace& InLimitingColorSpace, float InPeakLuminance) override
 	{
 		Table3D GamutCuspTable = MakeGamutTable(InLimitingColorSpace, InPeakLuminance);
-		check(GetSizeX() == GamutCuspTable.TotalSize);
-		const uint32 DataSize = GetSizeX() * sizeof(FVector3f);
+		check(GetSizeY() == GamutCuspTable.TotalSize);
 		const FUpdateTextureRegion2D Region(0, 0, 0, 0, GetSizeX(), GetSizeY());
 
-		RHICmdList.UpdateTexture2D(TextureRHI, 0, Region, DataSize, (const uint8*)GamutCuspTable.Data);
+		RHICmdList.UpdateTexture2D(TextureRHI, 0, Region, sizeof(FVector3f), (const uint8*)GamutCuspTable.Data);
+	}
+
+	virtual uint32 GetSizeX() const override
+	{
+		// Store JMh channels to separate texels.
+		return 3;
 	}
 
-	uint32 GetSizeX() const override
+	virtual uint32 GetSizeY() const override
 	{
 		return TABLE_TotalSize;
 	}
 
-	EPixelFormat GetPixelFormat() const override
+	virtual EPixelFormat GetPixelFormat() const override
 	{
-		return PF_R32G32B32F;
+		return PF_R32_FLOAT;
 	}
 };
 
