Path Tracer: Fix broken visualization modes for smooth and geo normals

#rnx

diff --git a/Engine/Shaders/Private/PathTracing/PathTracingDebug.usf b/Engine/Shaders/Private/PathTracing/PathTracingDebug.usf
--- a/Engine/Shaders/Private/PathTracing/PathTracingDebug.usf
+++ b/Engine/Shaders/Private/PathTracing/PathTracingDebug.usf
@@ -998,8 +998,8 @@ FPathTracingResult PathTracingDebugCore(FRayDesc Ray, uint2 LaunchIndex)
 			{
 				case PATH_TRACER_DEBUG_VIZ_RADIANCE: 					Color = HitPayload.Radiance * View.PreExposure; break;
 				case PATH_TRACER_DEBUG_VIZ_WORLD_NORMAL: 				Color = WorldNormal * 0.5 + 0.5; break;
-				case PATH_TRACER_DEBUG_VIZ_WORLD_SMOOTH_NORMAL:			Color = HitPayload.WorldSmoothNormal * 0.5 + 0.5;
-				case PATH_TRACER_DEBUG_VIZ_WORLD_GEO_NORMAL:			Color = HitPayload.WorldGeoNormal * 0.5 + 0.5;
+				case PATH_TRACER_DEBUG_VIZ_WORLD_SMOOTH_NORMAL:			Color = HitPayload.WorldSmoothNormal * 0.5 + 0.5; break;
+				case PATH_TRACER_DEBUG_VIZ_WORLD_GEO_NORMAL:			Color = HitPayload.WorldGeoNormal * 0.5 + 0.5; break;
 #if SUBSTRATE_ENABLED
 				// TODO: implement more visualizations for this case
 				case PATH_TRACER_DEBUG_VIZ_BASE_COLOR:					Color = lerp(HitPayload.DiffuseColor, HitPayload.SpecularColor, F0RGBToMetallic(HitPayload.SpecularColor)); break;
