Path Tracer: Improve adaptive sampling heuristic

Previous metric was prone to missing edges, particularly around bright light sources.

diff --git a/Engine/Shaders/Private/PathTracing/PathTracingAdaptiveUtils.ush b/Engine/Shaders/Private/PathTracing/PathTracingAdaptiveUtils.ush
--- a/Engine/Shaders/Private/PathTracing/PathTracingAdaptiveUtils.ush
+++ b/Engine/Shaders/Private/PathTracing/PathTracingAdaptiveUtils.ush
@@ -24,10 +24,10 @@ float PerceptualError(float Mean, float StdErr, float Exposure)
 	float d = 0.59f;
 	float e = 0.14f;
 	return StdErr * Exposure * (a * x * (2 * e + d * x) + b * (e - c * x * x)) / Square(x * (c * x + d) + e);
-#elif 0
+#elif 1
 	// Cheap fit to curve above
 	// tonemap(x) = (6*x/(6*x+1))^2
-	return StdErr * Exposure * 72.0 * x / Pow3(6.0 * x + 1.0);
+	return StdErr * Exposure * 72.0 * x * rcp(Pow3(6.0 * x + 1.0));
 #elif 1
 	// Similar curve, but with faster highlight rolloff
 	// tonemap(x) = (1 - exp(-4*x))^2
@@ -78,18 +78,26 @@ float4 SampleVarianceTexture(float2 VarianceUV, int Mip)
 
 float GetAdaptiveSampleFactor(uint2 TexCoord, float Exposure)
 {
-	float2 VarianceUV = float2(TexCoord + 0.5) / float2(VarianceTextureDims.xy);
-	float SampleScaleFactor = 0.0;
-	for (int Mip = 1, NumMips = VarianceTextureDims.z; Mip < NumMips; Mip++)
+	float SampleScaleFactor = 0;
+	for (int Mip = 0, NumMips = VarianceTextureDims.z; Mip < NumMips; Mip++)
 	{
-		float4 Variance = SampleVarianceTexture(VarianceUV, Mip);
-		float NumSamples = Variance.z;
-		float OutVariance = Variance.y - Variance.x * Variance.x;
-		float Lum = Variance.x;
-		float StdErr = sqrt(OutVariance / NumSamples);
-		float RelErr = PerceptualError(Lum, StdErr, Exposure);
-		// project the multiple of how many more samples we will need
-		SampleScaleFactor = max(SampleScaleFactor, (1 + Square(RelErr / AdaptiveSamplingErrorThreshold)) / NumSamples);
+		const int R = 2;
+		for (int dy = -R; dy <= R; dy++)
+		for (int dx = -R; dx <= R; dx++)
+		{
+			int2 Coord = (TexCoord >> Mip) + int2(dx, dy);
+			if (all(uint2(Coord) < uint2(VarianceTextureDims.xy >> Mip)))
+			{
+				float4 Variance = VarianceTexture.Load(int3(Coord, Mip));
+				float NumSamples = Variance.z;
+				float OutVariance = abs(Variance.y - Variance.x * Variance.x);
+				float Lum = Variance.x;
+				float StdErr = sqrt(OutVariance / NumSamples);
+				float RelErr = PerceptualError(Lum, StdErr, Exposure);
+				// project the multiple of how many more samples we will need
+				SampleScaleFactor = max(SampleScaleFactor, (1 + Square(RelErr / AdaptiveSamplingErrorThreshold)) / NumSamples);
+			}
+		}
 	}
 	return SampleScaleFactor;
-}
\ No newline at end of file
+}
diff --git a/Engine/Shaders/Private/PathTracing/PathTracingBuildAdaptiveError.usf b/Engine/Shaders/Private/PathTracing/PathTracingBuildAdaptiveError.usf
--- a/Engine/Shaders/Private/PathTracing/PathTracingBuildAdaptiveError.usf
+++ b/Engine/Shaders/Private/PathTracing/PathTracingBuildAdaptiveError.usf
@@ -20,23 +20,27 @@ void PathTracingBuildAdaptiveErrorTextureCS(uint2 DispatchThreadId : SV_Dispatch
 	float2 UV = (OutPixelCoord + 0.5) / float2(OutputResolution);
 	float2 InvInputResolution = rcp(float2(InputResolution));
 
-#if 0
-	// single tap mipmap
-	float4 Result = InputMip.SampleLevel(InputMipSampler, UV, 0);
-#elif 1
-	float4 Result = 0.0;
-	float FilterSum = 0.0;
-	const int R = 2;
-	for (int dy = -R; dy <= R; dy++)
-	for (int dx = -R; dx <= R; dx++)
+#if 1
+	float4 Result = 0;
+	const int R = 1;
+	for (int dy = 0; dy <= R; dy++)
+	for (int dx = 0; dx <= R; dx++)
 	{
-		float2 Offset = float2(dx, dy);
-		float Filter = exp(-0.5 * length2(Offset));
-		float4 Pixel = InputMip.SampleLevel(InputMipSampler, UV + Offset * InvInputResolution, 0);
-		FilterSum += Filter;
-		Result = lerp(Result, Pixel, Filter / FilterSum);
+		int2 InPixelCoord = 2 * OutPixelCoord + int2(dx, dy);
+		if (all(uint2(InPixelCoord) < InputResolution))
+		{
+			// inside range
+			float4 Pixel = InputMip.Load(int3(InPixelCoord, 0));
+			float Blend = Pixel.z / (Result.z + Pixel.z);
+			Result.xy = lerp(Result.xy, Pixel.xy, Blend);
+			Result.z += Pixel.z;
+		}
 	}
+#else
+	// single tap mipmap
+	// NOTE: this is wrong as it doesn't account for the fact that each pixel can have a variable number of samples
+	// This is similar to the idea of down-sampling an image with unpremultiplied alpha
+	float4 Result = InputMip.SampleLevel(InputMipSampler, UV, 0);
 #endif
-	//Result.z *= 4.0; // we are downsampling by 2x in each direction, so the overal normalization in the number of samples should be 4x
 	OutputMip[OutPixelCoord] = Result;
 }
\ No newline at end of file
diff --git a/Engine/Shaders/Private/PathTracing/PathTracingCompositingPixelShader.usf b/Engine/Shaders/Private/PathTracing/PathTracingCompositingPixelShader.usf
--- a/Engine/Shaders/Private/PathTracing/PathTracingCompositingPixelShader.usf
+++ b/Engine/Shaders/Private/PathTracing/PathTracingCompositingPixelShader.usf
@@ -38,7 +38,7 @@ void CompositeMain(
 	{
 		// Visualize which pixels are still being worked and color them according to how many more samples we think they need
 		float2 VarianceUV = float2(TexCoord.xy + 0.5) * View.BufferSizeAndInvSize.zw;
-		float SampleScaleFactor = GetAdaptiveSampleFactor(TexCoord.xy, View.PreExposure);
+		float SampleScaleFactor = GetAdaptiveSampleFactor(TexCoord.xy, View.PreExposure / PreExposure); // Run with same exposure as adaptive code
 		if (SampleScaleFactor > 1.0)
 		{
 			float4 Variance = VarianceTexture.Load(TexCoord);
@@ -64,7 +64,7 @@ void CompositeMain(
 		float OutVariance = Variance.y - Variance.x * Variance.x;
 		float Lum = Variance.x;
 		float StdErr = sqrt(OutVariance / NumSamples);
-		float RelErr = PerceptualError(Lum, StdErr, View.PreExposure);
+		float RelErr = PerceptualError(Lum, StdErr, View.PreExposure / PreExposure); // Run with same exposure as adaptive code
 		OutColor.xyz = RelErr;
 	}
 	
