Fix SSPreintegrated conversion with Substrate.

* Fix SSS tracking for SS preintegrated when Substrate is enabled
* Fix opacity incorrectly setup
* Added an optional better conversion, but it requires SSS diffusion per-pixel which is not support for Format=0 yet.

diff --git a/Engine/Shaders/Private/Substrate/SubstrateLegacyConversion.ush b/Engine/Shaders/Private/Substrate/SubstrateLegacyConversion.ush
--- a/Engine/Shaders/Private/Substrate/SubstrateLegacyConversion.ush
+++ b/Engine/Shaders/Private/Substrate/SubstrateLegacyConversion.ush
@@ -39,6 +39,7 @@ struct FSubstrateLegacyParameters
 	float  Roughness;
 	float  Anisotropy;
 	uint   SSSProfileID;
+	bool bSupportDefaultSSSProfile;
 	float3 SSSMFP;
 	float  SSSMFPScale;
 	float  SSSPhaseAniso;
@@ -76,6 +77,7 @@ FSubstrateLegacyParameters InitSubstrateLegacyParameters(FSubstratePixelFootprin
 	Out.Roughness = 0.5f;
 	Out.Anisotropy = 0.f;
 	Out.SSSProfileID = 0;
+	Out.bSupportDefaultSSSProfile = true;
 	Out.SSSMFP = Zeros;
 	Out.SSSMFPScale = 1.f;
 	Out.SSSPhaseAniso = 0.f;
@@ -121,7 +123,7 @@ FSubstrateData CreateLegacySlab(
 		In.Roughness,										// Roughness
 		In.Anisotropy,										// Anisotropy
 		In.SSSProfileID,									// SSSProfileID
-		true,												// bSupportDefaultSSSProfile
+		In.bSupportDefaultSSSProfile,						// bSupportDefaultSSSProfile
 		In.SSSMFP,											// SSSMFP
 		In.SSSMFPScale,										// SSSMFPScale
 		In.SSSPhaseAniso,									// SSSPhaseAniso
@@ -313,8 +315,18 @@ FSubstrateData SubstrateConvertLegacyMaterialDynamic(
 	#if MATERIAL_SHADINGMODEL_PREINTEGRATED_SKIN
 	if (ShadingModel == SHADINGMODELID_PREINTEGRATED_SKIN)
 	{
+		#if 0
+		// This is a better conversion, but requires support for PerPixelSSS which is currently not support with Format=0 (we don't have GBuffer encoding for this yet)
+		// Clamp MFP to MaxSSSDiffusionMFP, as the current SSS screen diffusion algorithm can be too noisy (current value is arbitrary)
+		const float ThicknessInCm = SUBSTRATE_LAYER_DEFAULT_THICKNESS_CM;
+		const float MaxSSSDiffusionMFP = 1.f;
+		const float3 MFPInCm = min(TransmittanceToMeanFreePath(SubSurfaceColor, ThicknessInCm * CENTIMETER_TO_METER) * METER_TO_CENTIMETER, MaxSSSDiffusionMFP);
+		const bool bSupportDefaultSSSProfile = false;
+		#else
 		// Use default profile MFP (e.g., human skin) as MFP value for converting hardcoded pre-integrated SSS texture for deferred material.
 		const float3 MFPInCm = float3(1.0f, 0.088964f, 0.072095f) * 2.6748f * 0.1f;
+		const bool bSupportDefaultSSSProfile = true;
+		#endif
 
 		// Legacy material uses Subsurface color as transmission 'tinting', but we can represent that with a single layer. So instead we take 
 		// the max color of SubSurfaceColor & BaseColor
@@ -324,8 +336,9 @@ FSubstrateData SubstrateConvertLegacyMaterialDynamic(
 		LegacySlab.Roughness = Roughness;
 		LegacySlab.Anisotropy = Anisotropy;
 		LegacySlab.SSSMFP = max(1e-05f, MFPInCm); // Ensure the MFP is not null to force the material as have SSS
+		LegacySlab.bSupportDefaultSSSProfile = bSupportDefaultSSSProfile;
 		LegacySlab.SSSProfileID = 0;
-		LegacySlab.SSSMFPScale = OneMinusOpacity;
+		LegacySlab.SSSMFPScale = 1.f;
 		LegacySlab.SSSPhaseAniso = 0.93f;
 		LegacySlab.SSSType = SSS_TYPE_DIFFUSION;
 		LegacySlab.Emissive = Emissive;
diff --git a/Engine/Source/Runtime/Engine/Public/MaterialShared.h b/Engine/Source/Runtime/Engine/Public/MaterialShared.h
--- a/Engine/Source/Runtime/Engine/Public/MaterialShared.h
+++ b/Engine/Source/Runtime/Engine/Public/MaterialShared.h
@@ -39,6 +39,7 @@
 #include "Rendering/SubstrateMaterialShared.h"
 #include "Shader/Preshader.h"
 #include "Shader/ShaderTypes.h"
+#include "RenderUtils.h"
 
 #include <atomic>
 
@@ -135,7 +136,7 @@ inline bool UseSubsurfaceProfile(FMaterialShadingModelField ShadingModel)
 
 inline uint32 GetUseSubsurfaceProfileShadingModelMask()
 {
-	return (1 << MSM_SubsurfaceProfile) | (1 << MSM_Eye);
+	return (1u << MSM_SubsurfaceProfile) | (1u << MSM_Eye) | (Substrate::IsSubstrateEnabled() ? (1u << MSM_PreintegratedSkin) : 0u);
 }
 
 UE_DEPRECATED(5.7, "Please use AllowDitheredLODTransition with a EShaderPlatform argument")
diff --git a/Engine/Source/Runtime/Renderer/Private/Substrate/SubstrateVisualize.cpp b/Engine/Source/Runtime/Renderer/Private/Substrate/SubstrateVisualize.cpp
--- a/Engine/Source/Runtime/Renderer/Private/Substrate/SubstrateVisualize.cpp
+++ b/Engine/Source/Runtime/Renderer/Private/Substrate/SubstrateVisualize.cpp
@@ -9,6 +9,7 @@
 #include "PixelShaderUtils.h"
 #include "BasePassRendering.h"
 #include "IndirectLightRendering.h"
+#include "PostProcess/PostProcessSubsurface.h"
 #include "SubstrateVisualizationData.h"
 #include "SubstrateVisualizeDefinitions.h"
 #include "CanvasItem.h"
@@ -1340,6 +1341,10 @@ static void AddVisualizeSystemInfoPasses(FRDGBuilder& GraphBuilder, const FViewI
 				PrintS(DrawPosX, DrawPosY, TEXT("  Anisotropy"), White);
 				PrintBool(DrawPosX + 150, DrawPosY, Substrate::GetSubstrateUsesAnisotropy(View));
 				NewLine();
+
+				PrintS(DrawPosX, DrawPosY, TEXT("  Sub-Surface"), White);
+				PrintBool(DrawPosX + 150, DrawPosY, IsSubsurfaceRequiredForView(View));
+				NewLine();
 			}
 			NewLine();
 		});
