Path Tracer: Fix double counting of exposure in adaptive sampler

The path tracer normalizes the exposure during rendering, so it is not necessary to include the exposure during the adaptive error measurement step.

Additionally, improve the down-sampling filter used when preparing adaptive sampling mipmaps.

diff --git a/Engine/Shaders/Private/PathTracing/PathTracingAdaptiveStart.usf b/Engine/Shaders/Private/PathTracing/PathTracingAdaptiveStart.usf
--- a/Engine/Shaders/Private/PathTracing/PathTracingAdaptiveStart.usf
+++ b/Engine/Shaders/Private/PathTracing/PathTracingAdaptiveStart.usf
@@ -8,7 +8,6 @@ RWBuffer<uint> NextActivePaths;
 RWBuffer<uint> NumPathStates;
 int2 TileTextureOffset;
 uint2 DispatchDim;
-float ViewPreExposure;
 
 
 [numthreads(THREADGROUPSIZE_X, THREADGROUPSIZE_Y, 1)]
@@ -28,7 +27,7 @@ void PathTracingAdaptiveStartCS(uint2 DispatchThreadId : SV_DispatchThreadID)
 		NumPathStates[2] = 1;
 	}
 
-	float SampleScaleFactor = GetAdaptiveSampleFactor(DispatchIdx + TileTextureOffset, ViewPreExposure);
+	float SampleScaleFactor = GetAdaptiveSampleFactor(DispatchIdx + TileTextureOffset);
 	if (SampleScaleFactor > 1.0)
 	{
 		// This pixel needs more samples, allocate a path state for it
diff --git a/Engine/Shaders/Private/PathTracing/PathTracingAdaptiveUtils.ush b/Engine/Shaders/Private/PathTracing/PathTracingAdaptiveUtils.ush
--- a/Engine/Shaders/Private/PathTracing/PathTracingAdaptiveUtils.ush
+++ b/Engine/Shaders/Private/PathTracing/PathTracingAdaptiveUtils.ush
@@ -10,11 +10,11 @@ float 				AdaptiveSamplingErrorThreshold;
 
 // Variance of the tonemapped color can be predicted from the variance of the scene linear color:
 // https://en.wikipedia.org/wiki/Taylor_expansions_for_the_moments_of_functions_of_random_variables
-float PerceptualError(float Mean, float StdErr, float Exposure)
+float PerceptualError(float Mean, float StdErr)
 {
 	// The UE tone mapping curve is quite a bit more involved, so getting its derivative would be messy and likely overkill
 	// instead just use a roughly S-curved shape curve that roughly does something filmic without costing too much
-	float x = Mean * Exposure;
+	float x = Mean;
 #if 0
 	// https://knarkowicz.wordpress.com/2016/01/06/aces-filmic-tone-mapping-curve/
 	// saturate((x*(a*x+b))/(x*(c*x+d)+e))
@@ -23,22 +23,11 @@ float PerceptualError(float Mean, float StdErr, float Exposure)
 	float c = 2.43f;
 	float d = 0.59f;
 	float e = 0.14f;
-	return StdErr * Exposure * (a * x * (2 * e + d * x) + b * (e - c * x * x)) / Square(x * (c * x + d) + e);
-#elif 1
+	return StdErr * (a * x * (2 * e + d * x) + b * (e - c * x * x)) / Square(x * (c * x + d) + e);
+#else
 	// Cheap fit to curve above
 	// tonemap(x) = (6*x/(6*x+1))^2
-	return StdErr * Exposure * 72.0 * x * rcp(Pow3(6.0 * x + 1.0));
-#elif 1
-	// Similar curve, but with faster highlight rolloff
-	// tonemap(x) = (1 - exp(-4*x))^2
-	float E = exp(-4.0 * x);
-	return StdErr * Exposure * 8.0 * (E - E * E);
-#elif 0
-	return StdErr * Exposure * Square(rcp(x + 1));
-#elif 0
-	return StdErr * Exposure * 0.5 * rsqrt(x + 0.01);
-#elif 0
-	return StdErr * Exposure * 4.0 * exp(-4.0 * x);
+	return StdErr * 72.0 * x * rcp(Pow3(6.0 * x + 1.0));
 #endif
 }
 
@@ -76,10 +65,10 @@ float4 SampleVarianceTexture(float2 VarianceUV, int Mip)
 #endif
 }
 
-float GetAdaptiveSampleFactor(uint2 TexCoord, float Exposure)
+float GetAdaptiveSampleFactor(uint2 TexCoord)
 {
-	float SampleScaleFactor = 0;
-	for (int Mip = 0, NumMips = VarianceTextureDims.z; Mip < NumMips; Mip++)
+	float SampleScaleFactor = 0, SampleScale = 1;
+	for (int Mip = 0, NumMips = VarianceTextureDims.z; Mip < NumMips; Mip++, SampleScale *= 4)
 	{
 		const int R = 2;
 		for (int dy = -R; dy <= R; dy++)
@@ -93,9 +82,9 @@ float GetAdaptiveSampleFactor(uint2 TexCoord, float Exposure)
 				float OutVariance = abs(Variance.y - Variance.x * Variance.x);
 				float Lum = Variance.x;
 				float StdErr = sqrt(OutVariance / NumSamples);
-				float RelErr = PerceptualError(Lum, StdErr, Exposure);
+				float RelErr = PerceptualError(Lum, StdErr);
 				// project the multiple of how many more samples we will need
-				SampleScaleFactor = max(SampleScaleFactor, (1 + Square(RelErr / AdaptiveSamplingErrorThreshold)) / NumSamples);
+				SampleScaleFactor = max(SampleScaleFactor, (1 + Square(RelErr / AdaptiveSamplingErrorThreshold)) / NumSamples * SampleScale);
 			}
 		}
 	}
diff --git a/Engine/Shaders/Private/PathTracing/PathTracingBuildAdaptiveError.usf b/Engine/Shaders/Private/PathTracing/PathTracingBuildAdaptiveError.usf
--- a/Engine/Shaders/Private/PathTracing/PathTracingBuildAdaptiveError.usf
+++ b/Engine/Shaders/Private/PathTracing/PathTracingBuildAdaptiveError.usf
@@ -8,6 +8,14 @@ SamplerState InputMipSampler;
 Texture2D<float4>   InputMip;
 RWTexture2D<float4> OutputMip;
 
+float4 CombinePixels(float4 A, float4 B, float W)
+{
+	float NumSamples = lerp(A.z, B.z, W);
+	float RatioA = A.z / NumSamples;
+	float RatioB = B.z / NumSamples;
+	return float4(lerp(A.xy * RatioA, B.xy * RatioB, W), NumSamples, 0);
+}
+
 [numthreads(THREADGROUPSIZE_X, THREADGROUPSIZE_Y, 1)]
 void PathTracingBuildAdaptiveErrorTextureCS(uint2 DispatchThreadId : SV_DispatchThreadID)
 {
@@ -17,30 +25,115 @@ void PathTracingBuildAdaptiveErrorTextureCS(uint2 DispatchThreadId : SV_Dispatch
 		return;
 	}
 
-	float2 UV = (OutPixelCoord + 0.5) / float2(OutputResolution);
-	float2 InvInputResolution = rcp(float2(InputResolution));
-
-#if 1
 	float4 Result = 0;
-	const int R = 1;
-	for (int dy = 0; dy <= R; dy++)
-	for (int dx = 0; dx <= R; dx++)
+	const bool bOddX = (InputResolution.x & 1) == 1;
+	const bool bOddY = (InputResolution.y & 1) == 1;
+
+	if (bOddX)
 	{
-		int2 InPixelCoord = 2 * OutPixelCoord + int2(dx, dy);
-		if (all(uint2(InPixelCoord) < InputResolution))
+		if (bOddY)
 		{
-			// inside range
-			float4 Pixel = InputMip.Load(int3(InPixelCoord, 0));
-			float Blend = Pixel.z / (Result.z + Pixel.z);
-			Result.xy = lerp(Result.xy, Pixel.xy, Blend);
-			Result.z += Pixel.z;
+			// odd-odd, need 3x3 weighted downsample
+			float W0x = float(OutputResolution.x - OutPixelCoord.x) / float(InputResolution.x);
+			float W1x = float(OutputResolution.x                  ) / float(InputResolution.x);
+			float W2x = float(                 1 + OutPixelCoord.x) / float(InputResolution.x);
+
+			float W0y = float(OutputResolution.y - OutPixelCoord.y) / float(InputResolution.y);
+			float W1y = float(OutputResolution.y                  ) / float(InputResolution.y);
+			float W2y = float(                 1 + OutPixelCoord.y) / float(InputResolution.y);
+
+			float4 Pixelx0 = 0;
+			float4 Pixelx1 = 0;
+			float4 Pixelx2 = 0;
+
+			float4 Pixel00 = InputMip.Load(int3(2 * OutPixelCoord + int2(0, 0), 0));
+			float4 Pixel10 = InputMip.Load(int3(2 * OutPixelCoord + int2(1, 0), 0));
+			float4 Pixel20 = InputMip.Load(int3(2 * OutPixelCoord + int2(2, 0), 0));
+			Pixelx0 = CombinePixels(Pixel00, Pixel10, W0x / (W0x + W1x));
+			Pixelx0 = CombinePixels(Pixelx0, Pixel20, W2x);
+
+			float4 Pixel01 = InputMip.Load(int3(2 * OutPixelCoord + int2(0, 1), 0));
+			float4 Pixel11 = InputMip.Load(int3(2 * OutPixelCoord + int2(1, 1), 0));
+			float4 Pixel21 = InputMip.Load(int3(2 * OutPixelCoord + int2(2, 1), 0));
+			Pixelx1 = CombinePixels(Pixel01, Pixel11, W0x / (W0x + W1x));
+			Pixelx1 = CombinePixels(Pixelx1, Pixel21, W2x);
+
+			float4 Pixel02 = InputMip.Load(int3(2 * OutPixelCoord + int2(0, 2), 0));
+			float4 Pixel12 = InputMip.Load(int3(2 * OutPixelCoord + int2(1, 2), 0));
+			float4 Pixel22 = InputMip.Load(int3(2 * OutPixelCoord + int2(2, 2), 0));
+			Pixelx2 = CombinePixels(Pixel02, Pixel12, W0x / (W0x + W1x));
+			Pixelx2 = CombinePixels(Pixelx2, Pixel22, W2x);
+
+			Result = CombinePixels(Pixelx0, Pixelx1, W0y / (W0y + W1y));
+			Result = CombinePixels(Result, Pixelx2, W2y);
+		}
+		else
+		{
+			// odd-even 3x2 weighted downsample
+			float W0x = float(OutputResolution.x - OutPixelCoord.x) / float(InputResolution.x);
+			float W1x = float(OutputResolution.x                  ) / float(InputResolution.x);
+			float W2x = float(                 1 + OutPixelCoord.x) / float(InputResolution.x);
+
+			float4 Pixelx0 = 0;
+			float4 Pixelx1 = 0;
+
+			float4 Pixel00 = InputMip.Load(int3(2 * OutPixelCoord + int2(0, 0), 0));
+			float4 Pixel10 = InputMip.Load(int3(2 * OutPixelCoord + int2(1, 0), 0));
+			float4 Pixel20 = InputMip.Load(int3(2 * OutPixelCoord + int2(2, 0), 0));
+			Pixelx0 = CombinePixels(Pixel00, Pixel10, W0x / (W0x + W1x));
+			Pixelx0 = CombinePixels(Pixelx0, Pixel20, W2x);
+
+			float4 Pixel01 = InputMip.Load(int3(2 * OutPixelCoord + int2(0, 1), 0));
+			float4 Pixel11 = InputMip.Load(int3(2 * OutPixelCoord + int2(1, 1), 0));
+			float4 Pixel21 = InputMip.Load(int3(2 * OutPixelCoord + int2(2, 1), 0));
+			Pixelx1 = CombinePixels(Pixel01, Pixel11, W0x / (W0x + W1x));
+			Pixelx1 = CombinePixels(Pixelx1, Pixel21, W2x);
+
+			Result = CombinePixels(Pixelx0, Pixelx1, 0.5);
+		}
+	}
+	else
+	{
+		if (bOddY)
+		{
+			// even-odd, need 2x3 weighted downsample
+			float W0y = float(OutputResolution.y - OutPixelCoord.y) / float(InputResolution.y);
+			float W1y = float(OutputResolution.y                  ) / float(InputResolution.y);
+			float W2y = float(                 1 + OutPixelCoord.y) / float(InputResolution.y);
+
+			float4 Pixelx0 = 0;
+			float4 Pixelx1 = 0;
+			float4 Pixelx2 = 0;
+
+			float4 Pixel00 = InputMip.Load(int3(2 * OutPixelCoord + int2(0, 0), 0));
+			float4 Pixel10 = InputMip.Load(int3(2 * OutPixelCoord + int2(1, 0), 0));
+			Pixelx0 = CombinePixels(Pixel00, Pixel10, 0.5);
+
+			float4 Pixel01 = InputMip.Load(int3(2 * OutPixelCoord + int2(0, 1), 0));
+			float4 Pixel11 = InputMip.Load(int3(2 * OutPixelCoord + int2(1, 1), 0));
+			Pixelx1 = CombinePixels(Pixel01, Pixel11, 0.5);
+
+			float4 Pixel02 = InputMip.Load(int3(2 * OutPixelCoord + int2(0, 2), 0));
+			float4 Pixel12 = InputMip.Load(int3(2 * OutPixelCoord + int2(1, 2), 0));
+			Pixelx2 = CombinePixels(Pixel02, Pixel12, 0.5);
+
+			Result = CombinePixels(Pixelx0, Pixelx1, W0y / (W0y + W1y));
+			Result = CombinePixels(Result, Pixelx2, W2y);
+		}
+		else
+		{
+			// even-even, need 2x2 weighted downsample
+			float4 Pixel00 = InputMip.Load(int3(2 * OutPixelCoord + int2(0, 0), 0));
+			float4 Pixel10 = InputMip.Load(int3(2 * OutPixelCoord + int2(1, 0), 0));
+			float4 Pixelx0 = CombinePixels(Pixel00, Pixel10, 0.5);
+
+			float4 Pixel01 = InputMip.Load(int3(2 * OutPixelCoord + int2(0, 1), 0));
+			float4 Pixel11 = InputMip.Load(int3(2 * OutPixelCoord + int2(1, 1), 0));
+			float4 Pixelx1 = CombinePixels(Pixel01, Pixel11, 0.5);
+
+			Result = CombinePixels(Pixelx0, Pixelx1, 0.5);
 		}
 	}
-#else
-	// single tap mipmap
-	// NOTE: this is wrong as it doesn't account for the fact that each pixel can have a variable number of samples
-	// This is similar to the idea of down-sampling an image with unpremultiplied alpha
-	float4 Result = InputMip.SampleLevel(InputMipSampler, UV, 0);
-#endif
+	Result.z *= 4.0; // even though we did a weighted average, pretend like we have 4x more samples now
 	OutputMip[OutPixelCoord] = Result;
 }
\ No newline at end of file
diff --git a/Engine/Shaders/Private/PathTracing/PathTracingCompositingPixelShader.usf b/Engine/Shaders/Private/PathTracing/PathTracingCompositingPixelShader.usf
--- a/Engine/Shaders/Private/PathTracing/PathTracingCompositingPixelShader.usf
+++ b/Engine/Shaders/Private/PathTracing/PathTracingCompositingPixelShader.usf
@@ -38,7 +38,7 @@ void CompositeMain(
 	{
 		// Visualize which pixels are still being worked and color them according to how many more samples we think they need
 		float2 VarianceUV = float2(TexCoord.xy + 0.5) * View.BufferSizeAndInvSize.zw;
-		float SampleScaleFactor = GetAdaptiveSampleFactor(TexCoord.xy, View.PreExposure / PreExposure); // Run with same exposure as adaptive code
+		float SampleScaleFactor = GetAdaptiveSampleFactor(TexCoord.xy);
 		if (SampleScaleFactor > 1.0)
 		{
 			float4 Variance = VarianceTexture.Load(TexCoord);
@@ -64,7 +64,7 @@ void CompositeMain(
 		float OutVariance = Variance.y - Variance.x * Variance.x;
 		float Lum = Variance.x;
 		float StdErr = sqrt(OutVariance / NumSamples);
-		float RelErr = PerceptualError(Lum, StdErr, View.PreExposure / PreExposure); // Run with same exposure as adaptive code
+		float RelErr = PerceptualError(Lum, StdErr);
 		OutColor.xyz = RelErr;
 	}
 	
diff --git a/Engine/Source/Runtime/Renderer/Private/PathTracing.cpp b/Engine/Source/Runtime/Renderer/Private/PathTracing.cpp
--- a/Engine/Source/Runtime/Renderer/Private/PathTracing.cpp
+++ b/Engine/Source/Runtime/Renderer/Private/PathTracing.cpp
@@ -2097,7 +2097,6 @@ class FPathTracingAdaptiveStartCS : public FGlobalShader
 		SHADER_PARAMETER(float, AdaptiveSamplingErrorThreshold)
 		SHADER_PARAMETER(FIntPoint, TileTextureOffset)
 		SHADER_PARAMETER(FIntPoint, DispatchDim)
-		SHADER_PARAMETER(float, ViewPreExposure)
 		SHADER_PARAMETER_RDG_BUFFER_UAV(RWBuffer<int>, NextActivePaths)
 		SHADER_PARAMETER_RDG_BUFFER_UAV(RWBuffer<int>, NumPathStates)
 	END_SHADER_PARAMETER_STRUCT()
@@ -4149,7 +4148,6 @@ void FDeferredShadingSceneRenderer::RenderPathTracing(
 							PassParameters->VarianceSampler = TStaticSamplerState<SF_Bilinear>::GetRHI();
 							PassParameters->VarianceTextureDims = FIntVector(DispatchResX, DispatchResY, NumVarianceMips);
 							PassParameters->AdaptiveSamplingErrorThreshold = Config.AdaptiveSamplingThreshold;
-							PassParameters->ViewPreExposure = View.PreExposure;
 
 							PassParameters->NextActivePaths = GraphBuilder.CreateUAV(ActivePaths[0], PF_R32_UINT);
 							PassParameters->NumPathStates = GraphBuilder.CreateUAV(NumActivePaths, PF_R32_UINT);
