Virtual texture sampling: make use of our macro wrappers for ddx and ddy.
This then deals with volumetric case where we set USE_FORCE_TEXTURE_MIP.
Regular texture sampling already does this.
The change stops volumetric shaders failing to compile on DX11 where the raymarch loop was being forced to unroll for ddx/ddy.
#rnx

diff --git a/Engine/Shaders/Private/VirtualTextureCommon.ush b/Engine/Shaders/Private/VirtualTextureCommon.ush
--- a/Engine/Shaders/Private/VirtualTextureCommon.ush
+++ b/Engine/Shaders/Private/VirtualTextureCommon.ush
@@ -465,7 +465,7 @@ VTPageTableResult TextureLoadVirtualPageTable(
 		UV = UV * PageTableUniform.UVScale;
 		int vLevel = GetGlobalVirtualTextureMipBias(PageTableUniform.MipBiasGroup);
 #if PIXELSHADER
-		vLevel = TextureComputeVirtualMipLevel(Result, ddx(UV), ddy(UV), MipBias, SvPositionXY, PageTableUniform);
+		vLevel = TextureComputeVirtualMipLevel(Result, DDX(UV), DDY(UV), MipBias, SvPositionXY, PageTableUniform);
 #endif // PIXELSHADER
 		UV = ApplyAddressMode(UV, AddressU, AddressV);
 		TextureLoadVirtualPageTableInternal(Result, PageTable0, PageTableUniform, UV, vLevel);
@@ -487,7 +487,7 @@ VTPageTableResult TextureLoadVirtualPageTable(
 		UV = UV * PageTableUniform.UVScale;
 		int vLevel = GetGlobalVirtualTextureMipBias(PageTableUniform.MipBiasGroup);
 #if PIXELSHADER
-		vLevel = TextureComputeVirtualMipLevel(Result, ddx(UV), ddy(UV), MipBias, SvPositionXY, PageTableUniform);
+		vLevel = TextureComputeVirtualMipLevel(Result, DDX(UV), DDY(UV), MipBias, SvPositionXY, PageTableUniform);
 #endif // PIXELSHADER
 		UV = ApplyAddressMode(UV, AddressU, AddressV);
 		TextureLoadVirtualPageTableInternal(Result, PageTable0, PageTableUniform, UV, vLevel);
@@ -509,7 +509,7 @@ VTPageTableResult TextureLoadVirtualPageTable(
 		UV = UV * PageTableUniform.UVScale;
 		int vLevel = GetGlobalVirtualTextureMipBias(PageTableUniform.MipBiasGroup);
 #if PIXELSHADER
-		vLevel = TextureComputeVirtualMipLevel(Result, ddx(UV), ddy(UV), MipBias, SvPositionXY, PageTableUniform);
+		vLevel = TextureComputeVirtualMipLevel(Result, DDX(UV), DDY(UV), MipBias, SvPositionXY, PageTableUniform);
 #endif // PIXELSHADER
 		UV = ApplyAddressMode(UV, AddressU, AddressV);
 		TextureLoadVirtualPageTableInternal(Result, PageTable0, PageTable1, PageTableUniform, UV, vLevel);
@@ -531,7 +531,7 @@ VTPageTableResult TextureLoadVirtualPageTable(
 		UV = UV * PageTableUniform.UVScale;
 		int vLevel = GetGlobalVirtualTextureMipBias(PageTableUniform.MipBiasGroup);
 #if PIXELSHADER
-		vLevel = TextureComputeVirtualMipLevel(Result, ddx(UV), ddy(UV), MipBias, SvPositionXY, PageTableUniform);
+		vLevel = TextureComputeVirtualMipLevel(Result, DDX(UV), DDY(UV), MipBias, SvPositionXY, PageTableUniform);
 #endif // PIXELSHADER
 		UV = ApplyAddressMode(UV, AddressU, AddressV);
 		TextureLoadVirtualPageTableInternal(Result, PageTable0, PageTable1, PageTableUniform, UV, vLevel);
@@ -701,7 +701,7 @@ VTPageTableResult TextureLoadVirtualPageTableAdaptive(
 	UV = UV * PageTableUniform.UVScale;
 	int vLevel = GetGlobalVirtualTextureMipBias(PageTableUniform.MipBiasGroup);
 #if PIXELSHADER
-	vLevel = TextureComputeVirtualMipLevel(Result, ddx(UV), ddy(UV), MipBias, SvPositionXY, PageTableUniform);
+	vLevel = TextureComputeVirtualMipLevel(Result, DDX(UV), DDY(UV), MipBias, SvPositionXY, PageTableUniform);
 #endif // PIXELSHADER
 	UV = ApplyAddressMode(UV, AddressU, AddressV);
 	ApplyAdaptivePageTableUniform(PageTableIndirection, Result, PageTableUniform, UV, vLevel);
