[Subsurface Profile] Improve neighborhood validation for checkerboard, less bleeding into subsurface scattering region for low res.

#tests done on test level

diff --git a/Engine/Shaders/Private/PostProcessSubsurface.usf b/Engine/Shaders/Private/PostProcessSubsurface.usf
--- a/Engine/Shaders/Private/PostProcessSubsurface.usf
+++ b/Engine/Shaders/Private/PostProcessSubsurface.usf
@@ -447,6 +447,16 @@ half3 LookupSceneColor(float2 SceneUV, int2 PixelOffset)
 #endif
 	return SceneColor;
 }
+
+half3 ZeroAwareAdd(half3 A, half3 B)
+{
+	half3 R = A + B;
+#if CHECKERBOARD_NEIGHBOR_SSS_VALIDATION
+	R += select(all(A == 0), B, 0);
+	R += select(all(B == 0), A, 0);
+#endif
+	return R;
+}
  
 // @param UVSceneColor for the full res rendertarget (BufferSize) e.g. SceneColor or GBuffers
 // @param ReconstructMethod 0/1/2/3 (should be a literal constant to allow compiler optimizations)
@@ -482,26 +492,29 @@ SDiffuseAndSpecular ReconstructLighting(float2 UVSceneColor, uint ReconstructMet
 	  else if(ReconstructMethod == 1)
 	  {
 		  // acceptable but not perfect
-		  Quant1 = 0.5f * (
-			  LookupSceneColor(UVSceneColor, int2( 1, 0)) +
+		  Quant1 = 0.5f * ZeroAwareAdd(
+			  LookupSceneColor(UVSceneColor, int2( 1, 0)),
 			  LookupSceneColor(UVSceneColor, int2(-1, 0)));
 	  }
 	  else if(ReconstructMethod == 2)
 	  {
 		  // almost same as 1?
-		  Quant1 = 0.25f * (
-			  LookupSceneColor(UVSceneColor, int2( 1,  0)) +
-			  LookupSceneColor(UVSceneColor, int2( 0,  1)) +
-			  LookupSceneColor(UVSceneColor, int2(-1,  0)) +
-			  LookupSceneColor(UVSceneColor, int2( 0, -1)));
+		  half3 A = LookupSceneColor(UVSceneColor, int2(  1,  0));
+		  half3 B = LookupSceneColor(UVSceneColor, int2(  0,  1));
+		  half3 C = LookupSceneColor(UVSceneColor, int2( -1,  0));
+		  half3 D = LookupSceneColor(UVSceneColor, int2(  0, -1));
+
+		  Quant1 = 0.25f * ZeroAwareAdd(
+			                     ZeroAwareAdd(A, B),
+			                     ZeroAwareAdd(C, D));
 	  }
 	  else if(ReconstructMethod == 3)
 	  {
 		  // very good 
-		  float3 A = LookupSceneColor(UVSceneColor, int2( 1,  0));
-		  float3 B = LookupSceneColor(UVSceneColor, int2(-1,  0));
-		  float3 C = LookupSceneColor(UVSceneColor, int2( 0,  1));
-		  float3 D = LookupSceneColor(UVSceneColor, int2( 0, -1));
+		  half3 A = LookupSceneColor(UVSceneColor, int2( 1,  0));
+		  half3 B = LookupSceneColor(UVSceneColor, int2(-1,  0));
+		  half3 C = LookupSceneColor(UVSceneColor, int2( 0,  1));
+		  half3 D = LookupSceneColor(UVSceneColor, int2( 0, -1));
 
 		  // Luminance could be green channel only
 		  float a = Luminance(A);
@@ -513,7 +526,7 @@ SDiffuseAndSpecular ReconstructLighting(float2 UVSceneColor, uint ReconstructMet
 		  float cd = abs(c - d);
 
 		  // take the average in the direction that avoids dither pattern
-		  Quant1 = 0.5f * lerp(A + B, C + D, ab > cd);
+		  Quant1 = 0.5f * lerp(ZeroAwareAdd(A, B), ZeroAwareAdd(C, D), ab > cd);
 	  }
 
 	  Ret.Diffuse = lerp(Quant1, Quant0, bChecker);
