Substrate - fix for error "non-empty case statements must have break or return"

#rnx
#tests PC EDITOR & android GL preview

diff --git a/Engine/Shaders/Private/Substrate/Substrate.ush b/Engine/Shaders/Private/Substrate/Substrate.ush
--- a/Engine/Shaders/Private/Substrate/Substrate.ush
+++ b/Engine/Shaders/Private/Substrate/Substrate.ush
@@ -3796,7 +3796,7 @@ float SubstrateGetBSDFOpacity(in FSubstrateBSDF BSDF)
 				case SSS_TYPE_TWO_SIDED_WRAP	: Out = 1.f-SLAB_SSSPHASEANISOTROPY(BSDF); break;
 				case SSS_TYPE_DIFFUSION			: Out = 1.f; break;
 				case SSS_TYPE_DIFFUSION_PROFILE	: Out = SLAB_SSSPROFILERADIUSSCALE(BSDF); break;
-				case SSS_TYPE_SIMPLEVOLUME		: Out = 1.f; 
+				case SSS_TYPE_SIMPLEVOLUME		: Out = 1.f; break;
 			}
 		}
 		else if (BSDF_GETHASFUZZ(BSDF))
