Fix compile errors on some platforms

#tests

diff --git a/Engine/Shaders/Private/PostProcessSubsurface.usf b/Engine/Shaders/Private/PostProcessSubsurface.usf
--- a/Engine/Shaders/Private/PostProcessSubsurface.usf
+++ b/Engine/Shaders/Private/PostProcessSubsurface.usf
@@ -452,8 +452,8 @@ half3 ZeroAwareAdd(half3 A, half3 B)
 {
 	half3 R = A + B;
 #if CHECKERBOARD_NEIGHBOR_SSS_VALIDATION
-	R += select(all(A == 0), B, 0);
-	R += select(all(B == 0), A, 0);
+	R += select(all(A == 0.0f), B, 0.0f);
+	R += select(all(B == 0.0f), A, 0.0f);
 #endif
 	return R;
 }
