Substrate - fix for additive blending adding fog contribution to the surface when it should only be attenuated according to fog transmittance.

#rnx
#tests PC editor dx12

#[review]-45865250 [REDACTED_HANDLE]

diff --git a/Engine/Shaders/Private/BasePassPixelShader.usf b/Engine/Shaders/Private/BasePassPixelShader.usf
--- a/Engine/Shaders/Private/BasePassPixelShader.usf
+++ b/Engine/Shaders/Private/BasePassPixelShader.usf
@@ -2213,6 +2213,12 @@ void FPixelShaderInOut_MainPS(
 			DualBlendSurfaceCoverage = GetMaterialOpacity(PixelMaterialInputs);
 		#endif
 
+			// Additive surfaces should not add any fog color contribution. They should only be attenuated according to fog transmittance.
+			// This must be enforced here since coverage can be used to weight the emissive color.
+		#if SUBSTRATE_OPTIMIZED_UNLIT && MATERIALBLENDING_ADDITIVE
+			Fogging.rgb = 0;
+		#endif
+
 			// Add fog luminance according to surfacecoverage and reduce surface luminance according to fog coverage.
 			float3 AdjustedDualBlendAdd = DualBlendSurfaceCoverage * Fogging.rgb + Fogging.a * DualBlendSurfaceLuminancePostCoverage;
 			// Fade the surface color transmittance out to 1 according to the surface coverage, and take into account the fog coverage to the surface.
