Substrate - fix for android cooking error

#tests PC editor dx12 and cooking
#rnx
review https://p4-swarm.epicgames.net/reviews/45837172

diff --git a/Engine/Shaders/Private/MobileBasePassPixelShader.usf b/Engine/Shaders/Private/MobileBasePassPixelShader.usf
--- a/Engine/Shaders/Private/MobileBasePassPixelShader.usf
+++ b/Engine/Shaders/Private/MobileBasePassPixelShader.usf
@@ -1235,7 +1235,7 @@ void Main(
 #endif // TEMPLATE_USES_SUBSTRATE
 
 
-#if !TEMPLATE_USES_SUBSTRATE || SUBSTRATE_INLINE_SINGLELAYERWATER
+#if !SUBSTRATE_ENABLED || TEMPLATE_USES_LEGACYGBUFFERDATA || SUBSTRATE_INLINE_SINGLELAYERWATER
 
 #if !MATERIAL_SHADINGMODEL_UNLIT
 	half3 Color = DirectLighting.TotalLight;
@@ -1262,11 +1262,11 @@ void Main(
 
 	Color += Emissive;
 
-	// This logic matches MaterialRequiresColorTransmittanceBlending
-#if (MOBILE_TRANSLUCENT_COLOR_TRANSMITTANCE_DUAL_SRC_BLENDING || MATERIAL_SHADINGMODEL_THIN_TRANSLUCENT) && !TEMPLATE_USES_SUBSTRATE
+#if MOBILE_TRANSLUCENT_COLOR_TRANSMITTANCE
 	half3 DualBlendSurfaceLuminancePostCoverage = 0.0f;
 	half3 DualBlendSurfaceTransmittancePreCoverage = 1.0f;
 	half DualBlendSurfaceCoverage = 1.0f;
+#if MATERIAL_SHADINGMODEL_THIN_TRANSLUCENT && (!SUBSTRATE_ENABLED || TEMPLATE_USES_LEGACYGBUFFERDATA)
 	{
 		AccumulateThinTranslucentModel(
 			DualBlendSurfaceLuminancePostCoverage,
@@ -1282,7 +1282,8 @@ void Main(
 		Color = 0;
 		Opacity = 1.0f;
 	}
-#endif // MATERIAL_SHADINGMODEL_THIN_TRANSLUCENT
+#endif
+#endif // MOBILE_TRANSLUCENT_COLOR_TRANSMITTANCE
 	
 	// MATERIAL_SHADINGMODEL_THIN_TRANSLUCENT must come first because it also has MATERIALBLENDING_TRANSLUCENT defined
 	#if MOBILE_TRANSLUCENT_COLOR_TRANSMITTANCE
@@ -1325,7 +1326,7 @@ void Main(
 			OutColor.a = uint((Mask + 0.25f) * 4.0f) / 4.0f;
 		#endif
 	#endif
-#endif // TEMPLATE_USES_SUBSTRATE
+#endif
 
 #if HAS_SCENE_DEPTH_AUX_OUTPUT
 	OutSceneDepthAux = SvPosition.z;
diff --git a/Engine/Shaders/Private/ThinTranslucentCommon.ush b/Engine/Shaders/Private/ThinTranslucentCommon.ush
--- a/Engine/Shaders/Private/ThinTranslucentCommon.ush
+++ b/Engine/Shaders/Private/ThinTranslucentCommon.ush
@@ -6,8 +6,7 @@
 
 #define THIN_TRANSLUCENT_USE_DOTNV_THICKNESS 1
 
-// This logic matches MaterialRequiresColorTransmittanceBlending
-#if MATERIAL_SHADINGMODEL_THIN_TRANSLUCENT || MOBILE_TRANSLUCENT_COLOR_TRANSMITTANCE_DUAL_SRC_BLENDING
+#if MATERIAL_SHADINGMODEL_THIN_TRANSLUCENT || MOBILE_TRANSLUCENT_COLOR_TRANSMITTANCE && (MATERIAL_SHADINGMODEL_THIN_TRANSLUCENT && !TEMPLATE_USES_SUBSTRATE)
 void AccumulateThinTranslucentModel(inout float3 DualBlendSurfaceLuminancePostCoverage,
 								inout float3 DualBlendSurfaceTransmittancePreCoverage,
 								inout float DualBlendSurfaceCoverage,
