Add early out in Substrate lighting based on shadow mask attenuation for all path, and based on NoL for the fasth path.

This was measured as a 0.19ms gain on a simple directional light with favorable condition.

diff --git a/Engine/Shaders/Private/Substrate/SubstrateDeferredLighting.ush b/Engine/Shaders/Private/Substrate/SubstrateDeferredLighting.ush
--- a/Engine/Shaders/Private/Substrate/SubstrateDeferredLighting.ush
+++ b/Engine/Shaders/Private/Substrate/SubstrateDeferredLighting.ush
@@ -134,57 +134,74 @@ FSubstrateDeferredLighting SubstrateDeferredLighting(
 				BSDFShadowTerms);
 		}
 
+		// Lighting early out based on shadow attenuation
+		const bool bNeedLighting = (BSDFShadowTerms.SurfaceShadow + BSDFShadowTerms.TransmissionShadow) > 0;
 
-		float Roughness = SubstrateGetBSDFRoughness(BSDFContext.BSDF);
-
-		FAreaLight AreaLight = InitAreaLight();
-		FSubstrateEvaluateResult BSDFEvaluate = (FSubstrateEvaluateResult)0;
-		if (LightData.bRectLight)
-		{
-			FRect Rect = GetRect(ToLight, LightData);
-			if (!IsRectVisible(Rect))
-			{
-				return GetInitialisedSubstrateDeferredLighting(); // Rect light can be non visible due to barn door occlusion
-			}
-			AreaLight = CreateAreaLight(Roughness, BSDFContext.N, BSDFContext.V, Rect, RectTexture);
-
-			// We must have the SubstrateIntegrateBSDF inside the if due to the rectlight texture: it must be non ambiguous which texture is going to be used.
-			// After the compilation, a local resource must map to a unique global resource (the default or the actual rect light texture).
-			BSDFEvaluate = SubstrateEvaluateBSDFCommon(BSDFContext, BSDFShadowTerms, AreaLight, Settings, INTEGRATION_AREA_LIGHT_RECT);
-		}
-		else
+		#if !COMPILER_FXC // Work around "error X3524: duplicate attribute branch" in cases where this branch is constant
+		BRANCH
+		#endif
+		if (bNeedLighting)
 		{
-			FCapsuleLight Capsule = GetCapsule(ToLight, LightData);
-			AreaLight = CreateAreaLight(Roughness, BSDFContext.N, BSDFContext.V, Capsule, LightData.bInverseSquared);
+			float Roughness = SubstrateGetBSDFRoughness(BSDFContext.BSDF);
 
-			BRANCH
-			if(IsAreaLight(AreaLight))
+			FAreaLight AreaLight = InitAreaLight();
+			FSubstrateEvaluateResult BSDFEvaluate = (FSubstrateEvaluateResult)0;
+			if (LightData.bRectLight)
 			{
-				BSDFEvaluate = SubstrateEvaluateBSDFCommon(BSDFContext, BSDFShadowTerms, AreaLight, Settings, INTEGRATION_AREA_LIGHT_CAPSULE);
+				FRect Rect = GetRect(ToLight, LightData);
+				if (!IsRectVisible(Rect))
+				{
+					return GetInitialisedSubstrateDeferredLighting(); // Rect light can be non visible due to barn door occlusion
+				}
+				AreaLight = CreateAreaLight(Roughness, BSDFContext.N, BSDFContext.V, Rect, RectTexture);
+
+				// We must have the SubstrateIntegrateBSDF inside the if due to the rectlight texture: it must be non ambiguous which texture is going to be used.
+				// After the compilation, a local resource must map to a unique global resource (the default or the actual rect light texture).
+				BSDFEvaluate = SubstrateEvaluateBSDFCommon(BSDFContext, BSDFShadowTerms, AreaLight, Settings, INTEGRATION_AREA_LIGHT_RECT);
 			}
 			else
 			{
-				BSDFEvaluate = SubstrateEvaluateBSDFCommon(BSDFContext, BSDFShadowTerms, AreaLight, Settings, INTEGRATION_PUNCTUAL_LIGHT);
+				FCapsuleLight Capsule = GetCapsule(ToLight, LightData);
+				AreaLight = CreateAreaLight(Roughness, BSDFContext.N, BSDFContext.V, Capsule, LightData.bInverseSquared);
+
+				// Lighting early out based NoL. Only for SUBSTRATE_FASTPATH, as this path as no SSS/Wrap features.
+				#if SUBSTRATE_FASTPATH
+				BRANCH
+				if (AreaLight.NoL<=0)
+				{
+					return GetInitialisedSubstrateDeferredLighting(); // Rect light can be non visible due to barn door occlusion
+				}
+				#endif
+
+				BRANCH
+				if(IsAreaLight(AreaLight))
+				{
+					BSDFEvaluate = SubstrateEvaluateBSDFCommon(BSDFContext, BSDFShadowTerms, AreaLight, Settings, INTEGRATION_AREA_LIGHT_CAPSULE);
+				}
+				else
+				{
+					BSDFEvaluate = SubstrateEvaluateBSDFCommon(BSDFContext, BSDFShadowTerms, AreaLight, Settings, INTEGRATION_PUNCTUAL_LIGHT);
+				}
 			}
-		}
 
-		SubstrateLighting.EstimatedCost += 0.4;	// add the cost of the lighting computations (should sum up to 1 form one light)
+			SubstrateLighting.EstimatedCost += 0.4;	// add the cost of the lighting computations (should sum up to 1 form one light)
 
-		float3 DiffuseLuminance = BSDFEvaluate.IntegratedDiffuseValue * LightData.DiffuseScale;
-		float3 SpecularLuminance = BSDFEvaluate.IntegratedSpecularValue * LightData.SpecularScale;
+			float3 DiffuseLuminance = BSDFEvaluate.IntegratedDiffuseValue * LightData.DiffuseScale;
+			float3 SpecularLuminance = BSDFEvaluate.IntegratedSpecularValue * LightData.SpecularScale;
 
-		const float3 CommonMultiplier = LightData.Color * LightMask * LuminanceWeight(BSDFContext, BSDF)
-	#if USE_LIGHT_FUNCTION_ATLAS
-			* LightFunctionColor
-	#endif
-			;
+			const float3 CommonMultiplier = LightData.Color * LightMask * LuminanceWeight(BSDFContext, BSDF)
+		#if USE_LIGHT_FUNCTION_ATLAS
+				* LightFunctionColor
+		#endif
+				;
 
-		FLightAccumulator Out = (FLightAccumulator)0;
-		LightAccumulator_AddSplit(Out, DiffuseLuminance, SpecularLuminance, DiffuseLuminance, CommonMultiplier, BSDFEvaluate.bPostProcessSubsurface);
-		AccumulateSubstrateDeferredLighting(SubstrateLighting, Out, BSDFEvaluate.bPostProcessSubsurface, BSDF_GETISTOPLAYER(BSDF));
+			FLightAccumulator Out = (FLightAccumulator)0;
+			LightAccumulator_AddSplit(Out, DiffuseLuminance, SpecularLuminance, DiffuseLuminance, CommonMultiplier, BSDFEvaluate.bPostProcessSubsurface);
+			AccumulateSubstrateDeferredLighting(SubstrateLighting, Out, BSDFEvaluate.bPostProcessSubsurface, BSDF_GETISTOPLAYER(BSDF));
 
-		SubstrateLighting.TotalDiffuseLighting += Out.TotalLightDiffuse;
-		SubstrateLighting.TotalSpecularLighting += Out.TotalLightSpecular;
+			SubstrateLighting.TotalDiffuseLighting += Out.TotalLightDiffuse;
+			SubstrateLighting.TotalSpecularLighting += Out.TotalLightSpecular;
+		}
 	}
 
 	return SubstrateLighting;
