[HWRT] World Normal Debug Visualization with Inline RayTracing

diff --git a/Engine/Shaders/Private/RayTracing/RayTracingDebug.usf b/Engine/Shaders/Private/RayTracing/RayTracingDebug.usf
--- a/Engine/Shaders/Private/RayTracing/RayTracingDebug.usf
+++ b/Engine/Shaders/Private/RayTracing/RayTracingDebug.usf
@@ -19,6 +19,7 @@
 #include "../Substrate/Substrate.ush"
 #include "../Substrate/SubstrateEvaluation.ush"
 #include "../Substrate/SubstrateVisualizeCommon.ush"
+#include "../Nanite/NaniteRayTrace.ush"
 #include "/Engine/Shared/SubstrateVisualizeDefinitions.h"
 
 #include "RayTracingDebugUtils.ush"
@@ -32,6 +33,10 @@ RWTexture2D<float> OutputDepth;
 RaytracingAccelerationStructure TLAS;
 RaytracingAccelerationStructure FarFieldTLAS;
 
+#if RAYTRACING_DEBUG_INLINE
+	StructuredBuffer<FRayTracingSceneMetadataRecord> RayTracingSceneMetadata;
+#endif
+
 uint VisualizationMode;
 uint PickerDomain;
 uint ShouldUsePreExposure;
@@ -268,10 +273,45 @@ float4 DebugRayTracingInstance(FRayDesc Ray, out float OutDepth)
 	const FInstanceSceneData InstanceSceneData = GetInstanceSceneData(GPUSceneInstanceId);
 
 	Payload.InstanceHash = MurmurAdd(InstanceSceneData.PrimitiveId, InstanceSceneData.RelativeId);
-
 	Payload.TriangleIndex = TraceRayResult.PrimitiveIndex;
 
-	Payload.WorldNormal = TraceRayResult.WorldGeometryNormal;
+	float3 WorldNormal;
+#if PLATFORM_SUPPORTS_INLINE_RAY_TRACING_TRIANGLE_NORMALS
+	WorldNormal = TraceRayResult.WorldGeometryNormal;
+#elif PLATFORM_SUPPORTS_INLINE_RAY_TRACING_TRIANGLE_ATTRIBUTES
+	// NANITE_RAY_TRACING
+	FPrimitiveSceneData PrimitiveData = GetPrimitiveData(InstanceSceneData.PrimitiveId);
+	if (PrimitiveData.NaniteRayTracingDataOffset != uint(-1))
+	{
+		const uint HitGroupFirstPrimitive = GetInlineRayTracingHitGroupFirstPrimitive(
+			RayTracingSceneMetadata,
+			TraceRayResult.InstanceContributionToHitGroupIndex,
+			RAY_TRACING_NUM_SHADER_SLOTS,
+			TraceRayResult.GeometryIndex
+		);
+					
+		float3 LocalEdges[2];
+		GetNaniteTriangleEdges(PrimitiveData.NaniteRayTracingDataOffset + HitGroupFirstPrimitive + TraceRayResult.PrimitiveIndex, PrimitiveData.NaniteAssemblyTransformOffset, LocalEdges[0], LocalEdges[1]);
+
+		float3 LocalNormal = cross(LocalEdges[1], LocalEdges[0]);
+		WorldNormal = normalize(mul(LocalNormal, (float3x3)TraceRayResult.WorldToObject3x4));
+	}
+	else
+	{
+		FTriangleBaseAttributes Attributes = LoadInlineRayTracingTriangleAttributes(
+			RayTracingSceneMetadata,
+			TraceRayResult.InstanceContributionToHitGroupIndex,
+			RAY_TRACING_NUM_SHADER_SLOTS,
+			TraceRayResult.GeometryIndex,
+			TraceRayResult.PrimitiveIndex
+		);
+		float3 LocalNormal = cross(Attributes.LocalPositions[2] - Attributes.LocalPositions[0], Attributes.LocalPositions[1] - Attributes.LocalPositions[0]);
+		WorldNormal = normalize(mul(LocalNormal, (float3x3)TraceRayResult.WorldToObject3x4));	
+	}
+#else // !PLATFORM_SUPPORTS_INLINE_RAY_TRACING_TRIANGLE_ATTRIBUTES
+	WorldNormal = 0.0f;
+#endif // !PLATFORM_SUPPORTS_INLINE_RAY_TRACING_TRIANGLE_ATTRIBUTES
+	Payload.WorldNormal = WorldNormal;
 	Payload.InstanceIndex = TraceRayResult.InstanceIndex;
 	Payload.GeometryIndex = TraceRayResult.GeometryIndex;
 	Payload.ScenePrimitiveIndex = InstanceSceneData.PrimitiveId;
@@ -302,6 +342,9 @@ float4 DebugRayTracingInstance(FRayDesc Ray, out float OutDepth)
 		case RAY_TRACING_DEBUG_VIZ_INSTANCES:
 			Result.rgb = InstanceColor;
 			break;
+		case RAY_TRACING_DEBUG_VIZ_WORLD_NORMAL:
+			Result.rgb = float3(Payload.WorldNormal * 0.5 + 0.5);
+			break;
 		case RAY_TRACING_DEBUG_VIZ_TRIANGLE_HITCOUNT:
 			Result.rgb = ColorMapTurbo(Payload.TriangleHitCountPerRay / TriangleHitCountMaxThreshold);
 			break;
@@ -363,7 +406,7 @@ float4 DebugRayTracingInstance(FRayDesc Ray, out float OutDepth)
 		const float FakeLighting = dot(Payload.WorldNormal, normalize(float3(1.0f, 1.0f, 1.0f)));
 		const float FakeLightingScale = (VisualizationMode == RAY_TRACING_DEBUG_VIZ_INSTANCE_OVERLAP) ? 0.5f : 0.3f;
 		const float FakeLightingBias = (VisualizationMode == RAY_TRACING_DEBUG_VIZ_INSTANCE_OVERLAP) ? 0.5f : 0.7f;
-		const int NoFakeLighting = (VisualizationMode == RAY_TRACING_DEBUG_VIZ_TRIANGLE_HITCOUNT) || (VisualizationMode == RAY_TRACING_DEBUG_VIZ_HITCOUNT_PER_INSTANCE);
+		const int NoFakeLighting = (VisualizationMode == RAY_TRACING_DEBUG_VIZ_TRIANGLE_HITCOUNT) || (VisualizationMode == RAY_TRACING_DEBUG_VIZ_HITCOUNT_PER_INSTANCE) || (VisualizationMode == RAY_TRACING_DEBUG_VIZ_WORLD_NORMAL);
 
 		if (!NoFakeLighting)
 		{
diff --git a/Engine/Shaders/Shared/RaytracingDebugDefinitions.h b/Engine/Shaders/Shared/RaytracingDebugDefinitions.h
--- a/Engine/Shaders/Shared/RaytracingDebugDefinitions.h
+++ b/Engine/Shaders/Shared/RaytracingDebugDefinitions.h
@@ -99,10 +99,12 @@ INLINE_ATTR bool IsRayTracingDebugTraversalMode(UINT_TYPE DebugVisualizationMode
 INLINE_ATTR bool RayTracingDebugModeSupportsInline(UINT_TYPE DebugVisualizationMode)
 {
 	const bool bIsBarycentricsMode = DebugVisualizationMode == RAY_TRACING_DEBUG_VIZ_BARYCENTRICS;
+	const bool bIsWorldNormalMode = DebugVisualizationMode == RAY_TRACING_DEBUG_VIZ_WORLD_NORMAL;
 	const bool bUsesDebugCHSAndNotHitCount = RequiresRayTracingDebugCHS(DebugVisualizationMode)
 		&& (DebugVisualizationMode != RAY_TRACING_DEBUG_VIZ_TRIANGLE_HITCOUNT)
 		&& (DebugVisualizationMode != RAY_TRACING_DEBUG_VIZ_HITCOUNT_PER_INSTANCE);
 	const bool bIsTraversalMode = IsRayTracingDebugTraversalMode(DebugVisualizationMode);
+	const bool bIsPickingMode = DebugVisualizationMode == RAY_TRACING_DEBUG_VIZ_PICKER;
 
-	return bIsBarycentricsMode || bUsesDebugCHSAndNotHitCount || bIsTraversalMode;
+	return bIsBarycentricsMode || bIsWorldNormalMode || bIsPickingMode || bUsesDebugCHSAndNotHitCount || bIsTraversalMode;
 }
diff --git a/Engine/Source/Runtime/Renderer/Private/DeferredShadingRenderer.cpp b/Engine/Source/Runtime/Renderer/Private/DeferredShadingRenderer.cpp
--- a/Engine/Source/Runtime/Renderer/Private/DeferredShadingRenderer.cpp
+++ b/Engine/Source/Runtime/Renderer/Private/DeferredShadingRenderer.cpp
@@ -322,6 +322,11 @@ namespace Nanite
 	extern void ListStatFilters(FSceneRenderer* SceneRenderer);
 }
 
+namespace RayTracingDebug
+{
+	extern bool UseInlineHardwareRayTracing(const FSceneViewFamily& ViewFamily);
+}
+
 DECLARE_CYCLE_STAT(TEXT("InitViews Intentional Stall"), STAT_InitViews_Intentional_Stall, STATGROUP_InitViews);
 
 DECLARE_CYCLE_STAT(TEXT("DeferredShadingSceneRenderer UpdateDownsampledDepthSurface"), STAT_FDeferredShadingSceneRenderer_UpdateDownsampledDepthSurface, STATGROUP_SceneRendering);
@@ -1027,7 +1032,8 @@ void FDeferredShadingSceneRenderer::SetupRayTracingRenderingData(FRDGBuilder& Gr
 		}
 		
 		if (Lumen::AnyLumenHardwareInlineRayTracingPassEnabled(Scene, View)
-			|| MegaLights::UseInlineHardwareRayTracing(ViewFamily))
+			|| MegaLights::UseInlineHardwareRayTracing(ViewFamily)
+			|| RayTracingDebug::UseInlineHardwareRayTracing(ViewFamily))
 		{
 			bAnyInlineHardwareRayTracingPassEnabled = true;
 		}
diff --git a/Engine/Source/Runtime/Renderer/Private/RayTracing/RayTracingDebug.cpp b/Engine/Source/Runtime/Renderer/Private/RayTracing/RayTracingDebug.cpp
--- a/Engine/Source/Runtime/Renderer/Private/RayTracing/RayTracingDebug.cpp
+++ b/Engine/Source/Runtime/Renderer/Private/RayTracing/RayTracingDebug.cpp
@@ -29,6 +29,31 @@
 
 DECLARE_GPU_STAT(RayTracingDebug);
 
+static uint32 GetRaytracingDebugViewModeID(const FSceneView& View);
+
+namespace RayTracingDebug
+{
+	bool UseInlineHardwareRayTracing(const FSceneViewFamily& ViewFamily)
+	{
+		if (ViewFamily.Views.IsEmpty() || !ViewFamily.EngineShowFlags.RayTracingDebug)
+		{
+			return false;
+		}
+
+		// Any of the views require inline raytracing
+		for (const FSceneView* View : ViewFamily.Views)
+		{
+			if (View && ShouldRenderRayTracingEffect(true, ERayTracingPipelineCompatibilityFlags::Inline, *View) 
+				&& RayTracingDebugModeSupportsInline(GetRaytracingDebugViewModeID(*View)))
+			{
+				return true;
+			}
+		}
+
+		return false;
+	}
+}
+
 TAutoConsoleVariable<int32> CVarRayTracingVisualizePickerDomain(
 	TEXT("r.RayTracing.Visualize.PickerDomain"),
 	0,
@@ -185,6 +210,8 @@ BEGIN_SHADER_PARAMETER_STRUCT(FRayTracingDebugParameters, )
 	SHADER_PARAMETER_RDG_UNIFORM_BUFFER(FRayTracingLightGrid, LightGridPacked)
 	SHADER_PARAMETER_STRUCT_INCLUDE(ShaderPrint::FShaderParameters, ShaderPrintUniformBuffer)
 	SHADER_PARAMETER_RDG_BUFFER_UAV(RWBuffer<int>, SubstrateDebugDataUAV)
+	// Inline data
+	SHADER_PARAMETER_RDG_BUFFER_SRV(StructuredBuffer, RayTracingSceneMetadata)
 END_SHADER_PARAMETER_STRUCT()
 
 class FRayTracingDebugRGS : public FGlobalShader
@@ -258,6 +285,9 @@ class FRayTracingDebugCS : public FGlobalShader
 
 		OutEnvironment.SetDefine(TEXT("INLINE_RAY_TRACING_THREAD_GROUP_SIZE_X"), ThreadGroupSizeX);
 		OutEnvironment.SetDefine(TEXT("INLINE_RAY_TRACING_THREAD_GROUP_SIZE_Y"), ThreadGroupSizeY);
+	
+		// Nanite RayTracing
+		OutEnvironment.SetDefine(TEXT("VF_SUPPORTS_PRIMITIVE_SCENE_DATA"), 1);
 	}
 
 	BEGIN_SHADER_PARAMETER_STRUCT(FParameters, )
@@ -1724,6 +1754,8 @@ void RenderRayTracingDebug(FRDGBuilder& GraphBuilder, const FScene& Scene, const
 	SharedParameters.NumTotalInstances = NumInstances;
 	ShaderPrint::SetParameters(GraphBuilder, View.ShaderPrintData, SharedParameters.ShaderPrintUniformBuffer);
 
+	SharedParameters.RayTracingSceneMetadata = View.InlineRayTracingBindingDataBuffer ? GraphBuilder.CreateSRV(View.InlineRayTracingBindingDataBuffer) : nullptr;
+
 	FSubstrateViewDebugData::FTransientPixelDebugBuffer SubstratePixelDebugBuffer;
 	if (bSubstratePixelDebugEnable)
 	{
