[HZBTracing] Fix bStepOutOfCurrentTile logic causing leaking.

#FailedValidations Preflight

diff --git a/Engine/Shaders/Private/HZBTracing.ush b/Engine/Shaders/Private/HZBTracing.ush
--- a/Engine/Shaders/Private/HZBTracing.ush
+++ b/Engine/Shaders/Private/HZBTracing.ush
@@ -79,6 +79,7 @@ void TraceHZB(
 
 	float3 RayDirectionScreenUV = RayEndScreenUV - RayStartScreenUV;
 	float CurrentIntersectionTime = 0;
+	float LastAboveSurfaceTime = 0;
 
 	// Offset to pick which XY boundary planes to intersect
 	float2 FloorOffset = select(RayDirectionScreenUV.xy < 0, 0.0, 1.0);
@@ -106,11 +107,16 @@ void TraceHZB(
 		RayScreenUV = RayStartScreenUV + CurrentIntersectionTime * RayDirectionScreenUV;
 	}
 
+	if (CurrentIntersectionTime >= 1.0f)
+	{
+		// If bStepOutOfCurrentTile caused the ray to terminate rewind to the start of the ray since screen trace will be skipped
+		RayScreenUV = RayStartScreenUV;
+	}
+
 	float NumIterations = 0;
 	bHit = false;
 	bUncertain = false;
 	OutHitTileZ = 0;
-	float LastAboveSurfaceTime = CurrentIntersectionTime;
 
 	// Stackless HZB traversal
 	while (MipLevel >= BaseMipLevel
