Remove incomplete ISR/MMV support from GatherDeviceZ/RayTraceShadows, fixing a shader compilation failure for FDistanceFieldShadowingCS.

diff --git a/Engine/Shaders/Private/SceneTexturesCommon.ush b/Engine/Shaders/Private/SceneTexturesCommon.ush
--- a/Engine/Shaders/Private/SceneTexturesCommon.ush
+++ b/Engine/Shaders/Private/SceneTexturesCommon.ush
@@ -211,11 +211,7 @@ float LookupDeviceZ( float2 ScreenUV, float ArrayIndex )
 	return SCENE_TEXTURES_DISABLED_SCENE_DEPTH_VALUE;
 #elif FORCE_DEPTH_TEXTURE_READS || PLATFORM_NEEDS_DEPTH_TEXTURE_READS
 	// native Depth buffer lookup
-	#if MOBILE_MULTI_VIEW
-		return Texture2DArraySample(MobileSceneTextures.SceneDepthTextureArray, MobileSceneTextures.SceneDepthTextureSampler, float3(ScreenUV.xy,ArrayIndex)).r;
-	#else
-		return Texture2DSampleLevel(MobileSceneTextures.SceneDepthTexture, MobileSceneTextures.SceneDepthTextureSampler, ScreenUV, 0).r;
-	#endif
+	return Texture2DSampleLevel(MobileSceneTextures.SceneDepthTexture, MobileSceneTextures.SceneDepthTextureSampler, ScreenUV, 0).r;
 #elif MOBILE_DEPTHFECTH && COMPILER_GLSL_ES3_1
 	return DepthbufferFetchES2();
 #elif MOBILE_DEPTHFECTH && VULKAN_PROFILE
@@ -225,18 +221,10 @@ float LookupDeviceZ( float2 ScreenUV, float ArrayIndex )
 #elif (USE_SCENE_DEPTH_AUX && !MOBILE_DEFERRED_SHADING)
 	// SceneDepth texture is discarded after BasePass (with forward shading) 
 	// instead fetch DeviceZ from SceneDepthAuxTexture
-	#if MOBILE_MULTI_VIEW
-		return Texture2DArraySample(MobileSceneTextures.SceneDepthAuxTextureArray, MobileSceneTextures.SceneDepthAuxTextureSampler, float3(ScreenUV.xy,ArrayIndex)).r;
-	#else
-		return Texture2DSampleLevel(MobileSceneTextures.SceneDepthAuxTexture, MobileSceneTextures.SceneDepthAuxTextureSampler, ScreenUV, 0).r;
-	#endif
+	return Texture2DSampleLevel(MobileSceneTextures.SceneDepthAuxTexture, MobileSceneTextures.SceneDepthAuxTextureSampler, ScreenUV, 0).r;
 #else
 	// native Depth buffer lookup
-	#if MOBILE_MULTI_VIEW
-		return Texture2DArraySample(MobileSceneTextures.SceneDepthTextureArray, MobileSceneTextures.SceneDepthTextureSampler, float3(ScreenUV.xy,ArrayIndex)).r;
-	#else
-		return Texture2DSampleLevel(MobileSceneTextures.SceneDepthTexture, MobileSceneTextures.SceneDepthTextureSampler, ScreenUV, 0).r;
-	#endif
+	return Texture2DSampleLevel(MobileSceneTextures.SceneDepthTexture, MobileSceneTextures.SceneDepthTextureSampler, ScreenUV, 0).r;
 #endif
 }
 
