Follow-up to CL 45928947: correct changes to non-ArrayIndex GatherZ function.

diff --git a/Engine/Shaders/Private/SceneTexturesCommon.ush b/Engine/Shaders/Private/SceneTexturesCommon.ush
--- a/Engine/Shaders/Private/SceneTexturesCommon.ush
+++ b/Engine/Shaders/Private/SceneTexturesCommon.ush
@@ -211,7 +211,11 @@ float LookupDeviceZ( float2 ScreenUV, float ArrayIndex )
 	return SCENE_TEXTURES_DISABLED_SCENE_DEPTH_VALUE;
 #elif FORCE_DEPTH_TEXTURE_READS || PLATFORM_NEEDS_DEPTH_TEXTURE_READS
 	// native Depth buffer lookup
-	return Texture2DSampleLevel(MobileSceneTextures.SceneDepthTexture, MobileSceneTextures.SceneDepthTextureSampler, ScreenUV, 0).r;
+	#if MOBILE_MULTI_VIEW
+		return Texture2DArraySample(MobileSceneTextures.SceneDepthTextureArray, MobileSceneTextures.SceneDepthTextureSampler, float3(ScreenUV.xy,ArrayIndex)).r;
+	#else
+		return Texture2DSampleLevel(MobileSceneTextures.SceneDepthTexture, MobileSceneTextures.SceneDepthTextureSampler, ScreenUV, 0).r;
+	#endif
 #elif MOBILE_DEPTHFECTH && COMPILER_GLSL_ES3_1
 	return DepthbufferFetchES2();
 #elif MOBILE_DEPTHFECTH && VULKAN_PROFILE
@@ -221,10 +225,18 @@ float LookupDeviceZ( float2 ScreenUV, float ArrayIndex )
 #elif (USE_SCENE_DEPTH_AUX && !MOBILE_DEFERRED_SHADING)
 	// SceneDepth texture is discarded after BasePass (with forward shading) 
 	// instead fetch DeviceZ from SceneDepthAuxTexture
-	return Texture2DSampleLevel(MobileSceneTextures.SceneDepthAuxTexture, MobileSceneTextures.SceneDepthAuxTextureSampler, ScreenUV, 0).r;
+	#if MOBILE_MULTI_VIEW
+		return Texture2DArraySample(MobileSceneTextures.SceneDepthAuxTextureArray, MobileSceneTextures.SceneDepthAuxTextureSampler, float3(ScreenUV.xy,ArrayIndex)).r;
+	#else
+		return Texture2DSampleLevel(MobileSceneTextures.SceneDepthAuxTexture, MobileSceneTextures.SceneDepthAuxTextureSampler, ScreenUV, 0).r;
+	#endif
 #else
 	// native Depth buffer lookup
-	return Texture2DSampleLevel(MobileSceneTextures.SceneDepthTexture, MobileSceneTextures.SceneDepthTextureSampler, ScreenUV, 0).r;
+	#if MOBILE_MULTI_VIEW
+		return Texture2DArraySample(MobileSceneTextures.SceneDepthTextureArray, MobileSceneTextures.SceneDepthTextureSampler, float3(ScreenUV.xy,ArrayIndex)).r;
+	#else
+		return Texture2DSampleLevel(MobileSceneTextures.SceneDepthTexture, MobileSceneTextures.SceneDepthTextureSampler, ScreenUV, 0).r;
+	#endif
 #endif
 }
 
@@ -241,26 +253,14 @@ float4 GatherDeviceZ( float2 ScreenUV )
 	return SCENE_TEXTURES_DISABLED_SCENE_DEPTH_VALUE;
 #elif FORCE_DEPTH_TEXTURE_READS || PLATFORM_NEEDS_DEPTH_TEXTURE_READS
 	// native Depth buffer lookup
-	#if MOBILE_MULTI_VIEW
-	return MobileSceneTextures.SceneDepthTextureArray.GatherRed(MobileSceneTextures.SceneDepthTextureSampler, float3(ScreenUV.xy,ArrayIndex));
-	#else
-		return MobileSceneTextures.SceneDepthTexture.GatherRed(MobileSceneTextures.SceneDepthTextureSampler, ScreenUV);
-	#endif
+	return MobileSceneTextures.SceneDepthTexture.GatherRed(MobileSceneTextures.SceneDepthTextureSampler, ScreenUV);
 #elif (USE_SCENE_DEPTH_AUX && !MOBILE_DEFERRED_SHADING)
 	// SceneDepth texture is discarded after BasePass (with forward shading) 
 	// instead fetch DeviceZ from SceneDepthAuxTexture
-	#if MOBILE_MULTI_VIEW
-		return MobileSceneTextures.SceneDepthAuxTextureArray.GatherRed(MobileSceneTextures.SceneDepthAuxTextureSampler, float3(ScreenUV.xy,ArrayIndex));
-	#else
-		return MobileSceneTextures.SceneDepthAuxTexture.GatherRed(MobileSceneTextures.SceneDepthAuxTextureSampler, ScreenUV);
-	#endif
+	return MobileSceneTextures.SceneDepthAuxTexture.GatherRed(MobileSceneTextures.SceneDepthAuxTextureSampler, ScreenUV);
 #else
 	// native Depth buffer lookup
-	#if MOBILE_MULTI_VIEW
-		return MobileSceneTextures.SceneDepthTextureArray.GatherRed(MobileSceneTextures.SceneDepthTextureSampler, float3(ScreenUV.xy,ArrayIndex));
-	#else
-		return MobileSceneTextures.SceneDepthTexture.GatherRed(MobileSceneTextures.SceneDepthTextureSampler, ScreenUV);
-	#endif
+	return MobileSceneTextures.SceneDepthTexture.GatherRed(MobileSceneTextures.SceneDepthTextureSampler, ScreenUV);
 #endif
 }
 
