Fix error compiling some permutations of BasePassPixelShader.usf.

diff --git a/Engine/Shaders/Private/BasePassPixelShader.usf b/Engine/Shaders/Private/BasePassPixelShader.usf
--- a/Engine/Shaders/Private/BasePassPixelShader.usf
+++ b/Engine/Shaders/Private/BasePassPixelShader.usf
@@ -242,6 +242,8 @@ bool ManualDepthTestEqual(float4 SvPosition, float PSDeviceZWithOffset)
 
 #endif // !FORWARD_SHADING
 
+#if !FORWARD_SHADING || TRANSLUCENCY_LIGHTING_SURFACE_FORWARDSHADING
+
 /** Calculates lighting for translucency. */
 void GetTranslucencyVolumeLighting(
 	FMaterialPixelParameters MaterialParameters, 
@@ -395,6 +397,8 @@ void GetTranslucencyVolumeLighting(
 	#endif
 }
 
+#endif // !FORWARD_SHADING || TRANSLUCENCY_LIGHTING_SURFACE_FORWARDSHADING
+
 /** Computes sky diffuse lighting, including precomputed shadowing. */
 void GetSkyLighting(FMaterialPixelParameters MaterialParameters, VTPageTableResult LightmapVTPageTableResult, bool bEvaluateBackface, float3 WorldNormal, LightmapUVType LightmapUV, uint LightmapDataIndex, float3 SkyOcclusionUV3D, out float3 OutDiffuseLighting, out float3 OutSubsurfaceLighting)
 {
