Fix shader compilation warning.

diff --git a/Engine/Shaders/Private/DebugViewModePixelShader.usf b/Engine/Shaders/Private/DebugViewModePixelShader.usf
--- a/Engine/Shaders/Private/DebugViewModePixelShader.usf
+++ b/Engine/Shaders/Private/DebugViewModePixelShader.usf
@@ -565,6 +565,8 @@ void Main(
 			// as they are non-shadowcasting in 99% of cases anyway, and their bounds will still show.
 			// Ideally we'd render them as translucent, or cull only the non-shadowcasting ones.
 			discard;
+			OutColor = 0;
+			OutDebugAux = 0;
 			return;
 		}
 		DebugAux = InstanceId;
