Substrate: Add implementation of the EON rough diffuse model and add it as an alternative mode of Diffuse_Chan by scaling the roughness range

This BRDF simulates a similar microfacet structure, and matches the ground truth more closely. It is also fully energy conserving and preserving.

A number of optimizations have been made compared to the published paper to improve performance, but these changes agree very closely with the original.

Also change the argument to Diffuse_Chan to be based on Roughness rather than "a2" to avoid all call-sites needing to know the convention that a2=Pow4(Roughness).

The default remains the current v2 version of the Chan model until we can further investigate the performance of the EON model across platforms.

Refactor the naming of the diffuse function to reflect that it is intended to model a GGX microsurface with lambertian microfacets.

#rnx

diff --git a/Engine/Shaders/Private/BRDF.ush b/Engine/Shaders/Private/BRDF.ush
--- a/Engine/Shaders/Private/BRDF.ush
+++ b/Engine/Shaders/Private/BRDF.ush
@@ -229,19 +229,47 @@ float3 Diffuse_Gotanda( float3 DiffuseColor, float Roughness, float NoV, float N
 #endif
 }
 
-// [ Chan 2018, "Material Advances in Call of Duty: WWII" ]
-// It has been extended here to fade out retro reflectivity contribution from area light in order to avoid visual artefacts.
-#define DIFFUSE_CHAN_VERSION 2
-float3 Diffuse_Chan( float3 DiffuseColor, float a2, float NoV, float NoL, float VoH, float NoH, float RetroReflectivityWeight)
+// [Portsmouth et al. 2025, "EON: A Practical Energy-Preserving Rough Diffuse BRDF"]
+float3 Diffuse_EON( float3 DiffuseColor, float Roughness, float NoV, float NoL, float VoL )
+{
+	// Albedo inversion for EON model to maintain a consistent color with lambert
+	float3 Rho = DiffuseColor * (1.0 + (0.189468 - 0.189468 * DiffuseColor) * Roughness);
+
+	// This is the main shaping term from the Oren-Nayar model (with tweaks by Fujii)
+	float S = VoL - NoV * NoL;
+	float SOverT = max(S * rcp(max(1e-6, max(NoV, NoL))), S);
+	const float constant1_FON = 0.5f - 2.0f / (3.0f * PI);
+	// AF = rcp(1 + Roughness * constant1_FON) is nearly a straight line, so approximate it as such
+	float AF = 1 - Roughness * (1 - 1 / (1 + constant1_FON));
+	float f_ss = AF * (1 + Roughness * SOverT);
+
+	// 4th Order approximation from the paper is a bit too heavy, first order seems to work just as well
+	const float g1 = 0.262048f;
+	float GoverPi_V = g1 - g1 * NoV;
+	// Use (1 - Eo) only as a non-reciprocal approach to energy conservation
+	float f_ms = 1.0f - AF * (1 + Roughness * GoverPi_V);
+	// The Rho_ms term from the paper can be approximated as just Rho^2
+	return Rho * (f_ss + Rho * f_ms) * (1.0 / PI);
+}
+
+#define ROUGH_DIFFUSE_BRDF_VERSION 2
+// This models a rough surface that has a GGX NDF where each microfacet has a lambertian response. Various models have been proposed
+// to try and approximate this behavior.
+float3 Diffuse_GGX_Rough( float3 DiffuseColor, float Roughness, float NoV, float NoL, float VoH, float NoH, float RetroReflectivityWeight )
 {
 	// We saturate each input to avoid out of range negative values which would result in weird darkening at the edge of meshes (resulting from tangent space interpolation).
 	NoV = saturate(NoV);
 	NoL = saturate(NoL);
 	VoH = saturate(VoH);
 	NoH = saturate(NoH);
-#if DIFFUSE_CHAN_VERSION == 2
-	// Updated model [Multiscattering Diffuse and Specular BRDFs, Danny Chan, 2024 - Unpublished]
-	const float Roughness = a2;
+#if ROUGH_DIFFUSE_BRDF_VERSION == 3
+	// It turns out the EON model in the range [0, 0.4] is nearly a perfect match to a ground truth
+	// simulation of diffuse microfacets oriented with a GGX NDF.
+	float VoL = 2 * VoH * VoH - 1;		// double angle identity to keep signature above consistent with other models
+	return Diffuse_EON(DiffuseColor, RetroReflectivityWeight * Roughness * 0.4, NoV, NoL, VoL);
+#elif ROUGH_DIFFUSE_BRDF_VERSION == 2
+	// [ Chan 2024, "Multiscattering Diffuse and Specular BRDFs", Unpublished manuscript ]
+	Roughness *= RetroReflectivityWeight;
 	const float Alpha = Roughness * Roughness;
 	// The original writeup uses an FSmooth term inspired by Burley diffuse to balance energy between spec/diffuse.
 	// However in our implementation the energy balance between diffuse and spec is handled externally, so we stick
@@ -250,10 +278,13 @@ float3 Diffuse_Chan( float3 DiffuseColor, float a2, float NoV, float NoL, float
 	const float Scale = max(0.55 - 0.2 * Roughness, 1.25 - 1.6 * Roughness);
 	const float Bias = saturate(4 * Alpha);
 	const float FRough = Scale * (NoH + Bias) * rcp(NoH + 0.025) * VoH * VoH;
-	const float DiffuseSS = lerp(FSmooth, lerp(FSmooth, FRough, RetroReflectivityWeight), Roughness);
+	const float DiffuseSS = lerp(FSmooth, FRough, Roughness);
 	const float DiffuseMS = Alpha * 0.38;
 	return (1 / PI) * DiffuseColor * (DiffuseSS + DiffuseMS);
 #else
+	// [ Chan 2018, "Material Advances in Call of Duty: WWII" ]
+	// It has been extended here to fade out retro reflectivity contribution from area light in order to avoid visual artefacts.
+	float a2 = Pow4(Roughness);
 	// a2 = 2 / ( 1 + exp2( 18 * g )
 	float g = saturate( (1.0 / 18.0) * log2( 2 * rcpFast(a2) - 1 ) );
 
diff --git a/Engine/Shaders/Private/PathTracing/Material/PathTracingMaterialCommon.ush b/Engine/Shaders/Private/PathTracing/Material/PathTracingMaterialCommon.ush
--- a/Engine/Shaders/Private/PathTracing/Material/PathTracingMaterialCommon.ush
+++ b/Engine/Shaders/Private/PathTracing/Material/PathTracingMaterialCommon.ush
@@ -192,8 +192,8 @@ float ShadowTerminatorTerm(float3 L, float3 N, float3 Ns)
 float3 GetPathTracingDiffuseModel(float3 DiffuseColor, float Roughness, float NoV, float NoL, float VoH, float NoH)
 {
 #if MATERIAL_ROUGHDIFFUSE
-	// Clamp to [0,1] as the model is ever so slightly non-energy conserving at grazing angles which can cause unwanted noise
-	return saturate(PI * Diffuse_Chan(DiffuseColor, Pow4(Roughness), NoV, NoL, VoH, NoH, 1.0f));
+	// Cancel out the factor of PI because the path tracer wants the sample weight
+	return PI * Diffuse_GGX_Rough(DiffuseColor, Roughness, NoV, NoL, VoH, NoH, 1.0f);
 #else
 	return DiffuseColor;
 #endif
diff --git a/Engine/Shaders/Private/ShadingEnergyConservationTable.usf b/Engine/Shaders/Private/ShadingEnergyConservationTable.usf
--- a/Engine/Shaders/Private/ShadingEnergyConservationTable.usf
+++ b/Engine/Shaders/Private/ShadingEnergyConservationTable.usf
@@ -116,7 +116,7 @@ void BuildEnergyTableCS(uint3 DispatchThreadId : SV_DispatchThreadID)
 		const float3 H = normalize(V + L);
 		const float VoH = dot(V, H);
 		const float NoH = H.z;
-		const float DiffuseWeight = NoL > 0 && NoV > 0 ? PI * Diffuse_Chan(float3(1, 1, 1), Alpha.x, NoV, NoL, VoH, NoH, 1.0f).x : 0.0f;
+		const float DiffuseWeight = NoL > 0 && NoV > 0 ? PI * Diffuse_GGX_Rough(float3(1, 1, 1), Roughness, NoV, NoL, VoH, NoH, 1.0f).x : 0.0f;
 
 		Weight = float2(DiffuseWeight, 1.f);
 #endif
diff --git a/Engine/Shaders/Private/ShadingModels.ush b/Engine/Shaders/Private/ShadingModels.ush
--- a/Engine/Shaders/Private/ShadingModels.ush
+++ b/Engine/Shaders/Private/ShadingModels.ush
@@ -249,9 +249,10 @@ FDirectLighting DefaultLitBxDF( FGBufferData GBuffer, half3 N, half3 V, FAreaLig
 		Context.NoV = saturate(abs( Context.NoV ) + 1e-5);
 
 #if MATERIAL_ROUGHDIFFUSE
-		// Chan diffuse model with roughness == specular roughness. This is not necessarily a good modelisation of reality because when the mean free path is super small, the diffuse can in fact looks rougher. But this is a start.
-		// Also we cannot use the morphed context maximising NoH as this is causing visual artefact when interpolating rough/smooth diffuse response. 
-		Lighting.Diffuse = Diffuse_Chan(GBuffer.DiffuseColor, Pow4(GBuffer.Roughness), NoV, AreaLight.NoL, VoH, NoH, GetAreaLightDiffuseMicroReflWeight(AreaLight));
+		// Diffuse model with roughness == specular roughness. This is not necessarily a good model of reality because it assumes microfacets with lambertian response.
+		// The limit behavior of SSS as the mean free path gets small behaves differently, but this is a start.
+		// Also we cannot use the morphed context maximising NoH as this is causing visual artefact when interpolating rough/smooth diffuse response. (SUBSTRATE_TODO: Check this again with latest rough diffuse model)
+		Lighting.Diffuse = Diffuse_GGX_Rough(GBuffer.DiffuseColor, GBuffer.Roughness, NoV, AreaLight.NoL, VoH, NoH, GetAreaLightDiffuseMicroReflWeight(AreaLight));
 #else
 		Lighting.Diffuse = Diffuse_Lambert(GBuffer.DiffuseColor);
 #endif
@@ -624,7 +625,7 @@ FDirectLighting SubsurfaceProfileBxDF( FGBufferData GBuffer, half3 N, half3 V, F
 
 #if MATERIAL_ROUGHDIFFUSE
 	// Use Chan's diffuse model. It reduces flatness look and has better match, e.g., at the edge of human face skin when compared to GT.
-	const float3 DiffuseReflection = Diffuse_Chan(GBuffer.DiffuseColor, Pow4(GBuffer.Roughness), Context.NoV, AreaLight.NoL, Context.VoH, Context.NoH, GetAreaLightDiffuseMicroReflWeight(AreaLight));
+	const float3 DiffuseReflection = Diffuse_GGX_Rough(GBuffer.DiffuseColor, GBuffer.Roughness, Context.NoV, AreaLight.NoL, Context.VoH, Context.NoH, GetAreaLightDiffuseMicroReflWeight(AreaLight));
 #else
 	const float3 DiffuseReflection = Diffuse_Burley(GBuffer.DiffuseColor, GBuffer.Roughness, Context.NoV, AreaLight.NoL, Context.VoH);
 #endif
diff --git a/Engine/Shaders/Private/Substrate/SubstrateEvaluation.ush b/Engine/Shaders/Private/Substrate/SubstrateEvaluation.ush
--- a/Engine/Shaders/Private/Substrate/SubstrateEvaluation.ush
+++ b/Engine/Shaders/Private/Substrate/SubstrateEvaluation.ush
@@ -521,10 +521,11 @@ FSubstrateEvaluateResult SubstrateEvaluateBSDFCommon(FSubstrateBSDFContext BSDFC
 			if (Settings.bRoughDiffuseEnabled && any(DiffuseColor > 0))
 			{
 				// * If the specular layer is anisotropic, the diffuse takes the 'main' roughness value rather than the tangent/bitangent value
-				// * The Chan model bakes transmittance specular directional albedo assuming F=0.04. In previous code we reapplied 
+				// * The Chan diffuse model bakes transmittance specular directional albedo assuming F=0.04. In previous code we reapplied 
 				//   this transmittance, as the energy preservation code we remove it as well. However the visual effect was small 
-				//   and the cost was rather large (~8%). This is why we removed it in recent iteration
-				Sample.DiffusePathValue = Diffuse_Chan(DiffuseColor, Alpha2Spec, NoV, AreaLight.NoL, VoH, NoH, GetAreaLightDiffuseMicroReflWeight(AreaLight));
+				//   and the cost was rather large (~8%). See the code below for where this is handled, this is no longer necessary
+				//   with the latest model.
+				Sample.DiffusePathValue = Diffuse_GGX_Rough(DiffuseColor, SafeRoughness, NoV, AreaLight.NoL, VoH, NoH, GetAreaLightDiffuseMicroReflWeight(AreaLight));
 			}
 			else
 #endif
@@ -800,7 +801,7 @@ FSubstrateEvaluateResult SubstrateEvaluateBSDFCommon(FSubstrateBSDFContext BSDFC
 
 			// Specular to Diffuse energy preservation
 			{
-				#if DIFFUSE_CHAN_VERSION == 1 && MATERIAL_ROUGHDIFFUSE && SUBSTRATE_SHADING_QUALITY == 1	
+				#if ROUGH_DIFFUSE_BRDF_VERSION == 1 && MATERIAL_ROUGHDIFFUSE && SUBSTRATE_SHADING_QUALITY == 1	
 				// We remove the Chan model baked transmittance specular directional albedo assuming F=0.04.
 				// Then we apply the transmittance according to the current surface setup.
 				// To make this a safe operation:
