Fix typo in shader causing compilation error.

diff --git a/Engine/Shaders/Private/Substrate/SubstrateExport.ush b/Engine/Shaders/Private/Substrate/SubstrateExport.ush
--- a/Engine/Shaders/Private/Substrate/SubstrateExport.ush
+++ b/Engine/Shaders/Private/Substrate/SubstrateExport.ush
@@ -915,7 +915,7 @@ void WriteHair(
 	bool bWriteHairData,
 	inout FSubstrateTopLayerData TopLayerData,
 	inout FRWSubstrateMaterialContainer SubstrateBuffer,
-	inout RWTexture2DArray<uint> ExtraMaterialDataUAV)
+	RWTexture2DArray<uint> ExtraMaterialDataUAV)
 {
 	// Top layer data
 	{
