Revert change 42922634

diff --git a/Engine/Shaders/Private/ShadingModels.ush b/Engine/Shaders/Private/ShadingModels.ush
--- a/Engine/Shaders/Private/ShadingModels.ush
+++ b/Engine/Shaders/Private/ShadingModels.ush
@@ -202,14 +202,6 @@ FDirectLighting DefaultLitBxDF( FGBufferData GBuffer, half3 N, half3 V, FAreaLig
 	Lighting.Specular = 0;
 	Lighting.Transmission = 0;
 
-	
-#if FEATURE_LEVEL == FEATURE_LEVEL_ES3_1
-	// adjust N for Adreno shader compilation bug. See UE-274816
-	N = N * 0.5f;
-	N = normalize(N);
-#endif
-	
-	
 	BRANCH
 	if (AreaLight.NoL > 0.0f)
 	{
