Fix TextureGraph's Opacity export with Substrate enaled and using Format=1.

diff --git a/Engine/Plugins/TextureGraph/Shaders/TextureGraphMaterialShader.usf b/Engine/Plugins/TextureGraph/Shaders/TextureGraphMaterialShader.usf
--- a/Engine/Plugins/TextureGraph/Shaders/TextureGraphMaterialShader.usf
+++ b/Engine/Plugins/TextureGraph/Shaders/TextureGraphMaterialShader.usf
@@ -104,7 +104,7 @@ void TextureGraphMaterialShaderPS(
 	MatRoughness = Export.Roughness;
 	MatEmissive = Export.EmissiveLuminance;
 	MatNormal = Export.WorldNormal * 0.5 + 0.5;
-	MatOpacity = Export.Coverage;
+	MatOpacity = Export.Opacity;
 
 #if MATERIAL_USES_ANISOTROPY
 	MatTangent = Export.WorldTangent * 0.5 + 0.5;
diff --git a/Engine/Shaders/Private/Substrate/SubstrateExport.ush b/Engine/Shaders/Private/Substrate/SubstrateExport.ush
--- a/Engine/Shaders/Private/Substrate/SubstrateExport.ush
+++ b/Engine/Shaders/Private/Substrate/SubstrateExport.ush
@@ -1013,23 +1013,27 @@ FExportResult SubstrateMaterialExportOut(
 	const int BSDFIdx = 0;
 	#define CurrentBSDF SubstratePixelHeader.SubstrateTree.BSDFs[BSDFIdx]
 
-#if SUBTRATE_OPTIMISED_LEGACY_OPAQUE_EXPORT
+	#if !SUBTRATE_OPTIMISED_LEGACY_OPAQUE_EXPORT
+	if (SubstratePixelHeader.SubstrateTree.BSDFCount > 0)
+	#endif
 	{
+		#if SUBTRATE_OPTIMISED_LEGACY_OPAQUE_EXPORT
 		// In this case, there is always a BSDF and it is always 100% visible.
 		// Sanitize BSDF before it is exported.
 		CurrentBSDF.SubstrateSanitizeBSDF();
 		CurrentBSDF.LuminanceWeightV = 1.0f; 
-		{
-#else
-	if (SubstratePixelHeader.SubstrateTree.BSDFCount > 0)
-	{
 
+		Out.Coverage = 1.f;
+		Out.TransmittancePreCoverage = 1.f;
+		#else
 		// Update tree (coverage/transmittance/luminace weights) and Sanitize its content.
 		SubstratePixelHeader.SubstrateUpdateTree(SubstrateData, V, Settings, Out.Coverage, Out.TransmittancePreCoverage);
+		#endif
 
+		#if !SUBTRATE_OPTIMISED_LEGACY_OPAQUE_EXPORT
 		if (SubstrateIsBSDFVisible(CurrentBSDF))
+		#endif
 		{
-#endif
 			const uint BSDFType = BSDF_GETTYPE(CurrentBSDF);
 			const uint SSSType  = BSDF_GETSSSTYPE(CurrentBSDF);
 
