Path Tracer: Fix simple volume substrate slabs getting their MFP rescaled even when they are at the bottom of a solid material.

#rnx

diff --git a/Engine/Shaders/Private/Substrate/Substrate.ush b/Engine/Shaders/Private/Substrate/Substrate.ush
--- a/Engine/Shaders/Private/Substrate/Substrate.ush
+++ b/Engine/Shaders/Private/Substrate/Substrate.ush
@@ -3465,9 +3465,19 @@ void FSubstrateBSDF::PostProcessSlabBeforeLighting(bool bIsSubstrateOpaqueMateri
 	// * Rescale MFP to computation space
 	// * Set thickness to normalized simple volume thickness
 	// * Set SSS type to SimpleVolume
-	if (BSDF_GETHASMFP(this) && (!this.bIsBottom || BSDF_GETSSSTYPE(this) == SSS_TYPE_SIMPLEVOLUME))
+	if (BSDF_GETHASMFP(this))
 	{
-		this.EnableSlabBSDFSimpleVolumetric();
+#if PATH_TRACING
+		// In the case of Path Tracing, we implement the support of MFP on the bottom most slab of solid surfaces with a real random walk
+		// against the geometry, not the implied thickness layer.  Therefore in this case we don't want to modify the value provided in the
+		// material graph.
+		if (!this.bIsBottom || (BSDF_GETSSSTYPE(this) == SSS_TYPE_SIMPLEVOLUME && BSDF_GETISTHIN(this)))
+#else
+		if (!this.bIsBottom || BSDF_GETSSSTYPE(this) == SSS_TYPE_SIMPLEVOLUME)
+#endif
+		{
+			this.EnableSlabBSDFSimpleVolumetric();
+		}
 	}
 }
 
@@ -3875,7 +3885,7 @@ float SubstrateGetLegacyShadingModels(in FSubstrateBSDF BSDF)
 				}
 			}
 			return SHADINGMODELID_DEFAULT_LIT;
-		} // Also include SSS_TYPE_DIFFUSION and SSS_TYPE_SIMPLE_VOLUME
+		} // Also include SSS_TYPE_DIFFUSION and SSS_TYPE_SIMPLEVOLUME
 	}
 	case SUBSTRATE_BSDF_TYPE_HAIR: 								{ return SHADINGMODELID_HAIR; }
 	case SUBSTRATE_BSDF_TYPE_EYE: 								{ return SHADINGMODELID_EYE; }
