Fix View.PreExposure applied 2 time when computing the ReflectionEnvironmentSkyLighting and APPLY_SKY_SHADOWING is defined.

#tests odsc with different game modes

diff --git a/Engine/Shaders/Private/SkyLightingDiffuseShared.ush b/Engine/Shaders/Private/SkyLightingDiffuseShared.ush
--- a/Engine/Shaders/Private/SkyLightingDiffuseShared.ush
+++ b/Engine/Shaders/Private/SkyLightingDiffuseShared.ush
@@ -131,7 +131,10 @@ float3 SkyLightDiffuse(FGBufferData GBuffer, float AmbientOcclusion, float2 Buff
 
 	// And accumulate the lighting
 	Lighting += (SkyVisData.SkyDiffuseLookUpMul * DiffuseLookup + SkyVisData.SkyDiffuseLookUpAdd) * DiffuseColor * DiffuseWeight;
+
+#if FEATURE_LEVEL > FEATURE_LEVEL_ES3_1
 	Lighting *= View.PreExposure;
+#endif
 
 	//Lighting = (Texture2DSampleLevel(BentNormalAOTexture, BentNormalAOSampler, UV, 0).xyz);
 
