Fix incorrect Lumen's specular reflection fetching for pixel with SSS diffuse with Substrate.

Backscattering flag should only be set for SSS Wrap or SSS TwoSided Wrap in order to match legacy behavior.
This should also improve Lumen's reflection performance with Substrate Format=1 on SSS diffusion surface.

diff --git a/Engine/Shaders/Private/Substrate/Substrate.ush b/Engine/Shaders/Private/Substrate/Substrate.ush
--- a/Engine/Shaders/Private/Substrate/Substrate.ush
+++ b/Engine/Shaders/Private/Substrate/Substrate.ush
@@ -3502,7 +3502,7 @@ bool FSubstrateBSDF::HasScatteringData()
 bool FSubstrateBSDF::HasBackScattering()
 {
 	const uint SSSType = BSDF_GETSSSTYPE(this);
-	return SSSType != SSS_TYPE_NONE && SSSType != SSS_TYPE_SIMPLEVOLUME;
+	return SSSType == SSS_TYPE_WRAP || SSSType == SSS_TYPE_TWO_SIDED_WRAP;
 }
 
 bool FSubstrateBSDF::HasSSSDiffusion()
