Fix incorrect GBuffer value initiallization from hair strands samples causing incorrect MegaLight shading on hair strands.

diff --git a/Engine/Shaders/Private/HairStrands/HairStrandsVisibilityUtils.ush b/Engine/Shaders/Private/HairStrands/HairStrandsVisibilityUtils.ush
--- a/Engine/Shaders/Private/HairStrands/HairStrandsVisibilityUtils.ush
+++ b/Engine/Shaders/Private/HairStrands/HairStrandsVisibilityUtils.ush
@@ -16,8 +16,8 @@ FGBufferData HairSampleToGBufferData(in FHairSample In, in float InDualScatterin
 	Out.GBufferAO = 1;
 	Out.CustomData = float4(InDualScatteringRoughnessOverride, 0, In.Backlit, 0);
 	Out.IndirectIrradiance = 1000000;
-	Out.DiffuseColor = 0;
-	Out.SpecularColor = 0;
+	Out.DiffuseColor = In.BaseColor;
+	Out.SpecularColor = DielectricSpecularToF0(In.Specular);
 	Out.PrecomputedShadowFactors = 1;
 	Out.PerObjectGBufferData = HasHairFlags(In.Flags, HAIR_FLAGS_CAST_CONTACT_SHADOW) ? 1u : 0u;
 	return Out;
