Optimized ScreenSpaceShadowRayCast for Substrate path.
Enable this for non-substrate later.

diff --git a/Engine/Shaders/Private/ScreenSpaceShadowRayCast.ush b/Engine/Shaders/Private/ScreenSpaceShadowRayCast.ush
--- a/Engine/Shaders/Private/ScreenSpaceShadowRayCast.ush
+++ b/Engine/Shaders/Private/ScreenSpaceShadowRayCast.ush
@@ -28,10 +28,44 @@ float CastScreenSpaceShadowRay(
 
 	const float CompareTolerance = abs(RayDepth.z - RayStart.z) * Step * CompareToleranceScale;
 
-	float SampleTime = StepOffset * Step + Step;
-
 	const float StartDepth = LookupDeviceZ(RayStart.xy * View.ScreenPositionScaleBias.xy + View.ScreenPositionScaleBias.wz);
 
+#if SUBSTRATE_ENABLED // TODO: Reconsider this for non-Substrate
+	const float FirstSampleTime = StepOffset * Step + Step;
+
+	const float3 SamplePos = RayStart + RayStep * FirstSampleTime;
+
+	float2 SampleUV = SamplePos.xy * View.ScreenPositionScaleBias.xy + View.ScreenPositionScaleBias.wz;
+	float SampleZ = SamplePos.z + CompareTolerance;	// Apply CompareTolerance outside of loop
+	float Result = RayLength * FirstSampleTime;
+
+	const float2 SampleUVStep = (RayStep.xy * Step) * View.ScreenPositionScaleBias.xy;
+	const float SampleZStep = RayStep.z * Step;	
+	const float ResultStep = RayLength * Step;
+
+	UNROLL
+	for (int i = 0; i < NumSteps; i++)
+	{
+		const float SampleDepth = LookupDeviceZ(SampleUV);
+
+		// Avoid self-intersection with the start pixel (exact comparison due to point sampling depth buffer)
+		// Exception is made for hair for occluding transmitted light with non-shadow casting light
+		if (abs(SampleZ - SampleDepth) < CompareTolerance && (SampleDepth != StartDepth || bHairNoShadowLight))
+		{
+			HitUV = SampleUV;
+
+			// Off screen masking
+			bool bValidPos = all( abs(SampleUV - View.ScreenPositionScaleBias.wz) < abs(View.ScreenPositionScaleBias.xy) );
+			return bValidPos ? Result : -1.0;
+		}
+
+		SampleUV += SampleUVStep;
+		SampleZ += SampleZStep;
+		Result += ResultStep;
+	}
+#else
+	float SampleTime = StepOffset * Step + Step;
+
 	UNROLL
 	for (int i = 0; i < NumSteps; i++)
 	{
@@ -60,6 +94,7 @@ float CastScreenSpaceShadowRay(
 
 		SampleTime += Step;
 	}
+#endif
 
 	return -1;
 }
\ No newline at end of file
