Path Tracer: Implement support for scattering inside refractive media

This allows us to better support the SLW shading model as well as Substrate slabs. For consistency with opaque behavior, we blend towards SSS as the mean free path gets smaller (which blends towards a plain diffuse or sss response as the mean free path gets close to 0).

Change the rule for volumetric tracking so that volumes (atmosphere, fog and heterogeenous volumes are ignored when inside a medium). These rules may be relaxed in the future, but simplify the code for now.

diff --git a/Engine/Shaders/Private/PathTracing/Material/PathTracingMaterialSampling.ush b/Engine/Shaders/Private/PathTracing/Material/PathTracingMaterialSampling.ush
--- a/Engine/Shaders/Private/PathTracing/Material/PathTracingMaterialSampling.ush
+++ b/Engine/Shaders/Private/PathTracing/Material/PathTracingMaterialSampling.ush
@@ -309,6 +309,8 @@ float GetAverageRoughness(FPathTracingPayload Payload)
 			AvgRoughness = lerp(1.0, AvgRoughness, DiffuseFraction);
 			break;
 		}
+		case SHADINGMODELID_MEDIUM:
+			return 1;
 		// SUBSTRATE_TODO: do we need special cases for the other slab types?
 	}
 	return AvgRoughness;
@@ -328,6 +330,8 @@ float GetAverageRoughness(FPathTracingPayload Payload)
 		case SHADINGMODELID_THIN_TRANSLUCENT:
 		case SHADINGMODELID_SOLID_GLASS:
 			return Roughness; // both cases are purely specular, no diffuse
+		case SHADINGMODELID_MEDIUM:
+			return 1;
 	}
 	float AvgRoughness = lerp(1.0 /* diffuse lobe */, Roughness, Payload.Metallic);
 	return AvgRoughness;
diff --git a/Engine/Shaders/Private/PathTracing/PathTracing.usf b/Engine/Shaders/Private/PathTracing/PathTracing.usf
--- a/Engine/Shaders/Private/PathTracing/PathTracing.usf
+++ b/Engine/Shaders/Private/PathTracing/PathTracing.usf
@@ -61,6 +61,7 @@ FPathState LoadPathStateData(uint Index)
 	Output.Ray.TMax = RAY_DEFAULT_T_MAX;
 	Output.PathThroughput = Data.PathThroughput;
 	Output.PackedRoughnessSigma = Data.PackedRoughnessSigma;
+	Output.PackedScatterPhase = Data.PackedScatterPhase;
 	return Output;
 }
 
@@ -77,6 +78,7 @@ void StorePathStateData(FPathState PathState, uint Index)
 	Data.RayDirection = PathState.Ray.Direction;
 	Data.PathThroughput = PathState.PathThroughput;
 	Data.PackedRoughnessSigma = PathState.PackedRoughnessSigma;
+	Data.PackedScatterPhase = PathState.PackedScatterPhase;
 	PathStateData[Index] = Data;
 }
 #endif
diff --git a/Engine/Shaders/Private/PathTracing/PathTracingCommon.ush b/Engine/Shaders/Private/PathTracing/PathTracingCommon.ush
--- a/Engine/Shaders/Private/PathTracing/PathTracingCommon.ush
+++ b/Engine/Shaders/Private/PathTracing/PathTracingCommon.ush
@@ -819,6 +819,11 @@ struct FPathTracingPayload : FMinimalPayload
 
 	// SHADINGMODELID_SOLID_GLASS
 	float3 GetExtinction() { return rcp(max(MeanFreePath, 1e-5)); }
+	float3 GetGlassAlbedo() { return DiffuseColor; }
+	void SetGlassAlbedo(float3 A) { DiffuseColor = A; }
+	float GetGlassPhase() { return PhaseG; }
+	void SetGlassPhase(float G) { PhaseG = G; }
+
 
 	// SHADINGMODELID_THIN_TRANSLUCENT
 	float3 GetTransmittanceColor() { return exp(-GetExtinction() * SUBSTRATE_SIMPLEVOLUME_THICKNESS_CM); }
@@ -888,6 +893,11 @@ struct FPathTracingPayload : FMinimalPayload
 	// SHADINGMODELID_SOLID_GLASS
 	float3 GetExtinction() { return CustomData0.xyz; }
 	void SetExtinction(float3 SigmaT) { CustomData0.xyz = SigmaT; }
+	float3 GetGlassAlbedo() { return BaseColor; }
+	void SetGlassAlbedo(float3 A) { BaseColor = A; }
+	float GetGlassPhase() { return CustomData0.w; }
+	void SetGlassPhase(float G) { CustomData0.w = G; }
+
 
 	// SHADINGMODELID_NUM (Lambert diffuse) and SHADINGMODELID_MEDIUM
 	float3 GetBaseColor() { return BaseColor; }
diff --git a/Engine/Shaders/Private/PathTracing/PathTracingCore.ush b/Engine/Shaders/Private/PathTracing/PathTracingCore.ush
--- a/Engine/Shaders/Private/PathTracing/PathTracingCore.ush
+++ b/Engine/Shaders/Private/PathTracing/PathTracingCore.ush
@@ -101,21 +101,23 @@ RaytracingAccelerationStructure TLAS;
 RaytracingAccelerationStructure DecalTLAS;
 uint SceneVisibleLightCount;
 
-Buffer<float> StartingExtinctionCoefficient;
+Buffer<float> StartingMediumData;
 
 struct FPathState
 {                                   // packed size
 	FRandomSequence RandSequence;   //  8 bytes
 	float3 Radiance;                // 12 bytes
-	float  Alpha;                   //  4 bytes
+	float Alpha;                   //  4 bytes
 	uint3 PackedAlbedoNormal;       //  6 bytes // half precision (Albedo)
-                                    //  6 bytes // half precision (Normal)
+	//  6 bytes // half precision (Normal)
 	FRayDesc Ray;                   // 12+12 = 24 bytes (TMin/TMax do not need to be stored)
 	float3 PathThroughput;          // 12 bytes
 	uint2 PackedRoughnessSigma;     // 2 bytes // half precision (PathRoughness)
 	                                // 6 bytes // half precision  (SigmaT)
 	                                // 3 bits - packed in sign bits of SigmaT (FirstScatterType)
-									// 80 bytes total (see FPathTracingPackedPathState)
+	uint2 PackedScatterPhase;       // 6 bytes // half precision (SigmaS)
+	                                // 2 bytes // phase g (could be quantized further)
+	                                // 80 bytes total (see FPathTracingPackedPathState)
 
 	// Temporary parameters not stored				
 	float DepthZ;                   // Camera ray depth
@@ -123,23 +125,21 @@ struct FPathState
 
 	float3 GetAlbedo()
 	{
-		return float3(f16tof32(PackedAlbedoNormal.x),
-		              f16tof32(PackedAlbedoNormal.x >> 16),
-		              f16tof32(PackedAlbedoNormal.y));
+		return float3(UnpackFloat2FromUInt(PackedAlbedoNormal.x),
+		              UnpackFloat2FromUInt(PackedAlbedoNormal.y).x);
 	}
 
 	float3 GetNormal()
 	{
-		return float3(f16tof32(PackedAlbedoNormal.y >> 16),
-		              f16tof32(PackedAlbedoNormal.z),
-		              f16tof32(PackedAlbedoNormal.z >> 16));
+		return float3(UnpackFloat2FromUInt(PackedAlbedoNormal.y).y,
+		              UnpackFloat2FromUInt(PackedAlbedoNormal.z));
 	}
 
 	void SetAlbedoNormal(float3 Albedo, float3 Normal)
 	{
-		PackedAlbedoNormal.x = f32tof16(Albedo.x) | (f32tof16(Albedo.y) << 16);
-		PackedAlbedoNormal.y = f32tof16(Albedo.z) | (f32tof16(Normal.x) << 16);
-		PackedAlbedoNormal.z = f32tof16(Normal.y) | (f32tof16(Normal.z) << 16);
+		PackedAlbedoNormal.x = PackFloat2ToUInt(Albedo.xy);
+		PackedAlbedoNormal.y = PackFloat2ToUInt(float2(Albedo.z, Normal.x));
+		PackedAlbedoNormal.z = PackFloat2ToUInt(Normal.yz);
 	}
 
 	uint GetFirstScatterType()
@@ -195,6 +195,17 @@ struct FPathState
 		PackedRoughnessSigma.y |= f32tof16(SigmaT.z) << 16;
 	}
 
+	float4 GetSigmaSAndPhase()
+	{
+		return float4(UnpackFloat2FromUInt(PackedScatterPhase.x), UnpackFloat2FromUInt(PackedScatterPhase.y));
+	}
+
+	void SetSigmaSAndPhase(float3 SigmaS, float PhaseG)
+	{
+		PackedScatterPhase.x = PackFloat2ToUInt(SigmaS.xy);
+		PackedScatterPhase.y = PackFloat2ToUInt(float2(SigmaS.z, PhaseG));
+	}
+
 	bool HasMadeContributionScatter()
 	{
 		// The first scattering type that matters for the PATHTRACER_CONTRIBUTION_* flags is the first non-refractive one
@@ -697,6 +708,56 @@ FPackedPathTracingPayload TraceTransparentRay(inout FPathState PathState, int Bo
 			}
 		}
 
+		if (any(PathState.GetSigmaT() > 0))
+		{
+			float3 SigmaT = PathState.GetSigmaT();
+			float4 ScatterData = PathState.GetSigmaSAndPhase();
+			float3 MediumAlbedo = select(SigmaT > 0, ScatterData.xyz / SigmaT, 0.0);
+			float3 ColorChannelPdf = PathState.PathThroughput * MediumAlbedo;
+			float ScatterT = !bLastBounce ? SampleSpectralTransmittance(PathState.RandSequence.Get1D(), SigmaT, ColorChannelPdf) : -1;
+			float HitT = PackedPayload.IsHit() ? PackedPayload.HitT : PathState.Ray.TMax;
+			if (ScatterT < 0)
+			{
+				// scattering not present, just apply Beer's law through the current medium
+				PathState.PathThroughput *= exp(-SigmaT * (HitT - PathState.Ray.TMin));
+			}
+			else if (PathState.Ray.TMin + ScatterT < HitT)
+			{
+				float3 PathThroughput = PathState.PathThroughput * EvaluateSpectralTransmittanceHit(ScatterT, SigmaT, ColorChannelPdf).xyz;
+				ScatterT += PathState.Ray.TMin;
+				float3 Contrib = PathThroughput * ScatterData.xyz;
+
+				// accumulate what the denoiser wants into albedo/normal (as long as the current path is rough enough)
+				float3 Albedo = PathState.GetAlbedo();
+				AccumulateAlbedo(ScatterData.xyz, PathThroughput, PathState.GetPathRoughness(), Albedo);
+				PathState.SetAlbedoNormal(Albedo, PathState.GetNormal());
+
+				float SelectionWeight = max3(Contrib.x, Contrib.y, Contrib.z);
+				if (HitSample.Accept(SelectionWeight))
+				{
+					FVolumeShadedResult Result = (FVolumeShadedResult) 0;
+					Result.SigmaT = SigmaT;
+					Result.SigmaSHG = ScatterData.xyz;
+					Result.PhaseG = ScatterData.w;
+					// stash hit
+					Payload = PackPathTracingPayload(CreateMediumHitPayload(ScatterT, PathState.Ray.Origin + ScatterT * PathState.Ray.Direction, Result));
+					PayloadThroughput = Contrib / SelectionWeight;
+				}
+				break; // no more hits left
+			}
+			else
+			{
+				// no hit in scattering medium, continue to next surface hit
+				PathState.PathThroughput *= EvaluateSpectralTransmittanceMiss(HitT - PathState.Ray.TMin, SigmaT, ColorChannelPdf).xyz;
+			}
+
+			// Pretend like we didn't hit any volume to exlude them from the internal medium we are going through
+			// It is fairly unlikely (and probably undesireable) to want atmo/cloud/fog to be rendererd inside glass/water.
+			// We may eventually want to support heterogeneous volumes in these cases, but this will require unifying the ray marching of the ambient medium with the current one
+			VolumeIsectList.VolumeMask = 0;
+		}
+
+
 		// Volume Transmittance + Scatter Segment/Point selection
 		// This loop splits the potential volumetric segments into sub-intervals that do not overlap.
 		while (VolumeIsectList.HitVolume())
@@ -953,13 +1014,6 @@ FPackedPathTracingPayload TraceTransparentRay(inout FPathState PathState, int Bo
 			ApplyDecal(DecalPayload, PackedPayload, PathState.Ray.Direction, PathState.GetPathRoughness());
 		}
 #endif
-		{
-			// account for Beer's law through the currently active medium
-			// TODO: merge this with the volume raymarching?
-			// NOTE: select statement is to avoid corner case of 0.0*inf that can occur with directly visible skydomes
-			float3 SigmaT = PathState.GetSigmaT();
-			PathState.PathThroughput *= select(SigmaT > 0.0, exp(-SigmaT * (PackedPayload.HitT - PathState.Ray.TMin)), 1.0);
-		}
 
 		if (!PackedPayload.IsHoldout() && bIsCameraRay)
 		{
@@ -1697,7 +1751,8 @@ FPathState CreatePathState(int2 PixelIndex, int2 TextureIndex)
 	// path state variables (these cary information between bounces)
 	Output.PathThroughput = 1.0;
 	Output.SetPathRoughness(0.0);
-	Output.SetSigmaT(float3(StartingExtinctionCoefficient[0], StartingExtinctionCoefficient[1], StartingExtinctionCoefficient[2]));
+	Output.SetSigmaT(float3(StartingMediumData[0], StartingMediumData[1], StartingMediumData[2]));
+	Output.SetSigmaSAndPhase(float3(StartingMediumData[3], StartingMediumData[4], StartingMediumData[5]), StartingMediumData[6]);
 	Output.SetFirstScatterType(PATHTRACER_SCATTER_CAMERA);
 
 	#if DEBUG_ENABLE
@@ -1864,11 +1919,6 @@ bool PathTracingKernel(inout FPathState PathState, int Bounce)
 	FPathTracingPayload Payload = UnpackPathTracingPayload(PackedPayload, PathState.Ray);
 	AdjustPayloadAfterUnpack(Payload, ApplyDiffuseSpecularOverrides);
 
-#if 0
-	PathState.Radiance = 0.5 * Payload.WorldGeoNormal + 0.5;
-	return false;
-#endif
-
 	// As soon as the path is blurry enough, we can get away with diffuse sampling only
 	const bool bSimplifySSS = PathState.GetPathRoughness() >= 0.15;
 
@@ -1879,9 +1929,11 @@ bool PathTracingKernel(inout FPathState PathState, int Bounce)
 		return false;
 	}
 
-	const bool bIsVolumeSample = Payload.ShadingModelID == SHADINGMODELID_MEDIUM;
+	// Don't treat internal volume scattering as a volume sample (so its not affected by those options in the lights)
+	const bool bIsVolumeSample = Payload.ShadingModelID == SHADINGMODELID_MEDIUM && !any(PathState.GetSigmaT() > 0);
 
 	FLightLoopCount LightLoopCount = LightGridLookup(Payload.TranslatedWorldPos);
+	// If this is a volume sample, we have not forced the renderer to ignore volume segment sampling _and_ we aren't in an internal medium which does not use volume segments
 	if (bIsVolumeSample && VolumeMISMode != 0)
 	{
 		// if we are using the volume segment for local lights, exclude them from the light loop below
@@ -2070,19 +2122,29 @@ bool PathTracingKernel(inout FPathState PathState, int Bounce)
 	if (MaterialSample.PositionBiasSign < 0 && Payload.IsMaterialSolidGlass())
 	{
 		const float3 LocalSigmaT = Payload.GetExtinction();
+		const float3 LocalSigmaS = Payload.GetGlassAlbedo() * LocalSigmaT;
+		const float LocalPhase = Payload.GetGlassPhase();
 		float3 SigmaT = PathState.GetSigmaT();
+		float4 ScatterData = PathState.GetSigmaSAndPhase();
+		float WeightPrev = max(LobeColorToWeight(SigmaT), 0);
+		float WeightCurr = max(LobeColorToWeight(LocalSigmaT), 0);
 		if (Payload.IsFrontFace())
 		{
 			// entering
 			SigmaT += LocalSigmaT;
+			ScatterData.xyz += LocalSigmaS;
 		}
 		else
 		{
 			// exiting
 			SigmaT -= LocalSigmaT;
 			SigmaT = max(SigmaT, 0);
+			ScatterData.xyz -= LocalSigmaS;
 		}
+		ScatterData.xyz = clamp(ScatterData.xyz, 0, LocalSigmaT);
+		ScatterData.w = lerp(ScatterData.w, LocalPhase, MISWeightBalanced(WeightCurr, WeightPrev));
 		PathState.SetSigmaT(SigmaT);
+		PathState.SetSigmaSAndPhase(ScatterData.xyz, ScatterData.w);
 	}
 
 	// If we are using Material sampling for lights
diff --git a/Engine/Shaders/Private/PathTracing/PathTracingInitExtinctionCoefficient.usf b/Engine/Shaders/Private/PathTracing/PathTracingInitExtinctionCoefficient.usf
--- a/Engine/Shaders/Private/PathTracing/PathTracingInitExtinctionCoefficient.usf
+++ b/Engine/Shaders/Private/PathTracing/PathTracingInitExtinctionCoefficient.usf
@@ -4,13 +4,13 @@
 #include "PathTracingCommon.ush"
 
 RaytracingAccelerationStructure TLAS;
-RWBuffer<float> RWStartingExtinctionCoefficient;
+RWBuffer<float> RWStartingMediumData;
 
 RAY_TRACING_ENTRY_RAYGEN(PathTracingInitExtinctionCoefficientRG)
 {
 	const uint RayFlags = 0;
 	const uint MissShaderIndex = 0;
-	const uint InstanceInclusionMask = PATHTRACER_MASK_CAMERA; // Ignore hair strands, only trace what camera rays can see
+	const uint InstanceInclusionMask = PATHTRACER_MASK_CAMERA | PATHTRACER_MASK_CAMERA_TRANSLUCENT; // Ignore hair strands, only trace what camera rays can see
 
 	FRayDesc UpRay;
 	UpRay.Origin = PrimaryView.TranslatedWorldCameraOrigin;
@@ -19,11 +19,13 @@ RAY_TRACING_ENTRY_RAYGEN(PathTracingInitExtinctionCoefficientRG)
 	UpRay.TMax = RAY_DEFAULT_T_MAX;
 		
 	float3 StartingSigmaT = 0.0;
+	float3 StartingSigmaS = 0.0;
+	float StartingPhaseG = 0;
 	for (;;)
 	{
 		FPackedPathTracingPayload PackedPayload = InitPathTracingPayload(PATHTRACER_SCATTER_CAMERA, 0.0);
-		// SUBSTRATE_TODO: Set the stochastic slab random value? Add special flag to only return extinction of solid glass?
-		// The current behavior will only work correctly for single-slab materials
+		// TODO: Should we try to average the result of several paths here? If a material is only partially glass, we may not get a consistent answer here
+		PackedPayload.SetStochasticSlabRand(0.5);
 		TraceRay(
 			TLAS,
 			RayFlags,
@@ -49,18 +51,33 @@ RAY_TRACING_ENTRY_RAYGEN(PathTracingInitExtinctionCoefficientRG)
 		{
 			// Found a solid transmissive medium -- pickup the extinction
 			float3 LocalSigmaT = HitPayload.GetExtinction();
+			float3 LocalSigmaS = HitPayload.GetGlassAlbedo() * LocalSigmaT;
+			float PhaseG = HitPayload.GetGlassPhase();
+
+			float WeightPrev = max(LobeColorToWeight(StartingSigmaT), 0);
+			float WeightCurr = max(LobeColorToWeight(LocalSigmaT), 0);
 
 			// increase when we leave a medium (that means we must have been inside before)
 			// decrease when we enter a medium (will be cancelled out when we leave)
 			StartingSigmaT += HitPayload.IsFrontFace() ? -LocalSigmaT : +LocalSigmaT;
+			StartingSigmaS += HitPayload.IsFrontFace() ? -LocalSigmaS : +LocalSigmaS;
+
+			StartingPhaseG = lerp(StartingPhaseG, PhaseG, MISWeightBalanced(WeightCurr, WeightPrev));
 		}
 		// keep tracing
 		UpRay.TMin = ComputeNewTMin(UpRay.Origin, UpRay.Direction, PackedPayload.HitT);
 	}
 	// clamp against 0 in the unlikely case that we only hit front faces of glass on our way up
 	StartingSigmaT = max(StartingSigmaT, 0.0);
+	StartingSigmaS = max(StartingSigmaS, 0.0);
+
+	RWStartingMediumData[0] = StartingSigmaT.x;
+	RWStartingMediumData[1] = StartingSigmaT.y;
+	RWStartingMediumData[2] = StartingSigmaT.z;
+
+	RWStartingMediumData[3] = StartingSigmaS.x;
+	RWStartingMediumData[4] = StartingSigmaS.y;
+	RWStartingMediumData[5] = StartingSigmaS.z;
 
-	RWStartingExtinctionCoefficient[0] = StartingSigmaT.x;
-	RWStartingExtinctionCoefficient[1] = StartingSigmaT.y;
-	RWStartingExtinctionCoefficient[2] = StartingSigmaT.z;
+	RWStartingMediumData[6] = StartingPhaseG;
 }
diff --git a/Engine/Shaders/Private/PathTracing/PathTracingMaterialHitShader.usf b/Engine/Shaders/Private/PathTracing/PathTracingMaterialHitShader.usf
--- a/Engine/Shaders/Private/PathTracing/PathTracingMaterialHitShader.usf
+++ b/Engine/Shaders/Private/PathTracing/PathTracingMaterialHitShader.usf
@@ -75,6 +75,8 @@ float AdjustMaterialRoughness(float Roughness, float PathRoughness)
 	return max(Roughness, PathRoughness);
 }
 
+#define PATH_TRACER_USE_SUBSTRATE_GLASS_LOBE    (MATERIALBLENDING_TRANSLUCENT && (SUBSTRATE_BLENDING_TRANSLUCENT_GREYTRANSMITTANCE || SUBSTRATE_BLENDING_TRANSLUCENT_COLOREDTRANSMITTANCE))
+
 #if REFRACTION_USE_INDEX_OF_REFRACTION || REFRACTION_USE_INDEX_OF_REFRACTION_FROM_F0
 
 #define PATH_TRACER_REFRACTION_ENABLED 1
@@ -438,7 +440,7 @@ RAY_TRACING_ENTRY_CLOSEST_HIT(PathTracingMaterialCHS,
 						float LegacyOpacity = 1;
 						if (BSDF_GETSSSTYPE(BSDF) == SSS_TYPE_SIMPLEVOLUME && !BSDF.bIsBottom)
 						{
-							// TODO: This should really be represented by a different BSDF, but for now we are approximating it with a diffuse lobe
+							// SUBSTRATE_TODO: This should really be represented by a different BSDF, but for now we are approximating it with a diffuse lobe
 							DiffuseColor = GetSimpleVolumeDiffuseColor(DiffuseColor, SLAB_SSSMFP(BSDF));
 						}
 						else if (BSDF_GETSSSTYPE(BSDF) == SSS_TYPE_DIFFUSION_PROFILE)
@@ -502,17 +504,22 @@ RAY_TRACING_ENTRY_CLOSEST_HIT(PathTracingMaterialCHS,
 
 						float Ior = 1; // for refracted lobe only
 						float3 Transmittance = 0;
-						float SlabMix = 1; // percentage mix against glass lobe
-						#if MATERIALBLENDING_TRANSLUCENT && (SUBSTRATE_BLENDING_TRANSLUCENT_GREYTRANSMITTANCE || SUBSTRATE_BLENDING_TRANSLUCENT_COLOREDTRANSMITTANCE)
+						float RefractMix = 0; // percentage mix of the glass lobe
+					#if PATH_TRACER_USE_SUBSTRATE_GLASS_LOBE
 						if (BSDF.bIsBottom)
 						{
-							#if PATH_TRACER_REFRACTION_ENABLED
+						#if PATH_TRACER_REFRACTION_ENABLED
 							Ior = GetRefractionIor(PixelMaterialInputs, F0);
-							#endif
+						#endif
 							Transmittance = MeanFreePathToTransmittance(MeanFreePath * CENTIMETER_TO_METER, SUBSTRATE_SIMPLEVOLUME_THICKNESS_CM * CENTIMETER_TO_METER);
-							SlabMix = LobeSelectionProb(DiffuseColor, Transmittance);
-						}
+						#if MATERIAL_ISTHINSURFACE
+							RefractMix = saturate(Luminance(Transmittance));
+						#else
+							// blend towards Opaque SSS in the range of [0cm, 8cm]
+							RefractMix = saturate(max3(MeanFreePath.x, MeanFreePath.y, MeanFreePath.z));
 						#endif
+						}
+					#endif // PATH_TRACER_USE_SUBSTRATE_GLASS_LOBE
 
 						// SUBSTRATE_TODO: Need to account for glass case here for lobe selection ....
 						const float3 FuzzAlbedo = FuzzAmount * FuzzE * FuzzColor;
@@ -523,7 +530,7 @@ RAY_TRACING_ENTRY_CLOSEST_HIT(PathTracingMaterialCHS,
 						// NOTE: Technically, some of the albedo calculation are not needed if we aren't picking between slabs
 						// We rely on the compiler to get rid of dead code here to avoid diverging the logic too far between single and multi-lobe cases
 						float ContribSlab = LobeColorToWeight(SlabAlbedo);
-						if (SlabSampler.Accept(SlabMix * ContribSlab))
+						if (SlabSampler.Accept((1 - RefractMix) * ContribSlab))
 						{
 							// We've selected this particular slab - now write it into the payload
 							// NOTE: This could happen more than once, so be careful to write to all fields, even the unused ones!
@@ -585,7 +592,7 @@ RAY_TRACING_ENTRY_CLOSEST_HIT(PathTracingMaterialCHS,
 							}
 							else
 							{
-							#if (MATERIAL_ISTHINSURFACE || MATERIAL_SHADINGMODEL_THIN_TRANSLUCENT)
+							#if MATERIAL_ISTHINSURFACE
 								// material is thin, make sure we don't use random walk
 								// SUBSTRATE_TODO: derive the proper BSDF for the internal scattering layer inside the slab
 							#else
@@ -615,9 +622,10 @@ RAY_TRACING_ENTRY_CLOSEST_HIT(PathTracingMaterialCHS,
 							}
 							#endif
 						}
+					#if PATH_TRACER_USE_SUBSTRATE_GLASS_LOBE
 						// now test for the presence of a refracted lobe
 						float ContribRefr = LobeColorToWeight(RefrAlbedo);
-						if (SlabSampler.Accept((1 - SlabMix) * ContribRefr))
+						if (SlabSampler.Accept(RefractMix * ContribRefr))
 						{
 							// We've selected this particular slab - now write it into the payload
 							// NOTE: This could happen more than once, so be careful to write to all fields, even the unused ones!
@@ -630,13 +638,22 @@ RAY_TRACING_ENTRY_CLOSEST_HIT(PathTracingMaterialCHS,
 							Payload.WeightV               = BSDF.LuminanceWeightV;
 							Payload.TransmittanceN        = BSDF.TransmittanceAboveAlongN;
 							Payload.CoverageAboveAlongN   = BSDF.CoverageAboveAlongN;
-							#if MATERIAL_ISTHINSURFACE || MATERIAL_SHADINGMODEL_THIN_TRANSLUCENT // SUBSTRATE_TODO: Should THIN_TRANSLUCENT imply ISTHIN_SURFACE?
+							#if MATERIAL_ISTHINSURFACE
 							Payload.ShadingModelID = SHADINGMODELID_THIN_TRANSLUCENT;
 							#else
 							Payload.ShadingModelID = SHADINGMODELID_SOLID_GLASS;
+							// Albedo inversion for ambient medium
+							// NOTE: we could technically also include IOR and Phase function - but this basic inversion does most of the work
+							float3 Color = DiffuseColor;
+							// Van de hulst mapping (with improved constants by d'Eon, see formula 48.20)
+							const float b1 = 1.0 / 0.100392;
+							const float b2 = 1.21573;
+							const float3 t1 = (0.5 * b1 + 0.5) + Color * (0.5 * b2 * b1);
+							const float3 t2 = b1 - Color * b1;
+							const float3 s = t1 - sqrt(t1 * t1 - t2);
+							DiffuseColor = 1.0 - s * s;
 							#endif
-							Payload.Ior                   = 0.0;
-							Payload.DiffuseColor          = 0.0;
+							Payload.DiffuseColor          = DiffuseColor; // SUBSTRATE_TODO: Add albedo inversion for solid case?
 							Payload.SpecularColor         = F0;
 							Payload.SpecularEdgeColor     = F90;
 							Payload.RoughnessData         = RoughnessData;
@@ -646,16 +663,14 @@ RAY_TRACING_ENTRY_CLOSEST_HIT(PathTracingMaterialCHS,
 							Payload.FuzzColor     = FuzzColor;
 							Payload.FuzzRoughness = FuzzRoughness;
 							Payload.MeanFreePath  = MeanFreePath;
-							Payload.PhaseG        = 0;
-							Payload.Ior = Ior;
+							Payload.PhaseG        = PhaseG;
+							Payload.Ior           = Ior;
 							// SUBSTRATE_TODO: should these be supported in the Glass case?
 							Payload.SpecularProfileId = 0;
 							Payload.GlintValue = 1.f;
 
-
-							#if (MATERIAL_ISTHINSURFACE || MATERIAL_SHADINGMODEL_THIN_TRANSLUCENT) && MATERIALBLENDING_TRANSLUCENT && (SUBSTRATE_BLENDING_TRANSLUCENT_GREYTRANSMITTANCE || SUBSTRATE_BLENDING_TRANSLUCENT_COLOREDTRANSMITTANCE)
 							// Material is configured as a thin surface, and refraction is disabled
-							if (!(Payload.Ior > 0.0))
+							#if MATERIAL_ISTHINSURFACE && PATH_TRACER_REFRACTION_ENABLED == 0
 							{
 								// SUBSTRATE_TODO: Do we really want to keep this? Should maybe always use refraction path?
 								float3 Transmission = Payload.GetTransmittanceColor();
@@ -665,6 +680,7 @@ RAY_TRACING_ENTRY_CLOSEST_HIT(PathTracingMaterialCHS,
 							}
 							#endif
 						}
+					#endif // PATH_TRACER_USE_SUBSTRATE_GLASS_LOBE
 					}
 					break;
 				}
@@ -1050,12 +1066,12 @@ RAY_TRACING_ENTRY_CLOSEST_HIT(PathTracingMaterialCHS,
 
 		Payload.ShadingModelID = SHADINGMODELID_SOLID_GLASS; // replace shading model to avoid extra cases in raygen
 	#if USE_LEGACY_CODEPATH_FOR_SUBSTRATE
-		Payload.DiffuseColor = 0; // not used
 		Payload.MeanFreePath = rcp(max(1e-5, WaterExtinction));
 	#else
-		Payload.BaseColor = WaterAlbedo; // TODO: unused at the moment
 		Payload.SetExtinction(WaterExtinction);
 	#endif
+	    Payload.SetGlassAlbedo(WaterAlbedo);
+		Payload.SetGlassPhase(PhaseG);
 		Payload.Ior = DielectricF0ToIor(DielectricSpecularToF0(Specular));
 		// we are using actual refraction, so don't apply any transparency
 		Payload.TransparencyColor = 0.0;
@@ -1253,6 +1269,8 @@ RAY_TRACING_ENTRY_CLOSEST_HIT(PathTracingMaterialCHS,
 			#else
 				Payload.SetExtinction(Extinction);
 			#endif
+
+				Payload.SetGlassAlbedo(0.0);
 			}
 			Payload.BSDFOpacity = 1.0; // cancels out with the probability of selecting the lobe
 		}
@@ -1455,12 +1473,13 @@ RAY_TRACING_ENTRY_ANY_HIT(PathTracingMaterialAHS,
 		float NoV = 0.0f;
 		float3 LocalSigmaT = 0.0f;
 		float Ior = 1.0;
-		float RefrProb  = 0.0;
+		float RefractMix  = 0.0;
 		Substrate_for_unroll(int BSDFIdx = 0, BSDFIdx < SubstratePixelHeader.SubstrateTree.BSDFCount, ++BSDFIdx)
 		{
 #define BSDF SubstratePixelHeader.SubstrateTree.BSDFs[BSDFIdx]
 			if (BSDF.bIsBottom)
 			{
+				// SUBSTRATE_TODO: Improve this to correctly deal with multiple bottom slabs (in the case of non-parameter blended horizontal blends)
 				float3 N = SubstratePixelHeader.SharedLocalBases.Normals[BSDF_GETSHAREDLOCALBASISID(BSDF)];
 				NoV = dot(N, V_World);
 				switch (BSDF_GETTYPE(BSDF))
@@ -1477,21 +1496,35 @@ RAY_TRACING_ENTRY_ANY_HIT(PathTracingMaterialAHS,
 							RoughnessData.y = Haziness.Roughness;
 							RoughnessData.z = Haziness.Weight;
 						}
-
-						if (BSDF_GETSSSTYPE(BSDF) == SSS_TYPE_DIFFUSION)
+					#if PATH_TRACER_USE_SUBSTRATE_GLASS_LOBE
+						switch (BSDF_GETSSSTYPE(BSDF))
 						{
-							// only non-SSS slabs 
-							LocalSigmaT = rcp(max(SLAB_SSSMFP(BSDF), 0.0001));
-					#if MATERIALBLENDING_TRANSLUCENT && (SUBSTRATE_BLENDING_TRANSLUCENT_GREYTRANSMITTANCE || SUBSTRATE_BLENDING_TRANSLUCENT_COLOREDTRANSMITTANCE)
-						#if PATH_TRACER_REFRACTION_ENABLED
-							Ior = GetRefractionIor(PixelMaterialInputs, F0);
-						#endif
-							float SlabMix = LobeSelectionProb(SLAB_DIFFUSEALBEDO(BSDF), exp(-SUBSTRATE_SIMPLEVOLUME_THICKNESS_CM * LocalSigmaT));
-							RefrProb = 1 - SlabMix;
-
-					#endif
+							case SSS_TYPE_DIFFUSION:
+							case SSS_TYPE_SIMPLEVOLUME:
+							{
+								// only non-SSS slabs 
+								float3 MeanFreePath = SLAB_SSSMFP(BSDF);
+								LocalSigmaT = rcp(max(MeanFreePath, 0.0001));
+							#if PATH_TRACER_REFRACTION_ENABLED
+								Ior = GetRefractionIor(PixelMaterialInputs, F0);
+							#endif
+							#if MATERIAL_ISTHINSURFACE
+								float3 Transmittance = MeanFreePathToTransmittance(MeanFreePath * CENTIMETER_TO_METER, SUBSTRATE_SIMPLEVOLUME_THICKNESS_CM * CENTIMETER_TO_METER);
+								RefractMix = saturate(Luminance(Transmittance));
+							#else
+								// blend towards Opaque SSS in the range of [0cm, 1cm]
+								RefractMix = saturate(max3(MeanFreePath.x, MeanFreePath.y, MeanFreePath.z));
+							#endif
+								break;
+							}
+							default:
+							{
+								// All other SSS types (including none) do not product any transmission and should be treated as opaque
+								break;
+							}
 						}
 						break;
+					#endif // PATH_TRACER_USE_SUBSTRATE_GLASS_LOBE
 					}
 					case SUBSTRATE_BSDF_TYPE_SINGLELAYERWATER:
 					{
@@ -1502,6 +1535,7 @@ RAY_TRACING_ENTRY_ANY_HIT(PathTracingMaterialAHS,
 #endif
 						float F0 = DielectricSpecularToF0(SLW_SPECULAR(BSDF));
 						Ior = DielectricF0ToIor(F0);
+						RefractMix = 1;
 						break;
 					}
 				}
@@ -1515,13 +1549,13 @@ RAY_TRACING_ENTRY_ANY_HIT(PathTracingMaterialAHS,
 
 		// SUBSTRATE_TODO: Revisit this -- TotalTransmittancePreCoverage should probably account for this in the thin case
 		float3 Transmission = exp(-LocalSigmaT * SUBSTRATE_SIMPLEVOLUME_THICKNESS_CM);
-		float3 Tr = RefrProb * ComputeThinSlabWeights(Transmission, NoV, Ior, F0).Transmitted;
+		float3 Tr = RefractMix * ComputeThinSlabWeights(Transmission, NoV, Ior, F0).Transmitted;
 #else
 		HLSL_STATIC_ASSERT(PATH_TRACER_REFRACTION_ENABLED == 1, "Expect refraction to be true here");
 		// Solid Glass Fake Caustics
-		float Tr = RefrProb * FakeCaustics(F0, Ior, NoV);
+		float Tr = RefractMix * FakeCaustics(F0, Ior, NoV);
 		// Account for transmission 
-		PackedPayload.SetTau(PackedPayload.GetTau() + TotalCoverage * RefrProb * ComputeOpticalDepthForHit(LocalSigmaT, bIsFrontFace));
+		PackedPayload.SetTau(PackedPayload.GetTau() + TotalCoverage * RefractMix * ComputeOpticalDepthForHit(LocalSigmaT, bIsFrontFace));
 #endif
 
 		// If material is doing refraction, apply an extra factor to account for ApproximateCaustics
diff --git a/Engine/Shaders/Private/PathTracing/Volume/PathTracingVolumeCommon.ush b/Engine/Shaders/Private/PathTracing/Volume/PathTracingVolumeCommon.ush
--- a/Engine/Shaders/Private/PathTracing/Volume/PathTracingVolumeCommon.ush
+++ b/Engine/Shaders/Private/PathTracing/Volume/PathTracingVolumeCommon.ush
@@ -246,6 +246,9 @@ FPathTracingPayload CreateMediumHitPayload(float HitT, float3 TranslatedWorldPos
 	Result.SetHG(HGWeight, ShadedResult.PhaseG);
 	Result.SetDualHG(DualHGWeight, ShadedResult.DualPhaseData);
 	Result.SetCloudFactor(1.0);
+#if SUBSTRATE_ENABLED
+	Result.GlintValue = 1; // avoid encode/decode cycle zero'ing out some parameters
+#endif
 	return Result;
 }
 
diff --git a/Engine/Shaders/Shared/RayTracingTypes.h b/Engine/Shaders/Shared/RayTracingTypes.h
--- a/Engine/Shaders/Shared/RayTracingTypes.h
+++ b/Engine/Shaders/Shared/RayTracingTypes.h
@@ -75,8 +75,9 @@ struct FPathTracingPackedPathState
 	float3    RayDirection;
 	float3    PathThroughput;
 	uint2     PackedRoughnessSigma;
+	uint2     PackedScatterPhase;
 };
-HLSL_STATIC_ASSERT(sizeof(FPathTracingPackedPathState) == 80, "Packed Path State size should be minimized");
+HLSL_STATIC_ASSERT(sizeof(FPathTracingPackedPathState) == 88, "Packed Path State size should be minimized");
 
 struct FRayTracingDecal
 {
diff --git a/Engine/Source/Runtime/Renderer/Private/PathTracing.cpp b/Engine/Source/Runtime/Renderer/Private/PathTracing.cpp
--- a/Engine/Source/Runtime/Renderer/Private/PathTracing.cpp
+++ b/Engine/Source/Runtime/Renderer/Private/PathTracing.cpp
@@ -841,8 +841,8 @@ struct FPathTracingState
 	TRefCountPtr<IPooledRenderTarget> AtmosphereOpticalDepthLUT;
 	FAtmosphereConfig LastAtmosphereConfig;
 
-	// Buffer containing the starting medium extinction
-	TRefCountPtr<FRDGPooledBuffer>    StartingExtinctionCoefficient;
+	// Buffer containing the starting medium coefficients
+	TRefCountPtr<FRDGPooledBuffer>    StartingMediumData;
 
 	// Custom path tracing spacial temporal denoiser result, used by plugins
 	TRefCountPtr<UE::Renderer::Private::IPathTracingSpatialTemporalDenoiser::IHistory> SpatialTemporalDenoiserHistory;
@@ -1420,8 +1420,8 @@ class FPathTracingRG : public FGlobalShader
 		// scene decals
 		SHADER_PARAMETER_RDG_UNIFORM_BUFFER(FRayTracingDecals, DecalParameters)
 
-		// camera ray starting extinction coefficient
-		SHADER_PARAMETER_RDG_BUFFER_SRV(Buffer<float>, StartingExtinctionCoefficient)
+		// camera ray starting medium coefficients
+		SHADER_PARAMETER_RDG_BUFFER_SRV(Buffer<float>, StartingMediumData)
 
 		// Used by multi-GPU rendering and TDR-avoidance tiling
 		SHADER_PARAMETER(FIntPoint, TilePixelOffset)
@@ -1587,7 +1587,7 @@ class FPathTracingInitExtinctionCoefficientRG : public FGlobalShader
 		SHADER_PARAMETER_RDG_UNIFORM_BUFFER(FNaniteRayTracingUniformParameters, NaniteRayTracing)
 		SHADER_PARAMETER_STRUCT_REF(FViewUniformShaderParameters, ViewUniformBuffer)
 		RDG_BUFFER_ACCESS_ARRAY(SBTBuffers)
-		SHADER_PARAMETER_RDG_BUFFER_UAV(RWBuffer<float>, RWStartingExtinctionCoefficient)
+		SHADER_PARAMETER_RDG_BUFFER_UAV(RWBuffer<float>, RWStartingMediumData)
 	END_SHADER_PARAMETER_STRUCT()
 };
 IMPLEMENT_GLOBAL_SHADER(FPathTracingInitExtinctionCoefficientRG, "/Engine/Private/PathTracing/PathTracingInitExtinctionCoefficient.usf", "PathTracingInitExtinctionCoefficientRG", SF_RayGen);
@@ -3769,15 +3769,15 @@ void FDeferredShadingSceneRenderer::RenderPathTracing(
 
 	// If this is the first sample, recompute the initial medium
 	// In this case of an offline render, do this every frame so that motion blur through a boundary is properly accounted for
-	FRDGBufferRef StartingExtinctionCoefficient = nullptr;
+	FRDGBufferRef StartingMediumData = nullptr;
 	if (!Config.UseCameraMediumTracking)
 	{
-		PathTracingState->StartingExtinctionCoefficient.SafeRelease();
+		PathTracingState->StartingMediumData.SafeRelease();
 		// camera medium tracking is not enabled - just make a temp buffer and set it to 0
-		StartingExtinctionCoefficient = GraphBuilder.CreateBuffer(FRDGBufferDesc::CreateBufferDesc(sizeof(float), 3), TEXT("PathTracer.StartingExtinctionCoefficient"));
-		AddClearUAVPass(GraphBuilder, GraphBuilder.CreateUAV(StartingExtinctionCoefficient, PF_R32_FLOAT), 0);
+		StartingMediumData = GraphBuilder.CreateBuffer(FRDGBufferDesc::CreateBufferDesc(sizeof(float), 7), TEXT("PathTracer.StartingMediumData"));
+		AddClearUAVPass(GraphBuilder, GraphBuilder.CreateUAV(StartingMediumData, PF_R32_FLOAT), 0);
 	}
-	else if (!PathTracingState->StartingExtinctionCoefficient.IsValid() || PathTracingState->SampleIndex == 0 || View.bIsOfflineRender)
+	else if (!PathTracingState->StartingMediumData.IsValid() || PathTracingState->SampleIndex == 0 || View.bIsOfflineRender)
 	{
 		FPathTracingInitExtinctionCoefficientRG::FPermutationDomain PermutationVector;
 		PermutationVector.Set<FPathTracingInitExtinctionCoefficientRG::FNeedTMinWorkaround>(GPUNeedsTMinWorkaround());
@@ -3785,14 +3785,14 @@ void FDeferredShadingSceneRenderer::RenderPathTracing(
 		auto RayGenShader = GetGlobalShaderMap(View.FeatureLevel)->GetShader<FPathTracingInitExtinctionCoefficientRG>(PermutationVector);
 
 		// prepare extinction coefficient for camera rays
-		StartingExtinctionCoefficient = GraphBuilder.CreateBuffer(FRDGBufferDesc::CreateBufferDesc(sizeof(float), 3), TEXT("PathTracer.StartingExtinctionCoefficient"));
+		StartingMediumData = GraphBuilder.CreateBuffer(FRDGBufferDesc::CreateBufferDesc(sizeof(float), 7), TEXT("PathTracer.StartingMediumData"));
 
 		FPathTracingInitExtinctionCoefficientRG::FParameters* PassParameters = GraphBuilder.AllocParameters<FPathTracingInitExtinctionCoefficientRG::FParameters>();
 		PassParameters->TLAS = Scene->RayTracingScene.GetLayerView(ERayTracingSceneLayer::Base, View.GetRayTracingSceneViewHandle());
 		PassParameters->Scene = View.GetSceneUniforms().GetBuffer(GraphBuilder);
 		PassParameters->NaniteRayTracing = Nanite::GRayTracingManager.GetUniformBuffer();
 		PassParameters->ViewUniformBuffer = View.ViewUniformBuffer;
-		PassParameters->RWStartingExtinctionCoefficient = GraphBuilder.CreateUAV(StartingExtinctionCoefficient, PF_R32_FLOAT);
+		PassParameters->RWStartingMediumData = GraphBuilder.CreateUAV(StartingMediumData, PF_R32_FLOAT);
 
 		for (FRDGBuffer* RDGBuffer : View.DynamicRayTracingRDGBuffers)
 		{
@@ -3818,12 +3818,12 @@ void FDeferredShadingSceneRenderer::RenderPathTracing(
 					1, 1
 				);
 			});
-		GraphBuilder.QueueBufferExtraction(StartingExtinctionCoefficient, &PathTracingState->StartingExtinctionCoefficient);
+		GraphBuilder.QueueBufferExtraction(StartingMediumData, &PathTracingState->StartingMediumData);
 	}
 	else
 	{
-		check(PathTracingState->StartingExtinctionCoefficient.IsValid());
-		StartingExtinctionCoefficient = GraphBuilder.RegisterExternalBuffer(PathTracingState->StartingExtinctionCoefficient, TEXT("PathTracer.StartingExtinctionCoefficient"));
+		check(PathTracingState->StartingMediumData.IsValid());
+		StartingMediumData = GraphBuilder.RegisterExternalBuffer(PathTracingState->StartingMediumData, TEXT("PathTracer.StartingMediumData"));
 	}
 
 	// prepare atmosphere optical depth lookup texture (if needed)
@@ -4176,7 +4176,7 @@ void FDeferredShadingSceneRenderer::RenderPathTracing(
 							PassParameters->ViewUniformBuffer = View.ViewUniformBuffer;
 							PassParameters->PathTracingData = Config.PathTracingData;
 							PassParameters->RandomSequenceSpaceFillingCurve = GSystemTextures.GetSpaceFillingCurveTexture(GraphBuilder);
-							PassParameters->StartingExtinctionCoefficient = GraphBuilder.CreateSRV(StartingExtinctionCoefficient, PF_R32_FLOAT);
+							PassParameters->StartingMediumData = GraphBuilder.CreateSRV(StartingMediumData, PF_R32_FLOAT);
 							if (PreviousPassParameters == nullptr)
 							{
 								// upload sky/lights data
