[Nanite Skinning] Temporary fix for lack of templated ByteAddressBuffer Store on some platforms.

#rnx

diff --git a/Engine/Shaders/Private/Skinning/SkinningDataDecode.ush b/Engine/Shaders/Private/Skinning/SkinningDataDecode.ush
--- a/Engine/Shaders/Private/Skinning/SkinningDataDecode.ush
+++ b/Engine/Shaders/Private/Skinning/SkinningDataDecode.ush
@@ -76,7 +76,13 @@ void StoreCompressedBoneTransform(RWByteAddressBuffer DstBuffer, uint BaseOffset
 	const uint Address = BaseOffset + BoneIndex * (uint)sizeof(float3x4);
 	
 	const float3x4 Tmp = transpose(BoneTransform);
-	DstBuffer.Store<float3x4>(Address, Tmp);
+	//DstBuffer.Store<float3x4>(Address, Tmp);
+	const uint4 Row0 = asuint(Tmp[0]);
+	const uint4 Row1 = asuint(Tmp[1]);
+	const uint4 Row2 = asuint(Tmp[2]);
+	DstBuffer.Store4(Address +  0, Row0);
+	DstBuffer.Store4(Address + 16, Row1);
+	DstBuffer.Store4(Address + 32, Row2);
 }
 
 #endif // !USE_COMPRESSED_BONE_TRANSFORM
