Change Substrate's Format=0 Fuzz roughness to have fixed value.

There is not space for storing this roughness. Previously it was coupled with the main roughness but that would cause fuzz appearance to vanish for low roughness value.

diff --git a/Engine/Shaders/Private/Substrate/SubstrateRead.ush b/Engine/Shaders/Private/Substrate/SubstrateRead.ush
--- a/Engine/Shaders/Private/Substrate/SubstrateRead.ush
+++ b/Engine/Shaders/Private/Substrate/SubstrateRead.ush
@@ -116,7 +116,7 @@ FSubstrateGBufferBSDF SubstrateReadGBufferBSDF(in FScreenSpaceData ScreenSpaceDa
 			{
 				FuzzColor = ExtractSubsurfaceColor(GBuffer);
 				FuzzAmount = GBuffer.CustomData.a;
-				FuzzRoughness = GBuffer.Roughness;
+				FuzzRoughness = 0.5f; // Fuzz roughness is not stored into the gbuffer. Using GBuffer.Roughness it not ideal as low value, makes most of the fuzz contribution to vanish. Instead hardcode value.
 			}
 			else if (GBuffer.ShadingModelID == SHADINGMODELID_SUBSURFACE_PROFILE || GBuffer.ShadingModelID == SHADINGMODELID_PREINTEGRATED_SKIN)
 			{
