Fix shader validation error (out of bound read) on Substrate's classification.

diff --git a/Engine/Shaders/Private/Substrate/SubstrateMaterialClassification.usf b/Engine/Shaders/Private/Substrate/SubstrateMaterialClassification.usf
--- a/Engine/Shaders/Private/Substrate/SubstrateMaterialClassification.usf
+++ b/Engine/Shaders/Private/Substrate/SubstrateMaterialClassification.usf
@@ -304,7 +304,7 @@ void TileMainCS(uint2 DispatchThreadId : SV_DispatchThreadID, uint LinearIndex :
 	}
 	GroupMemoryBarrierWithGroupSync();
 
-	const uint FinalTileFlags = s_TileFlags[LinearIndex] | s_TileFlags[LinearIndex + 1];
+	const uint FinalTileFlags = s_TileFlags[0] | s_TileFlags[1];
 
 	const bool bTileContainsSubstrate						= (FinalTileFlags & 0x1u) > 0;
 	const bool bTileContainsSimple							= (FinalTileFlags & 0x2u) > 0;
