Mobile: Fixed broken ClearCoat shading when decals overwrite roughness to zero in GBuffer
#okforgtihub public

diff --git a/Engine/Shaders/Private/DeferredShadingCommon.ush b/Engine/Shaders/Private/DeferredShadingCommon.ush
--- a/Engine/Shaders/Private/DeferredShadingCommon.ush
+++ b/Engine/Shaders/Private/DeferredShadingCommon.ush
@@ -572,7 +572,8 @@ void MobileEncodeGBuffer(
 
 		OutGBufferB.r = GBuffer.Metallic;
 		OutGBufferB.g = GBuffer.Specular;
-		OutGBufferB.b = GBuffer.Roughness;
+		// GBuffer can be overwritten by decals and others. See MobileBasePassPixelShader.usf for this clamp
+		OutGBufferB.b = max(0.015625, GBuffer.Roughness);
 		OutGBufferB.a = GBuffer.ShadingModelID / 255.0;
 
 		OutGBufferC.rgb = EncodeBaseColor( GBuffer.BaseColor );
