Adjust the minimum positive value for the "pow" input to fix a device lost issue on some mali devices on vulkan.

diff --git a/Engine/Shaders/Private/Common.ush b/Engine/Shaders/Private/Common.ush
--- a/Engine/Shaders/Private/Common.ush
+++ b/Engine/Shaders/Private/Common.ush
@@ -815,9 +815,10 @@ MaterialFloat4 ClampedPow(MaterialFloat4 X,MaterialFloat4 Y)
 
 // Pow function that will return 0 if Base is <= 0 (or small enough to yield a floating point special).
 // This is done to prevent floating point specials when compilers expands pow into exp(Exponent * log(Base)).
+// Do not add the ".f" for the 2.980233e-8 because it will convert the Base from half to float, leave Exponent as half, which cause a device lost on some mali devices on Vulkan. UE-331629
 MaterialFloat PositiveClampedPow(MaterialFloat Base, MaterialFloat Exponent)
 {
-	return (Base <= 2.980233e-8f) ? 0.0f : pow(Base, Exponent);
+	return (Base <= 2.980233e-8) ? 0.0f : pow(Base, Exponent);
 }
 MaterialFloat2 PositiveClampedPow(MaterialFloat2 Base, MaterialFloat2 Exponent)
 {
diff --git a/Engine/Source/Runtime/Engine/Private/Materials/MaterialExpressionsIR.cpp b/Engine/Source/Runtime/Engine/Private/Materials/MaterialExpressionsIR.cpp
--- a/Engine/Source/Runtime/Engine/Private/Materials/MaterialExpressionsIR.cpp
+++ b/Engine/Source/Runtime/Engine/Private/Materials/MaterialExpressionsIR.cpp
@@ -861,7 +861,7 @@ static FValueRef PositiveClampedPow(MIR::FEmitter& Em, FValueRef Base, FValueRef
 
 	TOptional<MIR::FPrimitive> ValuePrimitiveType = Base->Type.AsPrimitive();
 	return Em.Select(
-				Em.LessThanOrEquals(Base, Em.ConstantFloat(2.980233e-8f)),
+				Em.LessThanOrEquals(Base, Em.ConstantFloat(2.980233e-8)),
 				Em.ConstantZero(ValuePrimitiveType->ScalarKind),
 				Em.Pow(Base, Exponent));
 }
