Substrate classification optimization: Use SelectiveOutputMask mask to skip the load of the Anisotropy property from the gbuffer.

diff --git a/Engine/Shaders/Private/Substrate/SubstrateMaterialClassification.usf b/Engine/Shaders/Private/Substrate/SubstrateMaterialClassification.usf
--- a/Engine/Shaders/Private/Substrate/SubstrateMaterialClassification.usf
+++ b/Engine/Shaders/Private/Substrate/SubstrateMaterialClassification.usf
@@ -137,8 +137,12 @@ void TileMainCS(uint2 DispatchThreadId : SV_DispatchThreadID, uint LinearIndex :
 
 		if (bContainsSubstrateMaterial)
 		{
-			bool bHasAnisotropy = abs(GBufferData.Anisotropy) > 0;
-
+			bool bHasAnisotropy = false;
+			BRANCH
+			if(HasAnisotropy(GBufferData.SelectiveOutputMask))
+			{
+				bHasAnisotropy = abs(GBufferData.Anisotropy) > 0;	// Skip this load when the material doesn't even support anisotropy.
+			}
 			bContainsComplexMaterial = bHasAnisotropy || GBufferData.ShadingModelID == SHADINGMODELID_EYE || GBufferData.ShadingModelID == SHADINGMODELID_HAIR;
 			bContainsSimpleMaterial = !bContainsComplexMaterial && GBufferData.ShadingModelID == SHADINGMODELID_DEFAULT_LIT;
 			bContainsSingleMaterial = !bContainsComplexMaterial && !bContainsSimpleMaterial;
