Add Wave permutation support for Substrate classification pass on platforms having wave size different from 64.

Improve Switch2 classification by 0.1ms (0.64ms -> 0.54ms).

diff --git a/Engine/Shaders/Private/Substrate/SubstrateMaterialClassification.usf b/Engine/Shaders/Private/Substrate/SubstrateMaterialClassification.usf
--- a/Engine/Shaders/Private/Substrate/SubstrateMaterialClassification.usf
+++ b/Engine/Shaders/Private/Substrate/SubstrateMaterialClassification.usf
@@ -4,7 +4,9 @@
 
 #define SUBSTRATE_INLINE_SHADING 0
 #define SUBSTRATE_SSS_MATERIAL_OVERRIDE 0
-#define SUBSTRATE_COMPLEXSPECIALPATH 1
+// Complex Special path is only supported with Format=1
+#define SUBSTRATE_COMPLEXSPECIALPATH (SUBSTRATE_GBUFFER_FORMAT==1)
+
 #include "/Engine/Private/Substrate/Substrate.ush"
 #include "SubstrateTile.ush"
 
@@ -13,6 +15,7 @@
 #endif
 
 #define GROUP_THREAD_COUNT (SUBSTRATE_TILE_SIZE * SUBSTRATE_TILE_SIZE)
+#define SUBSTRATE_TILE_BITMASK_SSS (1<<SUBSTRATE_TILE_TYPE_COUNT)
 
 ////////////////////////////////////////////////////////////////////////////////////////////////////////////
 
@@ -76,13 +79,12 @@ uint GetDBufferTargetMask(uint2 PixelPos)
 Texture2D<float3> OpaqueRoughRefractionTexture;
 #endif // SUBSTRATE_OPAQUE_ROUGH_REFRACTION_ENABLED
 
-#if !PERMUTATION_WAVE_OPS
-groupshared uint s_TileFlags[GROUP_THREAD_COUNT];
+#if PERMUTATION_WAVE_OPS
+groupshared uint s_TileBitmask;
+#else
+groupshared uint s_TileBitmask[GROUP_THREAD_COUNT];
 #endif
 
-#if PERMUTATION_WAVE_OPS && COMPILER_SUPPORTS_WAVE_SIZE
-WAVESIZE(64) // PERMUTATION_WAVE_OPS is true only when wave>=64 are available
-#endif
 [numthreads(SUBSTRATE_TILE_SIZE, SUBSTRATE_TILE_SIZE, 1)]
 void TileMainCS(uint2 DispatchThreadId : SV_DispatchThreadID, uint LinearIndex : SV_GroupIndex, uint3 GroupId : SV_GroupID)
 {
@@ -102,6 +104,7 @@ void TileMainCS(uint2 DispatchThreadId : SV_DispatchThreadID, uint LinearIndex :
 	// * If the tile is entirely empty: we clear the header & SSS data
 	// * If the data contains any data: we do fine grain checking, and clear header & SSS data only for needed pixels. The top layer data texture is used 
 	//   to know if a pixel is valid or not (since the material header is not cleared when the Cmask permutation is used).
+	uint GroupTileBitmask = 0;
 
 #if PERMUTATION_CMASK && SUBSTRATE_GBUFFER_FORMAT==1
 	// Coarse test for clearing header (& SSS data) based on CMask data
@@ -116,26 +119,13 @@ void TileMainCS(uint2 DispatchThreadId : SV_DispatchThreadID, uint LinearIndex :
 #endif
 	{
 
-	// Pixels outside of the view area are considered simple to enable screen borders to receive the simple permutation when not aligned to shader group size.
-	bool bContainsComplexSpecialMaterial			= false;
-	bool bContainsComplexMaterial					= false;
-	bool bContainsSimpleMaterial					= false;
-	bool bContainsSingleMaterial					= false;
-	bool bContainsSubstrateMaterial					= false;
-	bool bContainsDecals							= false;
-	bool bContainsOpaqueRoughRefraction				= false;
-	bool bContainsScreenSpaceSubsurfaceScattering	= false;
-
 	FSubstrateOpaqueRoughRefractionData OpaqueRoughRefractionData = (FSubstrateOpaqueRoughRefractionData)0;
 	if (bIsValid)
 	{
 #if SUBSTRATE_GBUFFER_FORMAT==0
-
 		// Control tiles using ShadingModelID
-		FGBufferData GBufferData   = GetGBufferData(BufferUV);
-		bContainsSubstrateMaterial = GBufferData.ShadingModelID != SHADINGMODELID_UNLIT;
-
-		if (bContainsSubstrateMaterial)
+		const FGBufferData GBufferData   = GetGBufferData(BufferUV);
+		if (GBufferData.ShadingModelID != SHADINGMODELID_UNLIT)
 		{
 			bool bHasAnisotropy = false;
 			BRANCH
@@ -143,10 +133,32 @@ void TileMainCS(uint2 DispatchThreadId : SV_DispatchThreadID, uint LinearIndex :
 			{
 				bHasAnisotropy = abs(GBufferData.Anisotropy) > 0;	// Skip this load when the material doesn't even support anisotropy.
 			}
-			bContainsComplexMaterial = bHasAnisotropy || GBufferData.ShadingModelID == SHADINGMODELID_EYE || GBufferData.ShadingModelID == SHADINGMODELID_HAIR;
-			bContainsSimpleMaterial = !bContainsComplexMaterial && GBufferData.ShadingModelID == SHADINGMODELID_DEFAULT_LIT;
-			bContainsSingleMaterial = !bContainsComplexMaterial && !bContainsSimpleMaterial;
-			bContainsScreenSpaceSubsurfaceScattering = GBufferData.ShadingModelID == SHADINGMODELID_SUBSURFACE_PROFILE || GBufferData.ShadingModelID == SHADINGMODELID_EYE;
+
+			#if SUBSTRATE_COMPLEXSPECIALPATH
+			if (false)
+			{
+				GroupTileBitmask |= SUBSTRATE_TILE_BITMASK_COMPLEXSPECIAL;
+			}
+			else 
+			#endif
+			if (bHasAnisotropy || GBufferData.ShadingModelID == SHADINGMODELID_EYE || GBufferData.ShadingModelID == SHADINGMODELID_HAIR)
+			{
+				GroupTileBitmask |= SUBSTRATE_TILE_BITMASK_COMPLEX;
+			}
+			else if (GBufferData.ShadingModelID == SHADINGMODELID_DEFAULT_LIT)
+			{
+				GroupTileBitmask |= SUBSTRATE_TILE_BITMASK_SIMPLE;
+			}
+			else
+			{
+				GroupTileBitmask |= SUBSTRATE_TILE_BITMASK_SINGLE;
+			}
+
+			// Skipped as only used for Format=0
+			// if (GBufferData.ShadingModelID == SHADINGMODELID_SUBSURFACE_PROFILE || GBufferData.ShadingModelID == SHADINGMODELID_EYE)
+			// {
+			// 	GroupTileBitmask |= SUBSTRATE_TILE_BITMASK_SSS;
+			// }
 		}
 
 #else // SUBSTRATE_GBUFFER_FORMAT==0
@@ -159,24 +171,44 @@ void TileMainCS(uint2 DispatchThreadId : SV_DispatchThreadID, uint LinearIndex :
 
 		const bool bIsSimple =  SubstratePixelHeader.IsSimpleMaterial() || SubstratePixelHeader.ClosureCount == 0; // BSDFCount == 0 ensures that non-Substrate pixel, like sky pixels, won't make a simple tile flagged as complex
 		const bool bIsSingle = !SubstratePixelHeader.IsSimpleMaterial() && SubstratePixelHeader.IsSingleMaterial();
-		const bool bIsComplexSpecial = SubstratePixelHeader.IsComplexSpecialMaterial();
-		bContainsSubstrateMaterial  = SubstratePixelHeader.ClosureCount > 0;
-		bContainsSimpleMaterial  = bIsSimple;
-		bContainsSingleMaterial  = bIsSingle;
-		bContainsComplexMaterial		= !bIsSingle && !bIsSimple && !bIsComplexSpecial;
-		bContainsComplexSpecialMaterial = !bIsSingle && !bIsSimple &&  bIsComplexSpecial;
+		
+		if (SubstratePixelHeader.IsComplexSpecialMaterial())
+		{
+			GroupTileBitmask |= SUBSTRATE_TILE_BITMASK_COMPLEXSPECIAL;
+		}
+		else if (!bIsSingle && !bIsSimple)
+		{
+			GroupTileBitmask |= SUBSTRATE_TILE_BITMASK_COMPLEX;
+		}
+		else if (bIsSingle)
+		{
+			GroupTileBitmask |= SUBSTRATE_TILE_BITMASK_SINGLE;
+		}
+		else if (bIsSimple)
+		{
+			GroupTileBitmask |= SUBSTRATE_TILE_BITMASK_SIMPLE;
+		}
 
-		bContainsScreenSpaceSubsurfaceScattering = SubstratePixelHeader.HasSubsurface();
+		if (SubstratePixelHeader.HasSubsurface())
+		{
+			GroupTileBitmask |= SUBSTRATE_TILE_BITMASK_SSS;
+		}
 
 	#if PERMUTATION_DECAL
 		const uint DBufferResponseMask = SceneStencilTexture.Load(uint3(PixelCoord, 0)) STENCIL_COMPONENT_SWIZZLE;
 		const uint DBufferTargetMask = GetDBufferTargetMask(PixelCoord);
-		bContainsDecals = DBufferResponseMask != 0 && DBufferTargetMask != 0;
+		if (DBufferResponseMask != 0 && DBufferTargetMask != 0)
+		{
+			GroupTileBitmask |= SUBSTRATE_TILE_BITMASK_DECAL;
+		}
 	#endif
 
 	#if SUBSTRATE_OPAQUE_ROUGH_REFRACTION_ENABLED
 		OpaqueRoughRefractionData = SubstrateUnpackOpaqueRoughRefractionData(OpaqueRoughRefractionTexture[PixelCoord]);
-		bContainsOpaqueRoughRefraction = OpaqueRoughRefractionData.OpaqueRoughRefractionEnabled > 0.0f;
+		if (OpaqueRoughRefractionData.OpaqueRoughRefractionEnabled > 0.0f)
+		{
+			GroupTileBitmask |= SUBSTRATE_TILE_BITMASK_ROUGH_REFRACT;
+		}
 	#endif
 
 		// Output/Patch SSS data for legacy encoding (this allows to save ALU & bandwidth during the base pass0
@@ -185,7 +217,7 @@ void TileMainCS(uint2 DispatchThreadId : SV_DispatchThreadID, uint LinearIndex :
 		const bool bIsLegacySSSProfile = OptimisedLegacyMode == SINGLE_OPTLEGACYMODE_SSSPROFILE;
 		if (bIsSingle && (bIsLegacyWrapOrWrapThin || bIsLegacySSSProfile))
 		{
-			bContainsScreenSpaceSubsurfaceScattering = true;
+			GroupTileBitmask |= SUBSTRATE_TILE_BITMASK_SSS;
 			if (bIsLegacyWrapOrWrapThin)
 			{
 				FSubstrateSubsurfaceHeader SSSHeader = (FSubstrateSubsurfaceHeader)0;
@@ -252,7 +284,7 @@ void TileMainCS(uint2 DispatchThreadId : SV_DispatchThreadID, uint LinearIndex :
 
 #if SUBSTRATE_GBUFFER_FORMAT==1
 	BRANCH
-	if (!bContainsScreenSpaceSubsurfaceScattering)
+	if ((GroupTileBitmask & SUBSTRATE_TILE_BITMASK_SSS) == 0)
 	{
 		// We must fill all the pixel which does not have subsurface scattering by default so that the SSS code is not executed where it should not.
 		SubstrateStoreSubsurfaceHeader(MaterialTextureArrayUAV, FirstSliceStoringSubstrateSSSData, PixelCoord, 0u); // This is a good clear for FSubstrateSubsurfaceHeader, and we only need to clear the header.
@@ -260,89 +292,88 @@ void TileMainCS(uint2 DispatchThreadId : SV_DispatchThreadID, uint LinearIndex :
 #endif
 
 #if PERMUTATION_WAVE_OPS
-	const bool bTileContainsSubstrate						= WaveActiveAnyTrue(bContainsSubstrateMaterial);
-	const bool bTileContainsSimple							= WaveActiveAnyTrue(bContainsSimpleMaterial);
-	const bool bTileContainsSingle							= WaveActiveAnyTrue(bContainsSingleMaterial);
-	const bool bTileContainsComplex							= WaveActiveAnyTrue(bContainsComplexMaterial);
-	const bool bTileContainsComplexSpecial					= WaveActiveAnyTrue(bContainsComplexSpecialMaterial);
-	const bool bTileContainsOpaqueRoughRefraction			= WaveActiveAnyTrue(bContainsOpaqueRoughRefraction);
-	const bool bTileContainsScreenSpaceSubsurfaceScattering	= WaveActiveAnyTrue(bContainsScreenSpaceSubsurfaceScattering);
-	const bool bTileContainsDecals							= WaveActiveAnyTrue(bContainsDecals);
+	GroupTileBitmask = WaveActiveBitOr(GroupTileBitmask);
+
+	if (WaveGetLaneCount() < SUBSTRATE_TILE_SIZE * SUBSTRATE_TILE_SIZE)
+	{
+		if (LinearIndex == 0)
+		{
+			s_TileBitmask = 0;
+		}
+
+		GroupMemoryBarrierWithGroupSync();
+		
+		if (WaveIsFirstLane())
+		{
+			uint Unused;
+			InterlockedOr(s_TileBitmask, GroupTileBitmask, Unused);
+		}
+		
+		GroupMemoryBarrierWithGroupSync();
+		
+		GroupTileBitmask = s_TileBitmask;
+	}
+
 #else // PERMUTATION_WAVE_OPS
 
-	s_TileFlags[LinearIndex] = 
-		  (bContainsSubstrateMaterial				? 0x1u	: 0u)
-		| (bContainsSimpleMaterial					? 0x2u	: 0u)
-		| (bContainsSingleMaterial					? 0x4u	: 0u) 
-		| (bContainsComplexMaterial					? 0x8u	: 0u) 
-		| (bContainsComplexSpecialMaterial			? 0x10u : 0u)
-		| (bContainsOpaqueRoughRefraction			? 0x20u : 0u)
-		| (bContainsScreenSpaceSubsurfaceScattering	? 0x40u : 0u)
-		| (bContainsDecals							? 0x80u : 0u);
+	s_TileBitmask[LinearIndex] = GroupTileBitmask;
 
 	GroupMemoryBarrierWithGroupSync();
 	if (LinearIndex < 32)
 	{
-		s_TileFlags[LinearIndex] = s_TileFlags[LinearIndex] | s_TileFlags[LinearIndex + 32];
+		s_TileBitmask[LinearIndex] = s_TileBitmask[LinearIndex] | s_TileBitmask[LinearIndex + 32];
 	}
 	GroupMemoryBarrierWithGroupSync();
 	if (LinearIndex < 16)
 	{
-		s_TileFlags[LinearIndex] = s_TileFlags[LinearIndex] | s_TileFlags[LinearIndex + 16];
+		s_TileBitmask[LinearIndex] = s_TileBitmask[LinearIndex] | s_TileBitmask[LinearIndex + 16];
 	}
 	GroupMemoryBarrierWithGroupSync();
 
 	if (LinearIndex < 8)
 	{
-		s_TileFlags[LinearIndex] = s_TileFlags[LinearIndex] | s_TileFlags[LinearIndex + 8];
+		s_TileBitmask[LinearIndex] = s_TileBitmask[LinearIndex] | s_TileBitmask[LinearIndex + 8];
 	}
 	GroupMemoryBarrierWithGroupSync();
 	if (LinearIndex < 4)
 	{
-		s_TileFlags[LinearIndex] = s_TileFlags[LinearIndex] | s_TileFlags[LinearIndex + 4];
+		s_TileBitmask[LinearIndex] = s_TileBitmask[LinearIndex] | s_TileBitmask[LinearIndex + 4];
 	}
 	GroupMemoryBarrierWithGroupSync();
 	if (LinearIndex < 2)
 	{
-		s_TileFlags[LinearIndex] = s_TileFlags[LinearIndex] | s_TileFlags[LinearIndex + 2];
+		s_TileBitmask[LinearIndex] = s_TileBitmask[LinearIndex] | s_TileBitmask[LinearIndex + 2];
 	}
 	GroupMemoryBarrierWithGroupSync();
 
-	const uint FinalTileFlags = s_TileFlags[0] | s_TileFlags[1];
-
-	const bool bTileContainsSubstrate						= (FinalTileFlags & 0x1u) > 0;
-	const bool bTileContainsSimple							= (FinalTileFlags & 0x2u) > 0;
-	const bool bTileContainsSingle							= (FinalTileFlags & 0x4u) > 0;
-	const bool bTileContainsComplex							= (FinalTileFlags & 0x8u) > 0;
-	const bool bTileContainsComplexSpecial					= (FinalTileFlags & 0x10u) > 0;
-	const bool bTileContainsOpaqueRoughRefraction			= (FinalTileFlags & 0x20u) > 0;
-	const bool bTileContainsScreenSpaceSubsurfaceScattering	= (FinalTileFlags & 0x40u) > 0;
-	const bool bTileContainsDecals							= (FinalTileFlags & 0x80u) > 0;
+	GroupTileBitmask = s_TileBitmask[0] | s_TileBitmask[1];
 #endif // PERMUTATION_WAVE_OPS
 
-	if (LinearIndex < 1 && bTileContainsSubstrate)
+	if (LinearIndex == 0 && GroupTileBitmask != 0)
 	{
-		uint EncodedTile = SubstratePackTile(GroupId.xy, TileEncoding);
-
-		if (bTileContainsComplexSpecial)
+		const uint EncodedTile = SubstratePackTile(GroupId.xy, TileEncoding);
+		#if SUBSTRATE_COMPLEXSPECIALPATH
+		if (GroupTileBitmask & SUBSTRATE_TILE_BITMASK_COMPLEXSPECIAL)
 		{
 			uint WriteToIndex;
 			InterlockedAdd(TileDrawIndirectDataBufferUAV[GetSubstrateTileTypeDrawIndirectArgOffset_DWord(SUBSTRATE_TILE_TYPE_COMPLEX_SPECIAL) + 1], 1, WriteToIndex);
 			TileListBufferUAV[GetTileListBufferOffsets(SUBSTRATE_TILE_TYPE_COMPLEX_SPECIAL) + WriteToIndex] = EncodedTile;
 		}
-		else if (bTileContainsComplex)
+		else 
+		#endif
+		if (GroupTileBitmask & SUBSTRATE_TILE_BITMASK_COMPLEX)
 		{
 			uint WriteToIndex;
 			InterlockedAdd(TileDrawIndirectDataBufferUAV[GetSubstrateTileTypeDrawIndirectArgOffset_DWord(SUBSTRATE_TILE_TYPE_COMPLEX) + 1], 1, WriteToIndex);
 			TileListBufferUAV[GetTileListBufferOffsets(SUBSTRATE_TILE_TYPE_COMPLEX) + WriteToIndex] = EncodedTile;
 		}
-		else if (bTileContainsSingle)
+		else if (GroupTileBitmask & SUBSTRATE_TILE_BITMASK_SINGLE)
 		{
 			uint WriteToIndex;
 			InterlockedAdd(TileDrawIndirectDataBufferUAV[GetSubstrateTileTypeDrawIndirectArgOffset_DWord(SUBSTRATE_TILE_TYPE_SINGLE) + 1], 1, WriteToIndex);
 			TileListBufferUAV[GetTileListBufferOffsets(SUBSTRATE_TILE_TYPE_SINGLE) + WriteToIndex] = EncodedTile;
 		}
-		else // (bTileContainsSimple)
+		else // (GroupTileBitmask & SUBSTRATE_TILE_BITMASK_SIMPLE)
 		{
 			uint WriteToIndex;
 			InterlockedAdd(TileDrawIndirectDataBufferUAV[GetSubstrateTileTypeDrawIndirectArgOffset_DWord(SUBSTRATE_TILE_TYPE_SIMPLE) + 1], 1, WriteToIndex);
@@ -350,6 +381,8 @@ void TileMainCS(uint2 DispatchThreadId : SV_DispatchThreadID, uint LinearIndex :
 		}
 
 	#if SUBSTRATE_OPAQUE_ROUGH_REFRACTION_ENABLED
+		const bool bTileContainsOpaqueRoughRefraction			= GroupTileBitmask & SUBSTRATE_TILE_BITMASK_ROUGH_REFRACT; 			
+		const bool bTileContainsScreenSpaceSubsurfaceScattering	= GroupTileBitmask & SUBSTRATE_TILE_BITMASK_SSS;
 		if (bTileContainsOpaqueRoughRefraction)
 		{
 			uint WriteToIndex;
@@ -365,21 +398,21 @@ void TileMainCS(uint2 DispatchThreadId : SV_DispatchThreadID, uint LinearIndex :
 	#endif
 
 	#if PERMUTATION_DECAL
-		if (bTileContainsDecals)
+		if (GroupTileBitmask & SUBSTRATE_TILE_BITMASK_DECAL)
 		{
-			if (bTileContainsComplex)
+			if (GroupTileBitmask & SUBSTRATE_TILE_BITMASK_COMPLEX)
 			{
 				uint WriteToIndex;
 				InterlockedAdd(TileDrawIndirectDataBufferUAV[GetSubstrateTileTypeDrawIndirectArgOffset_DWord(SUBSTRATE_TILE_TYPE_DECAL_COMPLEX) + 1], 1, WriteToIndex);
 				TileListBufferUAV[GetTileListBufferOffsets(SUBSTRATE_TILE_TYPE_DECAL_COMPLEX) + WriteToIndex] = EncodedTile;
 			}
-			else if (bTileContainsSingle)
+			else if (GroupTileBitmask & SUBSTRATE_TILE_BITMASK_SINGLE)
 			{
 				uint WriteToIndex;
 				InterlockedAdd(TileDrawIndirectDataBufferUAV[GetSubstrateTileTypeDrawIndirectArgOffset_DWord(SUBSTRATE_TILE_TYPE_DECAL_SINGLE) + 1], 1, WriteToIndex);
 				TileListBufferUAV[GetTileListBufferOffsets(SUBSTRATE_TILE_TYPE_DECAL_SINGLE) + WriteToIndex] = EncodedTile;
 			}
-			else // (bTileContainsSimple)
+			else // (GroupTileBitmask & SUBSTRATE_TILE_BITMASK_SIMPLE)
 			{
 				uint WriteToIndex;
 				InterlockedAdd(TileDrawIndirectDataBufferUAV[GetSubstrateTileTypeDrawIndirectArgOffset_DWord(SUBSTRATE_TILE_TYPE_DECAL_SIMPLE) + 1], 1, WriteToIndex);
diff --git a/Engine/Shaders/Shared/SubstrateDefinitions.h b/Engine/Shaders/Shared/SubstrateDefinitions.h
--- a/Engine/Shaders/Shared/SubstrateDefinitions.h
+++ b/Engine/Shaders/Shared/SubstrateDefinitions.h
@@ -86,6 +86,13 @@
 #define SUBSTRATE_TILE_BITMASK_SINGLE					(1u<<SUBSTRATE_TILE_TYPE_SINGLE)
 #define SUBSTRATE_TILE_BITMASK_COMPLEX					(1u<<SUBSTRATE_TILE_TYPE_COMPLEX)
 #define SUBSTRATE_TILE_BITMASK_COMPLEXSPECIAL			(1u<<SUBSTRATE_TILE_TYPE_COMPLEX_SPECIAL)
+#define SUBSTRATE_TILE_BITMASK_ROUGH_REFRACT			(1u<<SUBSTRATE_TILE_TYPE_ROUGH_REFRACT)
+#define SUBSTRATE_TILE_BITMASK_ROUGH_REFRACT_SSS_WITHOUT	(1u<<SUBSTRATE_TILE_TYPE_ROUGH_REFRACT_SSS_WITHOUT)
+#define SUBSTRATE_TILE_BITMASK_DECAL_SIMPLE				(1u<<SUBSTRATE_TILE_TYPE_DECAL_SIMPLE)
+#define SUBSTRATE_TILE_BITMASK_DECAL_SINGLE				(1u<<SUBSTRATE_TILE_TYPE_DECAL_SINGLE)
+#define SUBSTRATE_TILE_BITMASK_DECAL_COMPLEX			(1u<<SUBSTRATE_TILE_TYPE_DECAL_COMPLEX)
+#define SUBSTRATE_TILE_BITMASK_DECAL					(SUBSTRATE_TILE_BITMASK_DECAL_SIMPLE|SUBSTRATE_TILE_BITMASK_DECAL_SINGLE|SUBSTRATE_TILE_BITMASK_DECAL_COMPLEX)
+
 
 #define SUBSTRATE_MATERIAL_TYPE_SIMPLE					0
 #define SUBSTRATE_MATERIAL_TYPE_SINGLE					1
diff --git a/Engine/Source/Runtime/Renderer/Private/Substrate/Substrate.cpp b/Engine/Source/Runtime/Renderer/Private/Substrate/Substrate.cpp
--- a/Engine/Source/Runtime/Renderer/Private/Substrate/Substrate.cpp
+++ b/Engine/Source/Runtime/Renderer/Private/Substrate/Substrate.cpp
@@ -1011,16 +1011,7 @@ class FSubstrateMaterialTileClassificationPassCS : public FGlobalShader
 
 	static EShaderPermutationPrecacheRequest ShouldPrecachePermutation(const FGlobalShaderPermutationParameters& Parameters) 
 	{ 
-		bool bUsed = ShouldCompilePermutation(Parameters);
-		if (bUsed)
-		{
-			FPermutationDomain PermutationVector(Parameters.PermutationId);
-			if (PermutationVector.Get<FWaveOps>() && !DoesRuntimeSupportWave64())
-			{
-				bUsed = false;
-			}
-		}
-		return bUsed ? EShaderPermutationPrecacheRequest::Precached : EShaderPermutationPrecacheRequest::NotUsed;
+		return ShouldCompilePermutation(Parameters) ? EShaderPermutationPrecacheRequest::Precached : EShaderPermutationPrecacheRequest::NotUsed;
 	}
 
 	static void ModifyCompilationEnvironment(const FGlobalShaderPermutationParameters& Parameters, FShaderCompilerEnvironment& OutEnvironment)
@@ -1666,7 +1657,8 @@ void AddSubstrateMaterialClassificationPass(FRDGBuilder& GraphBuilder, const FMi
 		}
 
 		// Our current classification require 64 waves
-		const bool bWaveOps = DoesRuntimeSupportWave64() && SubstrateSupportsWaveOps(Platform) != ERHIFeatureSupport::Unsupported;
+		const bool bWaveOps = SubstrateSupportsWaveOps(Platform) != ERHIFeatureSupport::Unsupported;
+		const bool bWaveOps64 = DoesRuntimeSupportWave64() && bWaveOps ;
 		
 		const FSubstrateViewData* SubstrateViewData = &View.SubstrateViewData;
 		const FSubstrateSceneData* SubstrateSceneData = View.SubstrateViewData.SceneData;
@@ -1735,7 +1727,7 @@ void AddSubstrateMaterialClassificationPass(FRDGBuilder& GraphBuilder, const FMi
 			auto MarkClosureTilePass = [&](ESubstrateTileType TileType)
 			{
 				FSubstrateClosureTilePassCS::FPermutationDomain PermutationVector;
-				PermutationVector.Set< FSubstrateClosureTilePassCS::FWaveOps >(bWaveOps);
+				PermutationVector.Set< FSubstrateClosureTilePassCS::FWaveOps >(bWaveOps64);
 				TShaderMapRef<FSubstrateClosureTilePassCS> ComputeShader(View.ShaderMap, PermutationVector);
 				FSubstrateClosureTilePassCS::FParameters* PassParameters = GraphBuilder.AllocParameters<FSubstrateClosureTilePassCS::FParameters>();
 				PassParameters->ViewUniformBuffer = View.ViewUniformBuffer;
@@ -1756,7 +1748,7 @@ void AddSubstrateMaterialClassificationPass(FRDGBuilder& GraphBuilder, const FMi
 
 				FComputeShaderUtils::AddPass(
 					GraphBuilder,
-					RDG_EVENT_NAME("Substrate::ClosureTileAndOffsets(%s - %s)", ToString(TileType), bWaveOps ? TEXT("Wave") : TEXT("SharedMemory")),
+					RDG_EVENT_NAME("Substrate::ClosureTileAndOffsets(%s - %s)", ToString(TileType), bWaveOps64 ? TEXT("Wave") : TEXT("SharedMemory")),
 					PassFlags,
 					ComputeShader,
 					PassParameters,
