Path Tracer: Fix WorldSmoothNormal being wrong when using WPO

In this case we can't properly account for the difference between shaded and unshaded normal because the material is responsible for re-orienting the normal.

The best we can do is to simply make WorldSmoothNormal = WorldNormal in this case which avoids the most objectionable shading artifacts

diff --git a/Engine/Shaders/Private/PathTracing/PathTracingMaterialHitShader.usf b/Engine/Shaders/Private/PathTracing/PathTracingMaterialHitShader.usf
--- a/Engine/Shaders/Private/PathTracing/PathTracingMaterialHitShader.usf
+++ b/Engine/Shaders/Private/PathTracing/PathTracingMaterialHitShader.usf
@@ -1343,6 +1343,22 @@ RAY_TRACING_ENTRY_CLOSEST_HIT(PathTracingMaterialCHS,
 	}
 #endif
 
+#if USES_WORLD_POSITION_OFFSET
+	// WPO moves the position which invalidates the normal. This means the only way to compute a correctly oriented normal is for the shader to orient it properly
+	// The path tracer normally keeps the WorldSmoothNormal and WorldNormal separately, but this is not possible in this case. To avoid issues from having these two
+	// normals being un-aligned with each other, simply copy WorldNormal to WorldSmoothNormal when WPO is in use.
+	// See repro in UE-343656
+	#if ALWAYS_EVALUATE_WORLD_POSITION_OFFSET
+		// always handle WPO
+	#else
+		// only if the primitive is using it
+	if ((PrimitiveFlags & PRIMITIVE_SCENE_DATA_FLAG_EVALUATE_WORLD_POSITION_OFFSET) != 0u)
+	#endif
+	{
+		Payload.WorldSmoothNormal = Payload.WorldNormal;
+	}
+#endif
+
 #if MATERIAL_VIRTUALTEXTURE_FEEDBACK
 	if (PackedPayload.IsCameraRay())
 	{
