Add code in SubstrateGetBSDFRoughness in order to retrieve legacy roughness value for SSS profile materials.

This is not enabled, but can be used to retrieve old behavior if needed. By default Substrate will continue to return the actual primary roughness

diff --git a/Engine/Shaders/Private/Substrate/Substrate.ush b/Engine/Shaders/Private/Substrate/Substrate.ush
--- a/Engine/Shaders/Private/Substrate/Substrate.ush
+++ b/Engine/Shaders/Private/Substrate/Substrate.ush
@@ -3761,7 +3761,19 @@ float SubstrateGetBSDFRoughness(in FSubstrateBSDF BSDF)
 {
 	switch (BSDF_GETTYPE(BSDF))
 	{
-	case SUBSTRATE_BSDF_TYPE_SLAB: 			return SLAB_ROUGHNESS(BSDF);
+	case SUBSTRATE_BSDF_TYPE_SLAB:
+		{
+			float Out = SLAB_ROUGHNESS(BSDF);
+			// In order to match legacy behavior in DebugViewMode and other places, this function should return the original roughness value for SSS Profile material.
+			// For simplicity, and clarity, we prefer to display the true primary roughness. Secondary specular roughness could be displayed with addition view mode if needed
+			#if 0
+			if (BSDF_GETSSSTYPE(BSDF) == SSS_TYPE_DIFFUSION_PROFILE)
+			{
+				Out = GetSubsurfaceProfileOriginalRoughness(SLAB_SSSPROFILEID(BSDF), Out, SLAB_SSSPROFILERADIUSSCALE(BSDF));
+			}
+			#endif
+			return Out;
+		}
 	case SUBSTRATE_BSDF_TYPE_HAIR: 			return HAIR_ROUGHNESS(BSDF);
 	case SUBSTRATE_BSDF_TYPE_EYE: 				return EYE_ROUGHNESS(BSDF);
 	case SUBSTRATE_BSDF_TYPE_SINGLELAYERWATER: return SLW_ROUGHNESS(BSDF);
