SubstrateReadGBufferBSDF optimization: Define out code for shading models that can't be active based on the current complexity class.

diff --git a/Engine/Shaders/Private/Substrate/SubstrateRead.ush b/Engine/Shaders/Private/Substrate/SubstrateRead.ush
--- a/Engine/Shaders/Private/Substrate/SubstrateRead.ush
+++ b/Engine/Shaders/Private/Substrate/SubstrateRead.ush
@@ -29,6 +29,7 @@ FSubstrateGBufferBSDF SubstrateReadGBufferBSDF(in FScreenSpaceData ScreenSpaceDa
 	float3x3 BSDFTangentBasis = (float3x3)0;
 	float BSDFAO = 0.0f;
 
+#if SUBSTRATE_COMPLEXPATH
 	if (GBuffer.Anisotropy != 0)
 	{
 		BSDFTangentBasis[0] = GBuffer.WorldTangent;
@@ -36,6 +37,7 @@ FSubstrateGBufferBSDF SubstrateReadGBufferBSDF(in FScreenSpaceData ScreenSpaceDa
 		BSDFTangentBasis[2] = GBuffer.WorldNormal;
 	}
 	else
+#endif
 	{
 		BSDFTangentBasis = GetTangentBasis(GBuffer.WorldNormal);
 	}
@@ -83,6 +85,7 @@ FSubstrateGBufferBSDF SubstrateReadGBufferBSDF(in FScreenSpaceData ScreenSpaceDa
 		uint   SharedLocalBasisIndex = 0;
 		uint   SharedLocalBasisTypes = 0;
 
+	#if SUBSTRATE_COMPLEXPATH
 		if (GBuffer.ShadingModelID == SHADINGMODELID_HAIR)
 		{
 			const float Scatter = GBuffer.Metallic;
@@ -111,7 +114,9 @@ FSubstrateGBufferBSDF SubstrateReadGBufferBSDF(in FScreenSpaceData ScreenSpaceDa
 			BSDF = GetSubstrateEyeBSDF(DiffuseAlbedo, F0, GBuffer.Roughness, IrisMask, IrisDistance, IrisNormal, IrisPlaneNormal, SSSProfileID, Emissive, SharedLocalBasisIndex).InlinedBSDF;
 		}
 		else
+	#endif
 		{
+		#if SUBSTRATE_SINGLEPATH || SUBSTRATE_COMPLEXPATH
 			if (GBuffer.ShadingModelID == SHADINGMODELID_CLOTH)
 			{
 				FuzzColor = ExtractSubsurfaceColor(GBuffer);
@@ -160,6 +165,7 @@ FSubstrateGBufferBSDF SubstrateReadGBufferBSDF(in FScreenSpaceData ScreenSpaceDa
 				SecondRoughness = GBuffer.CustomData.y;
 				SecondRoughnessAsSimpleClearCoat = (SecondRoughnessWeight > 0.0f || any(GBuffer.WorldNormal != ClearCoatBottomNormal)) ? 1.0f : 0.0f;
 			}
+		#endif
 
 			BSDF = InternalGetSubstrateSlabBSDF(
 				Footprint,
