Fix incorrect SSS profile material conversion with Substrate.

SSS profile was not correctly demoted when Opacity was below a certain threshold.

diff --git a/Engine/Shaders/Private/Substrate/SubstrateLegacyConversion.ush b/Engine/Shaders/Private/Substrate/SubstrateLegacyConversion.ush
--- a/Engine/Shaders/Private/Substrate/SubstrateLegacyConversion.ush
+++ b/Engine/Shaders/Private/Substrate/SubstrateLegacyConversion.ush
@@ -359,15 +359,16 @@ FSubstrateData SubstrateConvertLegacyMaterialDynamic(
 		LegacySlab.Metallic = Metallic;
 		LegacySlab.Roughness = Roughness;
 		LegacySlab.Anisotropy = Anisotropy;
+		// Only convert material to SSS profile if opacity threshold is sufficient
 		#if SUBSURFACE_PROFILE_OPACITY_THRESHOLD
 		if (Opacity > SSSS_OPACITY_THRESHOLD_EPS)
 		#endif
 		{
 			LegacySlab.SSSProfileID = SubSurfaceProfileId;
 			LegacySlab.SSSMFPScale = Opacity;
+			LegacySlab.SSSPhaseAniso = 0.93f;
+			LegacySlab.SSSType = SSS_TYPE_DIFFUSION;	//  SSS_TYPE_DIFFUSION_PROFILE will be selected in InternalGetSubstrateSlabBSDF based on SSSProfileID and SSSMFPScale.
 		}
-		LegacySlab.SSSPhaseAniso = 0.93f;
-		LegacySlab.SSSType = SSS_TYPE_DIFFUSION;	//  SSS_TYPE_DIFFUSION_PROFILE will be selected in InternalGetSubstrateSlabBSDF based on SSSProfileID and SSSMFPScale.
 		LegacySlab.Emissive = Emissive;
 
 		FinalWeight = 1.0f;
