Replace EUpscaleMethod::Bilinear with EUpscaleMethod::Area upscale for EPass::PrimaryUpscale in AddMobilePostProcessingPasses

Less blurry than Bilinear at low resolution percentage, and less aliased than Nearest at high resolution percentage.

diff --git a/Engine/Shaders/Private/PostProcessUpscale.usf b/Engine/Shaders/Private/PostProcessUpscale.usf
--- a/Engine/Shaders/Private/PostProcessUpscale.usf
+++ b/Engine/Shaders/Private/PostProcessUpscale.usf
@@ -15,6 +15,7 @@
 #define UPSCALE_METHOD_LANCZOS 5
 #define UPSCALE_METHOD_GAUSSIAN 6
 #define UPSCALE_METHOD_SMOOTHSTEP 7
+#define UPSCALE_METHOD_AREA 8
 
 #define METHOD_SET_SIMPLE 1
 #define METHOD_SET_COMPLEX 2
@@ -37,11 +38,14 @@
 	#define INCLUDE_UPSCALE_METHOD_GAUSSIAN 1
 #elif DIM_METHOD == UPSCALE_METHOD_SMOOTHSTEP
 	#define INCLUDE_UPSCALE_METHOD_SMOOTHSTEP 1
+#elif DIM_METHOD == UPSCALE_METHOD_AREA
+	#define INCLUDE_UPSCALE_METHOD_AREA 1
 // see GetMethodSet
 #elif DIM_METHODSET == METHOD_SET_SIMPLE
 	#define INCLUDE_UPSCALE_METHOD_NEAREST 1
 	#define INCLUDE_UPSCALE_METHOD_BILINEAR 1
 	#define INCLUDE_UPSCALE_METHOD_SMOOTHSTEP 1
+	#define INCLUDE_UPSCALE_METHOD_AREA 1
 #elif DIM_METHODSET == METHOD_SET_COMPLEX
 	#define INCLUDE_UPSCALE_METHOD_DIRECTIONAL 1
 	#define INCLUDE_UPSCALE_METHOD_CATMULL_ROM 1
@@ -93,6 +97,21 @@ float2 GetSmoothstepUV(float2 LinearUV, float2 TextureSize, float2 TextureInvSiz
 	return TextureInvSize * (TopLeftPixelCoord + SmoothPixelInterp + 0.5);
 }
 
+float2 GetAreaUV(float2 LinearCoord, float2 TextureSize, float2 TextureInvSize)
+{
+	float2 MinCoord = mad(-0.5, Output_ViewportSizeInverse, LinearCoord);
+	float2 MaxCoord = MinCoord + Output_ViewportSizeInverse;
+	float2 FirstIndex = floor(MinCoord * TextureSize);
+	float2 LastIndex = ceil(MaxCoord * TextureSize) - 1.;
+	float2 LastCoord = LastIndex * TextureInvSize;
+	float2 bSameTexel = step(LastIndex, FirstIndex);
+	float2 AdjustmentInTexels = lerp(
+		(LinearCoord - LastCoord) * Output_ViewportSize, // Fair filter
+		.5, // Sample center
+		bSameTexel);
+	return mad(TextureInvSize, AdjustmentInTexels, LastCoord);
+}
+
 float Luma(float3 Color) 
 {
 #if UE_LEGACY_LUMINANCE_FACTORS || WORKING_COLOR_SPACE_IS_SRGB
@@ -336,6 +355,13 @@ float4 Upsample(float2 InputUV, uint EyeIndex, uint Method)
 		OutColor = SampleSceneColorRGBA(GetSmoothstepUV(InputUV, Input_Extent, Input_ExtentInverse));
 	}
 	break;
+#endif
+#if INCLUDE_UPSCALE_METHOD_AREA
+	case UPSCALE_METHOD_AREA:
+	{
+		OutColor = SampleSceneColorRGBA(GetAreaUV(InputUV, Input_Extent, Input_ExtentInverse));
+	}
+	break;
 #endif
 	}
 
diff --git a/Engine/Source/Runtime/Renderer/Private/PostProcess/PostProcessUpscale.cpp b/Engine/Source/Runtime/Renderer/Private/PostProcess/PostProcessUpscale.cpp
--- a/Engine/Source/Runtime/Renderer/Private/PostProcess/PostProcessUpscale.cpp
+++ b/Engine/Source/Runtime/Renderer/Private/PostProcess/PostProcessUpscale.cpp
@@ -97,8 +97,10 @@ public:
 			return false;
 		}
 
-		// Always allow point and bilinear upscale. (Provides upscaling for mobile emulation)
-		if (UpscaleMethod == EUpscaleMethod::Nearest || UpscaleMethod == EUpscaleMethod::Bilinear)
+		// Always allow point and bilinear and area upscale. (Provides upscaling for mobile emulation)
+		if (UpscaleMethod == EUpscaleMethod::Nearest ||
+			UpscaleMethod == EUpscaleMethod::Bilinear ||
+			UpscaleMethod == EUpscaleMethod::Area)
 		{
 			return true;
 		}
@@ -151,6 +153,7 @@ EMethodSet GetMethodSet(EUpscaleMethod Method)
 		case EUpscaleMethod::Nearest:
 		case EUpscaleMethod::Bilinear:
 		case EUpscaleMethod::SmoothStep:
+		case EUpscaleMethod::Area:
 			return EMethodSet::Simple;
 		case EUpscaleMethod::Directional:
 		case EUpscaleMethod::CatmullRom:
diff --git a/Engine/Source/Runtime/Renderer/Private/PostProcess/PostProcessUpscale.h b/Engine/Source/Runtime/Renderer/Private/PostProcess/PostProcessUpscale.h
--- a/Engine/Source/Runtime/Renderer/Private/PostProcess/PostProcessUpscale.h
+++ b/Engine/Source/Runtime/Renderer/Private/PostProcess/PostProcessUpscale.h
@@ -15,6 +15,7 @@ enum class EUpscaleMethod : uint8
 	Lanczos,
 	Gaussian,
 	SmoothStep,
+	Area,
 	MAX
 };
 
diff --git a/Engine/Source/Runtime/Renderer/Private/PostProcess/PostProcessing.cpp b/Engine/Source/Runtime/Renderer/Private/PostProcess/PostProcessing.cpp
--- a/Engine/Source/Runtime/Renderer/Private/PostProcess/PostProcessing.cpp
+++ b/Engine/Source/Runtime/Renderer/Private/PostProcess/PostProcessing.cpp
@@ -3367,7 +3367,7 @@ void AddMobilePostProcessingPasses(FRDGBuilder& GraphBuilder, FScene* Scene, con
 		}
 		else
 		{
-			SceneColor = ISpatialUpscaler::AddDefaultUpscalePass(GraphBuilder, View, PassInputs, EUpscaleMethod::Bilinear, View.LensDistortionLUT);
+			SceneColor = ISpatialUpscaler::AddDefaultUpscalePass(GraphBuilder, View, PassInputs, EUpscaleMethod::Area, View.LensDistortionLUT);
 		}
 	}
 
