Fix incorrect VSM shadow with hair when using SMRT with SamplesPerRayHair set to 0.

diff --git a/Engine/Shaders/Private/VirtualShadowMaps/VirtualShadowMapProjection.usf b/Engine/Shaders/Private/VirtualShadowMaps/VirtualShadowMapProjection.usf
--- a/Engine/Shaders/Private/VirtualShadowMaps/VirtualShadowMapProjection.usf
+++ b/Engine/Shaders/Private/VirtualShadowMaps/VirtualShadowMapProjection.usf
@@ -380,7 +380,7 @@ FVirtualShadowMapSampleResult ProjectLight(
 		{
 			float AccShadowFactor = 0;
 			// TODO: Maybe do this min() in the uniform buffer parameter itself?
-			uint HairRayCount = min(VirtualShadowMap.SMRTHairRayCount, uint(GetSMRTRayCount()));
+			uint HairRayCount = max(1u, min(VirtualShadowMap.SMRTHairRayCount, uint(GetSMRTRayCount())));
 			uint RayIndex = 0; 
 			for (; RayIndex < HairRayCount; RayIndex++ )
 			{
