Fix for fog on cloud
is VolumetricRenderTarget.PreferAsyncCompute set on windows too?

#tests PC editor dx12 + console performance perf replay runs
#rnx

#DisabledValidations checkCopyright, checkUnacceptableWords

diff --git a/Engine/Shaders/Private/VolumetricRenderTarget.usf b/Engine/Shaders/Private/VolumetricRenderTarget.usf
--- a/Engine/Shaders/Private/VolumetricRenderTarget.usf
+++ b/Engine/Shaders/Private/VolumetricRenderTarget.usf
@@ -695,7 +695,7 @@ float4 ApplyFogToCloudSampleFromVPosAndDepth(float2 SVPos, float DepthKilometer,
 #endif
 
 #if PERMUTATION_APPLY_LOCAL_FOG_VOLUME
-	uint2 TilePos = SVPosition / LFVTilePixelSize.xx;
+	uint2 TilePos = (SVPosition - View.ViewRectMin.xy) / LFVTilePixelSize.xx; // Remove view rect min as we take the view pixel to sample the tiled LFV culling result
 	float4 LFVContribution = GetLFVContribution(PrimaryView, TilePos, FogSampleWorldPositionRelativeToCameraCm);
 	if (LFVRenderInVolumetricFog > 0)
 	{
