Substrate - fix for post process materials on mobile

#tests PC editor dx12 lyra and preview mobile renderer with substrate enabled. Bug reproduced and fix verified.
#rnx

No build patch diff
Original CL  https://horde.devtools.epicgames.com/job/68e4039ff63e2f8dde525bb6
Without CL  https://horde.devtools.epicgames.com/job/68e403f2a8471a142f3d2cc4

diff --git a/Engine/Shaders/Private/PostProcessMaterialShaders.usf b/Engine/Shaders/Private/PostProcessMaterialShaders.usf
--- a/Engine/Shaders/Private/PostProcessMaterialShaders.usf
+++ b/Engine/Shaders/Private/PostProcessMaterialShaders.usf
@@ -523,7 +523,7 @@ void MainPS(
 
 	float Alpha = 0.0f;
 	float3 EmissiveColor = 0.0f;
-#if SUBSTRATE_ENABLED
+#if SUBSTRATE_ENABLED && TEMPLATE_USES_SUBSTRATE
 
 	// Unlit forces a single BSDF
 	FSubstrateBSDF UnlitBSDF = PixelMaterialInputs.FrontMaterial.InlinedBSDF;
