Mark shader functions as deprecated for 5.7 as these function will no longer be used or will be replaced in the future when using Substrate.

diff --git a/Engine/Shaders/Private/CapsuleLightIntegrate.ush b/Engine/Shaders/Private/CapsuleLightIntegrate.ush
--- a/Engine/Shaders/Private/CapsuleLightIntegrate.ush
+++ b/Engine/Shaders/Private/CapsuleLightIntegrate.ush
@@ -104,7 +104,7 @@ FAreaLight CreateAreaLight(float Roughness, half3 N, half3 V, FCapsuleLight Caps
 	return AreaLight;
 }
 
-
+// UE_DEPRECATED 5.7 - Deprecated by Substrate
 FDirectLighting IntegrateBxDF( FGBufferData GBuffer, half3 N, half3 V, FCapsuleLight Capsule, FShadowTerms Shadow, bool bInverseSquared )
 {
 	GBuffer.Roughness = max( GBuffer.Roughness, View.MinRoughness );
@@ -112,6 +112,7 @@ FDirectLighting IntegrateBxDF( FGBufferData GBuffer, half3 N, half3 V, FCapsuleL
 	return IntegrateBxDF( GBuffer, N, V, AreaLight, Shadow );
 }
 
+// UE_DEPRECATED 5.7 - Deprecated by Substrate
 FDirectLighting IntegrateBxDF( FGBufferData GBuffer, half3 N, half3 V, FCapsuleLight Capsule, FShadowTerms Shadow, uint2 SVPos )
 {
 	FDirectLighting Lighting = (FDirectLighting)0;
diff --git a/Engine/Shaders/Private/DeferredLightingCommon.ush b/Engine/Shaders/Private/DeferredLightingCommon.ush
--- a/Engine/Shaders/Private/DeferredLightingCommon.ush
+++ b/Engine/Shaders/Private/DeferredLightingCommon.ush
@@ -313,6 +313,7 @@ FCapsuleLight GetCapsule( float3 ToLight, FDeferredLightData LightData )
 	return Capsule;
 }
 
+// UE_DEPRECATED 5.7 - Deprecated by Substrate - Uses SubstrateDeferredLighting() instead
 FLightAccumulator AccumulateDynamicLighting(
 	float3 TranslatedWorldPosition, half3 CameraVector, FGBufferData GBuffer, half AmbientOcclusion,
 	FDeferredLightData LightData, half4 LightAttenuation, float Dither, uint2 SVPos, 
@@ -456,6 +457,7 @@ FLightAccumulator AccumulateDynamicLighting(
 	return LightAccumulator;
 }
 
+// UE_DEPRECATED 5.7 - Deprecated by Substrate - Uses SubstrateDeferredLighting() instead
 /** Calculates lighting for a given position, normal, etc with a fully featured lighting model designed for quality. */
 FDeferredLightingSplit GetDynamicLightingSplit(
 	float3 TranslatedWorldPosition, float3 CameraVector, FGBufferData GBuffer, float AmbientOcclusion, 
@@ -466,6 +468,7 @@ FDeferredLightingSplit GetDynamicLightingSplit(
 	return LightAccumulator_GetResultSplit(LightAccumulator);
 }
 
+// UE_DEPRECATED 5.7 - Deprecated by Substrate - Uses SubstrateDeferredLighting() instead
 float4 GetDynamicLighting(
 	float3 TranslatedWorldPosition, float3 CameraVector, FGBufferData GBuffer, float AmbientOcclusion, 
 	FDeferredLightData LightData, float4 LightAttenuation, float Dither, uint2 SVPos, 
diff --git a/Engine/Shaders/Private/DeferredShadingCommon.ush b/Engine/Shaders/Private/DeferredShadingCommon.ush
--- a/Engine/Shaders/Private/DeferredShadingCommon.ush
+++ b/Engine/Shaders/Private/DeferredShadingCommon.ush
@@ -874,6 +874,7 @@ FGBufferData GetGBufferData(float2 UV, bool bGetNormalizedNormal = true, bool bF
 
 #endif //SHADING_PATH_MOBILE
 
+// UE_DEPRECATED 5.7 - Deprecated by Substrate
 /** Populates OutGBufferA, B and C */
 void EncodeGBuffer(
 	FGBufferData GBuffer,
@@ -980,6 +981,7 @@ void AdjustBaseColorAndSpecularColorForSubsurfaceProfileLighting(inout float3 Ba
 #endif
 }
 
+// UE_DEPRECATED 5.7 - Deprecated by Substrate
 /** Populates FGBufferData */
 // @param bChecker High frequency Checkerboard pattern computed with one of the CheckerFrom.. functions, todo: profile if float 0/1 would be better (need to make sure it's 100% the same)
 FGBufferData DecodeGBufferData(
@@ -1117,6 +1119,7 @@ float ConvertSubsurfaceProfileToFloat(uint InProfileIdUint)
 #if SHADING_PATH_DEFERRED
 
 #if FEATURE_LEVEL >= FEATURE_LEVEL_SM5
+	// UE_DEPRECATED 5.7 - Deprecated by Substrate
 	// @param PixelPos relative to left top of the rendertarget (not viewport)
 	FGBufferData GetGBufferDataUint(uint2 PixelPos, bool bGetNormalizedNormal = true)
 	{
@@ -1150,6 +1153,7 @@ float ConvertSubsurfaceProfileToFloat(uint InProfileIdUint)
 	#endif
 	}
 
+	// UE_DEPRECATED 5.7 - Deprecated by Substrate
 	// @param PixelPos relative to left top of the rendertarget (not viewport)
 	FScreenSpaceData GetScreenSpaceDataUint(uint2 PixelPos, bool bGetNormalizedNormal = true)
 	{
@@ -1164,6 +1168,7 @@ float ConvertSubsurfaceProfileToFloat(uint InProfileIdUint)
 	}
 #endif
 
+// UE_DEPRECATED 5.7 - Deprecated by Substrate
 // @param UV - UV space in the GBuffer textures (BufferSize resolution)
 // TOOD: need to find a way to make this more convenient.
 FGBufferData GetGBufferDataFromSceneTextures(float2 UV, bool bGetNormalizedNormal = true, bool bForceUnlitOrDefaultLit = false)
@@ -1201,6 +1206,7 @@ uint GetSceneLightingChannel(uint2 PixelCoord)
 	return ~0;
 }
 
+// UE_DEPRECATED 5.7 - Deprecated by Substrate
 // @param UV - UV space in the GBuffer textures (BufferSize resolution)
 FGBufferData GetGBufferData(float2 UV, bool bGetNormalizedNormal = true, bool bForceUnlitOrDefaultLit = false)
 {
@@ -1238,12 +1244,14 @@ FGBufferData GetGBufferData(float2 UV, bool bGetNormalizedNormal = true, bool bF
 #endif
 }
 
+// UE_DEPRECATED 5.7 - Deprecated by Substrate
 // Minimal path for just the lighting model, used to branch around unlit pixels (skybox)
 uint GetShadingModelId(float2 UV)
 {
 	return DecodeShadingModelId(Texture2DSampleLevel(SceneTexturesStruct.GBufferBTexture, SceneTexturesStruct_GBufferBTextureSampler, UV, 0).a);
 }
 
+// UE_DEPRECATED 5.7 - Deprecated by Substrate
 // @param UV - UV space in the GBuffer textures (BufferSize resolution)
 FScreenSpaceData GetScreenSpaceData(float2 UV, bool bGetNormalizedNormal = true, bool bForceUnlitOrDefaultLit = false)
 {
diff --git a/Engine/Shaders/Private/RayTracing/RayTracingSkyLightEvaluation.ush b/Engine/Shaders/Private/RayTracing/RayTracingSkyLightEvaluation.ush
--- a/Engine/Shaders/Private/RayTracing/RayTracingSkyLightEvaluation.ush
+++ b/Engine/Shaders/Private/RayTracing/RayTracingSkyLightEvaluation.ush
@@ -35,6 +35,8 @@ StructuredBuffer<FSkyLightVisibilityRays> SkyLightVisibilityRays;
 // bGBufferSampleOrigin - Indicates that the sample coordinate and world position were acquired from a g-buffer sample and that camera relative depth biasing should be used instead of just normal biased.
 // DeviceZ - The DeviceZ value from the g-buffer's depth buffer if bGBufferSampleOrigin is set to true, otherwise ignored.
 // bDecoupleSampleGeneration - Indicates weather or not to use decoupled sample generation for generating visibility rays, requires that the visibility ray buffer has been populated by the GenerateSkyLightVisibilityRaysCS pass.
+
+// UE_DEPRECATED 5.7 - Deprecated by Substrate
 void SkyLightEvaluate(
 	in const uint2 SampleCoord,
 	in const uint2 PixelCoord,
@@ -288,7 +290,7 @@ void SkyLightEvaluate(
 	}
 }
 
-
+// UE_DEPRECATED 5.7 - Deprecated by Substrate
 // special case of pixel coord == sample coord
 void SkyLightEvaluate(
 	in const uint2 SampleCoord,
diff --git a/Engine/Shaders/Private/RectLightIntegrate.ush b/Engine/Shaders/Private/RectLightIntegrate.ush
--- a/Engine/Shaders/Private/RectLightIntegrate.ush
+++ b/Engine/Shaders/Private/RectLightIntegrate.ush
@@ -73,6 +73,7 @@ FAreaLight CreateAreaLight( float Roughness, half3 N, half3 V, FRect Rect, FRect
 	return AreaLight;
 }
 
+// UE_DEPRECATED 5.7 - Deprecated by Substrate
 FDirectLighting IntegrateBxDF(FGBufferData GBuffer, half3 N, half3 V, FRect Rect, FShadowTerms Shadow, FRectTexture SourceTexture)
 {
 	// Compute distance to light plane and cull if not larger than threshold to 
@@ -90,6 +91,7 @@ FDirectLighting IntegrateBxDF(FGBufferData GBuffer, half3 N, half3 V, FRect Rect
 	return Out;
 }
 
+// UE_DEPRECATED 5.7 - Deprecated by Substrate
 FDirectLighting IntegrateBxDF( FGBufferData GBuffer, half3 N, half3 V, FRect Rect, FShadowTerms Shadow, FRectTexture SourceTexture, uint2 SVPos )
 {
 	FDirectLighting Lighting = (FDirectLighting)0;
diff --git a/Engine/Shaders/Private/ShadingModels.ush b/Engine/Shaders/Private/ShadingModels.ush
--- a/Engine/Shaders/Private/ShadingModels.ush
+++ b/Engine/Shaders/Private/ShadingModels.ush
@@ -97,6 +97,8 @@ struct FShadowTerms
 	half	TransmissionThickness;
 	FHairTransmittanceData HairTransmittance;
 };
+
+// UE_DEPRECATED 5.7 - Deprecated by Substrate
 FDirectLighting HairBxDF(FGBufferData GBuffer, half3 N, half3 V, FAreaLight AreaLight, FShadowTerms Shadow)
 {
 	const float3 BsdfValue = HairShading(GBuffer, AreaLight.DiffuseL, V, N, Shadow.TransmissionShadow, Shadow.HairTransmittance, HAIR_BSDF_BACKLIT, 0, uint2(0, 0));
@@ -149,6 +151,7 @@ float EnergyNormalization( inout float a2, float VoH, FAreaLight AreaLight )
 	return Energy;
 }
 
+// UE_DEPRECATED 5.7 - Deprecated by Substrate
 float3 SpecularGGX(float Roughness, float Anisotropy, float3 SpecularColor, BxDFContext Context, FAreaLight AreaLight)
 {
 	float Alpha = Roughness * Roughness;
@@ -164,6 +167,7 @@ float3 SpecularGGX(float Roughness, float Anisotropy, float3 SpecularColor, BxDF
 	return (D * Vis) * F;
 }
 
+// UE_DEPRECATED 5.7 - Deprecated by Substrate
 float3 SpecularGGX( float Roughness, float3 SpecularColor, BxDFContext Context, FAreaLight AreaLight )
 {
 	float a2 = Pow4( Roughness );
@@ -177,6 +181,7 @@ float3 SpecularGGX( float Roughness, float3 SpecularColor, BxDFContext Context,
 	return (D * Vis) * F;
 }
 
+// UE_DEPRECATED 5.7 - Deprecated by Substrate
 half3 DualSpecularGGX(half AverageRoughness, half Lobe0Roughness, half Lobe1Roughness, half LobeMix, half3 SpecularColor, BxDFContext Context, FAreaLight AreaLight)
 {
 	float AverageAlpha2 = Pow4(AverageRoughness);
@@ -194,6 +199,7 @@ half3 DualSpecularGGX(half AverageRoughness, half Lobe0Roughness, half Lobe1Roug
 	return (D * Vis) * F;
 }
 
+// UE_DEPRECATED 5.7 - Deprecated by Substrate
 FDirectLighting DefaultLitBxDF( FGBufferData GBuffer, half3 N, half3 V, FAreaLight AreaLight, FShadowTerms Shadow )
 {
 	BxDFContext Context;
@@ -281,6 +287,7 @@ FDirectLighting DefaultLitBxDF( FGBufferData GBuffer, half3 N, half3 V, FAreaLig
 	return Lighting;
 }
 
+// UE_DEPRECATED 5.7 - Deprecated by Substrate
 // Simple default lit model used by Lumen
 float3 SimpleShading( float3 DiffuseColor, float3 SpecularColor, float Roughness, float3 L, float3 V, half3 N )
 {
@@ -342,6 +349,7 @@ BxDFContext RefractClearCoatContext(BxDFContext Context)
 	return RefractedContext;
 }
 
+// UE_DEPRECATED 5.7 - Deprecated by Substrate
 FDirectLighting ClearCoatBxDF( FGBufferData GBuffer, half3 N, half3 V, FAreaLight AreaLight, FShadowTerms Shadow )
 {
 	const float ClearCoat			= GBuffer.CustomData.x;
@@ -575,6 +583,7 @@ void GetProfileDualSpecular(uint SubsurfaceProfileInt, half Roughness, half Opac
 #endif
 }
 
+// UE_DEPRECATED 5.7 - Deprecated by Substrate
 FDirectLighting SubsurfaceProfileBxDF( FGBufferData GBuffer, half3 N, half3 V, FAreaLight AreaLight, FShadowTerms Shadow )
 {
 	BxDFContext Context;
@@ -662,6 +671,7 @@ FDirectLighting SubsurfaceProfileBxDF( FGBufferData GBuffer, half3 N, half3 V, F
 	return Lighting;
 }
 
+// UE_DEPRECATED 5.7 - Deprecated by Substrate
 FDirectLighting ClothBxDF( FGBufferData GBuffer, half3 N, half3 V, FAreaLight AreaLight, FShadowTerms Shadow )
 {
 	const float3 FuzzColor	= ExtractSubsurfaceColor(GBuffer);
@@ -707,6 +717,7 @@ FDirectLighting ClothBxDF( FGBufferData GBuffer, half3 N, half3 V, FAreaLight Ar
 	return Lighting;
 }
 
+// UE_DEPRECATED 5.7 - Deprecated by Substrate
 FDirectLighting SubsurfaceBxDF(FGBufferData GBuffer, half3 N, half3 V, FAreaLight AreaLight, FShadowTerms Shadow )
 {
 	FDirectLighting Lighting = DefaultLitBxDF( GBuffer, N, V, AreaLight, Shadow);
@@ -834,6 +845,7 @@ FDirectLighting ShadeSGGX( FGBufferData GBuffer, float Geo_a2, half3 V, half3 L
 }
 #endif
 
+// UE_DEPRECATED 5.7 - Deprecated by Substrate
 FDirectLighting TwoSidedBxDF( FGBufferData GBuffer, half3 N, half3 V, FAreaLight AreaLight, FShadowTerms Shadow )
 {
 	half3 SubsurfaceColor = ExtractSubsurfaceColor(GBuffer);
@@ -947,6 +959,7 @@ FDirectLighting TwoSidedBxDF( FGBufferData GBuffer, half3 N, half3 V, FAreaLight
 	return Lighting;
 }
 
+// UE_DEPRECATED 5.7 - Deprecated by Substrate
 FDirectLighting EyeBxDF( FGBufferData GBuffer, half3 N, half3 V, FAreaLight AreaLight, FShadowTerms Shadow )
 {
 #if IRIS_NORMAL
@@ -1056,6 +1069,7 @@ FDirectLighting PreintegratedSkinBxDF( FGBufferData GBuffer, half3 N, half3 V, F
 	return Lighting;
 }
 
+// UE_DEPRECATED 5.7 - Deprecated by Substrate - Substrate uses SubstrateEvaluateBSDFCommon() instead
 FDirectLighting IntegrateBxDF( FGBufferData GBuffer, half3 N, half3 V, FAreaLight AreaLight, FShadowTerms Shadow )
 {
 	switch( GBuffer.ShadingModelID )
@@ -1085,6 +1099,7 @@ FDirectLighting IntegrateBxDF( FGBufferData GBuffer, half3 N, half3 V, FAreaLigh
 	}
 }
 
+// UE_DEPRECATED 5.7 - Deprecated by Substrate - Substrate uses SubstrateEvaluateBSDFCommon() instead
 FDirectLighting EvaluateBxDF( FGBufferData GBuffer, half3 N, half3 V, half3 L, float NoL, FShadowTerms Shadow )
 {
 	FAreaLight AreaLight = InitAreaLight();
diff --git a/Engine/Shaders/Private/ShadingModelsSampling.ush b/Engine/Shaders/Private/ShadingModelsSampling.ush
--- a/Engine/Shaders/Private/ShadingModelsSampling.ush
+++ b/Engine/Shaders/Private/ShadingModelsSampling.ush
@@ -128,6 +128,8 @@ FBxDFSample SampleDefaultLitBxDF(uint TermMask, float3 WorldNormal, float3x3 Tan
 
 	return BxDFSample;
 }
+
+// UE_DEPRECATED 5.7 - Deprecated by Substrate
 FBxDFSample SampleDefaultLitBxDF(uint TermMask, FGBufferData GBuffer, float3 V, float4 E) 
 { 
 	float3x3 TangentBasis;
@@ -145,6 +147,7 @@ FBxDFSample SampleDefaultLitBxDF(uint TermMask, FGBufferData GBuffer, float3 V,
 	return SampleDefaultLitBxDF(TermMask, GBuffer.WorldNormal, TangentBasis, GBuffer.Anisotropy, GBuffer.Roughness, V, E); 
 }
 
+// UE_DEPRECATED 5.7 - Deprecated by Substrate
 FBxDFSample SampleHairBxDF(uint TermMask, FGBufferData GBuffer, float3 V, float4 E)
 #if 1
 {
@@ -201,7 +204,7 @@ FBxDFSample SampleHairBxDF(uint TermMask, FGBufferData GBuffer, float3 V, float4
 }
 #endif
 
-// TODO: Oh god this feels duplicated with path tracer that is using ray tracing specific data structures :'(
+// UE_DEPRECATED 5.7 - Deprecated by Substrate
 FBxDFSample SampleBxDF(const uint TermMask, FGBufferData GBuffer, float3 V, float4 E)
 {
 	switch( GBuffer.ShadingModelID )
