Fix Lumen reflection with Splitscreen/VR when using Substrate's classification runs within StochasticLighting's classification pass.

diff --git a/Engine/Shaders/Private/StochasticLighting/StochasticLightingTileClassification.usf b/Engine/Shaders/Private/StochasticLighting/StochasticLightingTileClassification.usf
--- a/Engine/Shaders/Private/StochasticLighting/StochasticLightingTileClassification.usf
+++ b/Engine/Shaders/Private/StochasticLighting/StochasticLightingTileClassification.usf
@@ -18,6 +18,7 @@ MAX_OCCUPANCY
 
 #define DOWNSAMPLE_FACTOR	2
 #define TILE_SIZE			8
+#define TILE_SIZE_DIV_SHIFT	3
 
 #define STOCHASTIC_SAMPLE_OFFSET_2x1 1
 #define STOCHASTIC_SAMPLE_OFFSET_2x2 2
@@ -647,7 +648,8 @@ void StochasticLightingTileClassificationMarkCS(
 			const uint TileBitmask = ReadSharedTileBitmask(0, TILE_BITMASK_INDEX_SUBSTRATE);
 			if (TileBitmask != 0)
 			{
-				const uint EncodedTile = SubstratePackTile(Coord.SvPosition / TILE_SIZE, TileEncoding);
+				const uint2 TileCoord = (Coord.SvPosition.xy - uint2(View.ViewRectMin.xy)) >> TILE_SIZE_DIV_SHIFT;
+				const uint EncodedTile = SubstratePackTile(TileCoord, TileEncoding);
 
 				#if SUBSTRATE_COMPLEXSPECIALPATH
 				if (TileBitmask & SUBSTRATE_TILE_BITMASK_COMPLEXSPECIAL)
