[SingleLayerWater] Follow-up fix to 46630168, which made mobile SM5 SLW output to the GBuffer. That ended up skipping forward shading in MobileBasePassPixelShader.usf.

diff --git a/Engine/Shaders/Private/MobileBasePassPixelShader.usf b/Engine/Shaders/Private/MobileBasePassPixelShader.usf
--- a/Engine/Shaders/Private/MobileBasePassPixelShader.usf
+++ b/Engine/Shaders/Private/MobileBasePassPixelShader.usf
@@ -1076,7 +1076,9 @@ void Main(
 	GBuffer.IndirectIrradiance = IndirectIrradiance;
 	GBuffer.PerObjectGBufferData = GetPrimitive_PerObjectGBufferData(MaterialParameters.PrimitiveId);
 	MobileEncodeGBuffer(GBuffer, OutGBufferA, OutGBufferB, OutGBufferC, OutGBufferD);
-#else
+#endif
+
+#if !MOBILE_USE_GBUFFER || MATERIAL_SHADINGMODEL_SINGLELAYERWATER // With SM5 material support, SLW writes to the gbuffer but also still needs to evaluate this forward shading branch.
 
 	half3 CameraVector = -MaterialParameters.CameraVector;
 
@@ -1230,7 +1232,7 @@ void Main(
 	}
 #endif // MATERIAL_SHADINGMODEL_SINGLELAYERWATER
 
-#endif// DEFERRED_SHADING_PATH
+#endif // !MOBILE_USE_GBUFFER || MATERIAL_SHADINGMODEL_SINGLELAYERWATER
 
 #endif // TEMPLATE_USES_SUBSTRATE
 
