Fixed mobile SHADER_QUALITY for SSAO to match desktop renderer values.
Another platform was also using the same r.Mobile.AmbientOcclusionQuality setting. An A/B test showed a 0.17% dynamic resolution impact on that platform,
which we decided to accept in order to gain improved visual quality and maintain consistent SHADER_QUALITY settings across both desktop and mobile code paths.

diff --git a/Engine/Shaders/Private/PostProcessAmbientOcclusionMobile.usf b/Engine/Shaders/Private/PostProcessAmbientOcclusionMobile.usf
--- a/Engine/Shaders/Private/PostProcessAmbientOcclusionMobile.usf
+++ b/Engine/Shaders/Private/PostProcessAmbientOcclusionMobile.usf
@@ -22,8 +22,8 @@
 	#define QUAD_MESSAGE_PASSING_BLUR 2
 #elif SHADER_QUALITY == 2
 	// medium
-	#define USE_SAMPLESET 2
-	#define SAMPLE_STEPS 1
+	#define USE_SAMPLESET 1
+	#define SAMPLE_STEPS 2
 	#define QUAD_MESSAGE_PASSING_BLUR 2
 #elif SHADER_QUALITY == 3
 	// high
