Fixing Beer Shadow Map projection acne due to incorrect interpolation logic at edge conditions

diff --git a/Engine/Shaders/Private/HeterogeneousVolumes/HeterogeneousVolumesAdaptiveVolumetricShadowMapSampling.ush b/Engine/Shaders/Private/HeterogeneousVolumes/HeterogeneousVolumesAdaptiveVolumetricShadowMapSampling.ush
--- a/Engine/Shaders/Private/HeterogeneousVolumes/HeterogeneousVolumesAdaptiveVolumetricShadowMapSampling.ush
+++ b/Engine/Shaders/Private/HeterogeneousVolumes/HeterogeneousVolumesAdaptiveVolumetricShadowMapSampling.ush
@@ -96,7 +96,7 @@ float Interpolate(
 	float WorldHitT
 )
 {
-	float Transmittance = ShadowData[1].Tau;
+	float Transmittance = ShadowData[0].Tau;
 
 	if (INTERPOLANT_TYPE == INTERPOLANT_TYPE_LINEAR)
 	{
@@ -120,12 +120,9 @@ float Interpolate(
 	{
 		// Exponential fit
 		float DeltaX = max(ShadowData[1].X - ShadowData[0].X, 0.0);
-		float SigmaT = -SafeLog(ShadowData[1].Tau * SafeRcp(ShadowData[0].Tau)) * SafeRcp(DeltaX);
-		if (SigmaT > 0.0)
-		{
-			float X = clamp(WorldHitT - ShadowData[0].X, 0.0, DeltaX);
-			Transmittance = saturate(ShadowData[0].Tau * exp(-X * SigmaT));
-		}
+		float SigmaT = -log(ShadowData[1].Tau * SafeRcp(ShadowData[0].Tau)) * SafeRcp(DeltaX);
+		float X = clamp(WorldHitT - ShadowData[0].X, 0.0, DeltaX);
+		Transmittance = saturate(ShadowData[0].Tau * exp(-X * SigmaT));
 	}
 
 	return Transmittance;
@@ -713,11 +710,14 @@ float BeerShadowMap_SampleTransmittanceSimple(
 					ShadowZ = dot(ShadowMap.TranslatedWorldPlane[Face[Index]].xyz, TranslatedWorldPosition) + ShadowMap.TranslatedWorldPlane[Face[Index]].w;
 				}
 
-				FAVSMSampleData SampleData[] = {
-					AVSM_CreateSampleData(ShadowMapData.T[0], ShadowMapData.Tau[0]),
-					AVSM_CreateSampleData(ShadowMapData.T[1], ShadowMapData.Tau[1])
-				};
-				Transmittance[Index] = Interpolate(SampleData, ShadowZ);
+				if (ShadowZ > ShadowMapData.T[0])
+				{
+					FAVSMSampleData SampleData[] = {
+						AVSM_CreateSampleData(ShadowMapData.T[0], ShadowMapData.Tau[0]),
+						AVSM_CreateSampleData(ShadowMapData.T[1], ShadowMapData.Tau[1])
+					};
+					Transmittance[Index] = Interpolate(SampleData, ShadowZ);
+				}
 			}
 		}
 	}
@@ -815,11 +815,14 @@ float BeerShadowMap_SampleTransmittance(
 					ShadowZ = dot(ShadowMap.TranslatedWorldPlane[Face[Index]].xyz, TranslatedWorldPosition) + ShadowMap.TranslatedWorldPlane[Face[Index]].w;
 				}
 
-				FAVSMSampleData SampleData[] = {
-					AVSM_CreateSampleData(ShadowMapData.T[0], ShadowMapData.Tau[0]),
-					AVSM_CreateSampleData(ShadowMapData.T[1], ShadowMapData.Tau[1])
-				};
-				Transmittance[Index] = Interpolate(SampleData, ShadowZ);
+				if (ShadowZ > ShadowMapData.T[0])
+				{
+					FAVSMSampleData SampleData[] = {
+						AVSM_CreateSampleData(ShadowMapData.T[0], ShadowMapData.Tau[0]),
+						AVSM_CreateSampleData(ShadowMapData.T[1], ShadowMapData.Tau[1])
+					};
+					Transmittance[Index] = Interpolate(SampleData, ShadowZ);
+				}
 			}
 		}
 	}
