[Stochastic lighting mark] Fix a null shader resource assert when shader optimization is disabled

#tests tested in editor

diff --git a/Engine/Shaders/Private/StochasticLighting/StochasticLightingTileClassification.usf b/Engine/Shaders/Private/StochasticLighting/StochasticLightingTileClassification.usf
--- a/Engine/Shaders/Private/StochasticLighting/StochasticLightingTileClassification.usf
+++ b/Engine/Shaders/Private/StochasticLighting/StochasticLightingTileClassification.usf
@@ -888,9 +888,11 @@ void StochasticLightingTileClassificationMarkCS(
 						UnpackNormalAndShadingInfo(float4(0, 0, 0, PackedW.w)).bIsHair);
 					MegaLightsHistoryWeights *= select(Material.bIsHair.xxxx == bIsHair4, half(1.0), half(0.0));
 
+				#if REPROJECT_MEGALIGHTS
 					// History might have invalid lighting pixels. Only use history lighting if it has valid data.
 					half4 FramesAccumulatedHistory4 = MegaLightsNumFramesAccumulatedHistory.GatherRed(GlobalPointClampedSampler, HistoryGatherUV).wzxy;
 					MegaLightsHistoryWeights = select(FramesAccumulatedHistory4 > 0.0, MegaLightsHistoryWeights, 0.0);
+				#endif
 				}
 
 				// History neighbor valid mask for ScreenProbeTemporalReprojectionCS and DenoiserTemporalCS
