Fix incorrect hair strands culling for long strands.

#test Editor ContentExamples, QAGames

diff --git a/Engine/Shaders/Private/HairStrands/HairStrandsClusterCulling.usf b/Engine/Shaders/Private/HairStrands/HairStrandsClusterCulling.usf
--- a/Engine/Shaders/Private/HairStrands/HairStrandsClusterCulling.usf
+++ b/Engine/Shaders/Private/HairStrands/HairStrandsClusterCulling.usf
@@ -250,7 +250,6 @@ groupshared uint AllocationCount;
 
 // Store locally the culled point index (encoded into 8 bits, and store 4 per uint)
 groupshared uint CulledPointIndices[HAIR_MAX_NUM_POINT_PER_CURVE]; 
-groupshared uint ValidPoints[HAIR_MAX_NUM_POINT_PER_CURVE];
 
 // Version for curve with >32 points
 [numthreads(GROUP_SIZE, 1, 1)]
@@ -268,24 +267,15 @@ void Main(uint3 DispatchThreadId : SV_DispatchThreadID, uint ThreadIndex : SV_Gr
 	{
 		// 1.1 Cluster ID
 		// 1.2 Get cluster info (i.e. map all info per LOD for a given cluster ID) and derive effective LOD data
-		FHairClusterLOD ClusterLODInfo = (FHairClusterLOD)0;
-		if (bIsCurveValid)
-		{
-			const uint ClusterId = CurveToClusterIdBuffer[CurveIndex];
-			ClusterLODInfo = GetHairClusterLOD(ClusterInfoBuffer[ClusterId], ClusterInfoParameters, LODIndex);
-		}
+		const uint ClusterId = CurveToClusterIdBuffer[CurveIndex];
+		const FHairClusterLOD ClusterLODInfo = GetHairClusterLOD(ClusterInfoBuffer[ClusterId], ClusterInfoParameters, LODIndex);
 		ClusterLOD = ClusterLODInfo.LOD;
 		LODCount = ClusterLODInfo.LODCount;
 
 		// 1.4 Curve's point count/offset
-		CurvePointIndex = 0;
-		CurvePointCount = 0;
-		if (bIsCurveValid)
-		{
-			const FHairCurve Curve = ReadHairCurve(CurveBuffer, CurveIndex);
-			CurvePointIndex = Curve.PointIndex;
-			CurvePointCount = Curve.PointCount;
-		}
+		const FHairCurve Curve = ReadHairCurve(CurveBuffer, CurveIndex);
+		CurvePointIndex = Curve.PointIndex;
+		CurvePointCount = Curve.PointCount;
 
 		AllocationCount = 0;
 		AllocationOffset = 0;
@@ -301,19 +291,16 @@ void Main(uint3 DispatchThreadId : SV_DispatchThreadID, uint ThreadIndex : SV_Gr
 			const uint CurrentPointIndex = ThreadIndex + LoopIt * GROUP_SIZE;
 	
 			// 2.2 Mark point having their LOD visibility smaller than the requested LOD level
-			ValidPoints[CurrentPointIndex] = 0; 
-			if (CurveIndex < CurveCount)
-			{	
-				if (CurrentPointIndex < CurvePointCount)
+			bool bIsValid = false;
+			if (CurrentPointIndex < CurvePointCount)
+			{
+				const uint PointIndex = CurrentPointIndex + CurvePointIndex;
+				const uint MinLOD = GetHairControlPointMinLOD(PointIndex, PointLODBuffer);
+
+				// Prune control point based on their MinLOD value
+				if (IsHairControlPointActive(MinLOD, ClusterLOD))
 				{
-					const uint PointIndex = CurrentPointIndex + CurvePointIndex;
-					const uint MinLOD = GetHairControlPointMinLOD(PointIndex, PointLODBuffer);
-
-					// Prune control point based on their MinLOD value
-					if (IsHairControlPointActive(MinLOD, ClusterLOD))
-					{
-						ValidPoints[CurrentPointIndex] = 1;
-					}
+					bIsValid = true;
 				}
 			}
 		
@@ -321,8 +308,7 @@ void Main(uint3 DispatchThreadId : SV_DispatchThreadID, uint ThreadIndex : SV_Gr
 			{
 				uint IndexOffset = AllocationCount;
 				uint LocalAllocationCount = 0;
-				const uint bIsValid = ValidPoints[CurrentPointIndex];
-				const uint PrefixSum = ThreadGroupPrefixSum(bIsValid, ThreadIndex, LocalAllocationCount); // Has a ThreadMemoryGroupSync internally
+				const uint PrefixSum = ThreadGroupPrefixSum(uint(bIsValid), ThreadIndex, LocalAllocationCount); // Has a ThreadMemoryGroupSync internally
 				if (bIsValid)
 				{
 					CulledPointIndices[IndexOffset + PrefixSum] = CurrentPointIndex;
@@ -333,6 +319,8 @@ void Main(uint3 DispatchThreadId : SV_DispatchThreadID, uint ThreadIndex : SV_Gr
 					AllocationCount += LocalAllocationCount;
 				}
 			}
+
+			GroupMemoryBarrierWithGroupSync();
 		}
 	}
 
@@ -344,7 +332,6 @@ void Main(uint3 DispatchThreadId : SV_DispatchThreadID, uint ThreadIndex : SV_Gr
 	GroupMemoryBarrierWithGroupSync();
 
 	// 5. Write out index/scale
-	if (CurveIndex < CurveCount)
 	{
 		const uint LoopCount = DivideAndRoundUp(AllocationCount, GROUP_SIZE);
 		for (uint LoopIt = 0; LoopIt < LoopCount; ++LoopIt)
@@ -359,7 +346,7 @@ void Main(uint3 DispatchThreadId : SV_DispatchThreadID, uint ThreadIndex : SV_Gr
 	}
 
 	// 6. Write out curve data with visible points for subsequent passes
-	if (ThreadIndex == 0 && bIsCurveValid)
+	if (ThreadIndex == 0)
 	{	
 		FHairCurve CulledCurve;
 		CulledCurve.PointCount = AllocationCount;
