Heterogeneous Volumes: Adding bilinear depth query to Beer Shadow Maps

Optionally extending Beer Shadow Map precision to 32-bits to counteract depth precision issues in large assets (1k x 1k, 256).
r.HeterogeneousVolumes.Shadows.Precision
Determines the bit-depth precision for shadows
0: 16-bit (default)
1: 32-bit

diff --git a/Engine/Config/BaseScalability.ini b/Engine/Config/BaseScalability.ini
--- a/Engine/Config/BaseScalability.ini
+++ b/Engine/Config/BaseScalability.ini
@@ -769,6 +769,7 @@ r.Refraction.OffsetQuality=1
 r.HeterogeneousVolumes.DownsampleFactor=2
 r.HeterogeneousVolumes.MaxStepCount=256
 r.HeterogeneousVolumes.Shadows.Resolution=512
+r.HeterogeneousVolumes.Shadows.Precision=0
 r.HeterogeneousVolumes.UseExistenceMask=0
 
 [EffectsQuality@Cine]
@@ -804,6 +805,7 @@ r.HeterogeneousVolumes.MaxStepCount=1024
 r.HeterogeneousVolumes.Shadows.Type=0
 r.HeterogeneousVolumes.Shadows.Resolution=1024
 r.HeterogeneousVolumes.Shadows.MaxSampleCount=16
+r.HeterogeneousVolumes.Shadows.Precision=1
 r.HeterogeneousVolumes.UseExistenceMask=1
 r.Translucency.HeterogeneousVolumes.Type=0
 
diff --git a/Engine/Shaders/Private/HeterogeneousVolumes/HeterogeneousVolumesAdaptiveVolumetricShadowMapSampling.ush b/Engine/Shaders/Private/HeterogeneousVolumes/HeterogeneousVolumesAdaptiveVolumetricShadowMapSampling.ush
--- a/Engine/Shaders/Private/HeterogeneousVolumes/HeterogeneousVolumesAdaptiveVolumetricShadowMapSampling.ush
+++ b/Engine/Shaders/Private/HeterogeneousVolumes/HeterogeneousVolumesAdaptiveVolumetricShadowMapSampling.ush
@@ -662,6 +662,29 @@ FBeerShadowMap GetBeerShadowMap()
 	return ShadowMap;
 }
 
+float BeerShadowMap_Sample(
+	inout FBeerShadowMap ShadowMap,
+	float2 UV,
+	float ShadowZ
+)
+{
+	float Transmittance = 1.0;
+
+	// NOTE: Consider using point-sampler for depth and bilinear-sampler for transmittance to improve quality
+	float4 PackedData = Texture2DSample(ShadowMap.BeerShadowMapTexture, ShadowMap.TextureSampler, UV);
+	FBeerShadowMapData ShadowMapData = UnpackBeerShadowMapData(PackedData);
+	if (ShadowZ > ShadowMapData.T[0])
+	{
+		FAVSMSampleData SampleData[] = {
+			AVSM_CreateSampleData(ShadowMapData.T[0], ShadowMapData.Tau[0]),
+			AVSM_CreateSampleData(ShadowMapData.T[1], ShadowMapData.Tau[1])
+		};
+		Transmittance = Interpolate(SampleData, ShadowZ);
+	}
+
+	return Transmittance;
+}
+
 float BeerShadowMap_SampleTransmittanceSimple(
 	inout FBeerShadowMap ShadowMap,
 	float3 TranslatedWorldPosition,
@@ -701,23 +724,39 @@ float BeerShadowMap_SampleTransmittanceSimple(
 
 			if (all(ShadowPosition > ClipUV[0]) && all(ShadowPosition < ClipUV[1]))
 			{
-				float4 PackedData = Texture2DSample(ShadowMap.BeerShadowMapTexture, ShadowMap.TextureSampler, ShadowPosition.xy);
-				FBeerShadowMapData ShadowMapData = UnpackBeerShadowMapData(PackedData);
-
 				float ShadowZ = length(TranslatedWorldPosition - ShadowMap.TranslatedWorldOrigin[Face[Index]]);
 				if (ShadowMap.bIsDirectionalLight)
 				{
 					ShadowZ = dot(ShadowMap.TranslatedWorldPlane[Face[Index]].xyz, TranslatedWorldPosition) + ShadowMap.TranslatedWorldPlane[Face[Index]].w;
 				}
 
-				if (ShadowZ > ShadowMapData.T[0])
+#define BILINEAR_INTERPOLATION 1
+#if BILINEAR_INTERPOLATION
+				// Bilinear interpolation of shadow data
+				float2 ShadowMapCoord = ShadowPosition.xy * ShadowMap.Resolution;
+				int2 ShadowMapCoordX[4] =
+				{
+					clamp(ShadowMapCoord, 0, ShadowMap.Resolution),
+					clamp(ShadowMapCoord + int2(1, 0), 0, ShadowMap.Resolution),
+					clamp(ShadowMapCoord + int2(0, 1), 0, ShadowMap.Resolution),
+					clamp(ShadowMapCoord + int2(1, 1), 0, ShadowMap.Resolution)
+				};
+
+				float TransmittanceX[4] =
 				{
-					FAVSMSampleData SampleData[] = {
-						AVSM_CreateSampleData(ShadowMapData.T[0], ShadowMapData.Tau[0]),
-						AVSM_CreateSampleData(ShadowMapData.T[1], ShadowMapData.Tau[1])
-					};
-					Transmittance[Index] = Interpolate(SampleData, ShadowZ);
-				}
+					BeerShadowMap_Sample(ShadowMap, float2(ShadowMapCoordX[0]) / ShadowMap.Resolution, ShadowZ),
+					BeerShadowMap_Sample(ShadowMap, float2(ShadowMapCoordX[1]) / ShadowMap.Resolution, ShadowZ),
+					BeerShadowMap_Sample(ShadowMap, float2(ShadowMapCoordX[2]) / ShadowMap.Resolution, ShadowZ),
+					BeerShadowMap_Sample(ShadowMap, float2(ShadowMapCoordX[3]) / ShadowMap.Resolution, ShadowZ)
+				};
+
+				float2 Weight = frac(ShadowMapCoord);
+				float TransmittanceY0 = lerp(TransmittanceX[0], TransmittanceX[1], Weight.x);
+				float TransmittanceY1 = lerp(TransmittanceX[2], TransmittanceX[3], Weight.x);
+				Transmittance[Index] = lerp(TransmittanceY0, TransmittanceY1, Weight.y);
+#else
+				Transmittance[Index] = BeerShadowMap_Sample(ShadowMap, ShadowPosition.xy, ShadowZ);
+#endif
 			}
 		}
 	}
@@ -806,23 +845,41 @@ float BeerShadowMap_SampleTransmittance(
 
 			if (all(ShadowPosition > ClipUV[0]) && all(ShadowPosition < ClipUV[1]))
 			{
-				float4 PackedData = Texture2DSample(ShadowMap.BeerShadowMapTexture, ShadowMap.TextureSampler, ShadowPosition.xy);
-				FBeerShadowMapData ShadowMapData = UnpackBeerShadowMapData(PackedData);
-
 				float ShadowZ = length(TranslatedWorldPosition - ShadowMap.TranslatedWorldOrigin[Face[Index]]);
 				if (ShadowMap.bIsDirectionalLight)
 				{
 					ShadowZ = dot(ShadowMap.TranslatedWorldPlane[Face[Index]].xyz, TranslatedWorldPosition) + ShadowMap.TranslatedWorldPlane[Face[Index]].w;
 				}
 
-				if (ShadowZ > ShadowMapData.T[0])
+#define BILINEAR_INTERPOLATION 1
+#if BILINEAR_INTERPOLATION
+				int2 AtlasResolution = ShadowMap.Resolution * int2(BeerShadowMap.NumShadowMatrices, 1);
+
+				// Bilinear interpolation of shadow data
+				float2 ShadowMapCoord = ShadowPosition.xy * AtlasResolution;
+				int2 ShadowMapCoordX[4] =
 				{
-					FAVSMSampleData SampleData[] = {
-						AVSM_CreateSampleData(ShadowMapData.T[0], ShadowMapData.Tau[0]),
-						AVSM_CreateSampleData(ShadowMapData.T[1], ShadowMapData.Tau[1])
-					};
-					Transmittance[Index] = Interpolate(SampleData, ShadowZ);
-				}
+					clamp(ShadowMapCoord, 0, AtlasResolution),
+					clamp(ShadowMapCoord + int2(1, 0), 0, AtlasResolution),
+					clamp(ShadowMapCoord + int2(0, 1), 0, AtlasResolution),
+					clamp(ShadowMapCoord + int2(1, 1), 0, AtlasResolution)
+				};
+
+				float TransmittanceX[4] =
+				{
+					BeerShadowMap_Sample(ShadowMap, float2(ShadowMapCoordX[0]) / AtlasResolution, ShadowZ),
+					BeerShadowMap_Sample(ShadowMap, float2(ShadowMapCoordX[1]) / AtlasResolution, ShadowZ),
+					BeerShadowMap_Sample(ShadowMap, float2(ShadowMapCoordX[2]) / AtlasResolution, ShadowZ),
+					BeerShadowMap_Sample(ShadowMap, float2(ShadowMapCoordX[3]) / AtlasResolution, ShadowZ)
+				};
+
+				float2 Weight = frac(ShadowMapCoord);
+				float TransmittanceY0 = lerp(TransmittanceX[0], TransmittanceX[1], Weight.x);
+				float TransmittanceY1 = lerp(TransmittanceX[2], TransmittanceX[3], Weight.x);
+				Transmittance[Index] = lerp(TransmittanceY0, TransmittanceY1, Weight.y);
+#else
+				Transmittance[Index] = BeerShadowMap_Sample(ShadowMap, ShadowPosition.xy, ShadowZ);
+#endif
 			}
 		}
 	}
diff --git a/Engine/Source/Runtime/Renderer/Private/HeterogeneousVolumes/HeterogeneousVolumesLiveShadingPipeline.cpp b/Engine/Source/Runtime/Renderer/Private/HeterogeneousVolumes/HeterogeneousVolumesLiveShadingPipeline.cpp
--- a/Engine/Source/Runtime/Renderer/Private/HeterogeneousVolumes/HeterogeneousVolumesLiveShadingPipeline.cpp
+++ b/Engine/Source/Runtime/Renderer/Private/HeterogeneousVolumes/HeterogeneousVolumesLiveShadingPipeline.cpp
@@ -2378,6 +2378,15 @@ TArray<FVector> TransformVertices(const TArray<FVector>& Vertices, const FMatrix
 	return ToVertices;
 }
 
+static TAutoConsoleVariable<int32> CVarHeterogeneousVolumesShadowsPrecision(
+	TEXT("r.HeterogeneousVolumes.Shadows.Precision"),
+	0,
+	TEXT("Determines the bit-depth precision for shadows")
+	TEXT("0: 16-bit (default)")
+	TEXT("1: 32-bit"),
+	ECVF_Scalability | ECVF_RenderThreadSafe
+);
+
 bool RenderVolumetricShadowMapForLightWithLiveShading(
 	FRDGBuilder& GraphBuilder,
 	// Scene data
@@ -2830,6 +2839,10 @@ bool RenderVolumetricShadowMapForLightWithLiveShading(
 
 		FRDGTextureDesc Desc = SceneTextures.Color.Target->Desc;
 		Desc.Format = PF_FloatRGBA;
+		if (CVarHeterogeneousVolumesShadowsPrecision.GetValueOnRenderThread() == 1)
+		{
+			Desc.Format = PF_A32B32G32R32F;
+		}
 		Desc.Flags &= ~(TexCreate_FastVRAM);
 		Desc.Extent = FIntPoint(ShadowMapResolution.X * GroupCount.Z, ShadowMapResolution.Y);
 		BeerShadowMapTexture = GraphBuilder.CreateTexture(Desc, TEXT("HeterogeneousVolumes.Light.BeerShadowMapTexture"));
@@ -3064,51 +3077,58 @@ void RenderAdaptiveVolumetricShadowMapWithLiveShading(
 
 				if (bIsCreated)
 				{
-					FRDGBufferRef VolumetricShadowMapIndirectionBuffer;
-					FRDGBufferRef VolumetricShadowMapSampleBuffer;
-					CompressVolumetricShadowMap(
-						GraphBuilder,
-						View,
-						GroupCount,
-						VolumetricShadowMapResolution,
-						VolumetricShadowMapMaxSampleCount,
-						VolumetricShadowMapLinkedListBuffer,
-						VolumetricShadowMapIndirectionBuffer,
-						VolumetricShadowMapSampleBuffer
-					);
-
 					float DownsampleFactor = 1.0f;
-					CreateAdaptiveVolumetricShadowMapUniformBuffer(
-						GraphBuilder,
-						TranslatedWorldOrigin,
-						TranslatedWorldPlane,
-						TranslatedWorldToShadow,
-						SplitDepths,
-						HeterogeneousVolumes::GetShadowType() == HeterogeneousVolumes::EShadowType::AdaptiveVolumetricShadowMap ? VolumetricShadowMapResolution : FIntPoint::ZeroValue,
-						DownsampleFactor,
-						NumShadowMatrices,
-						VolumetricShadowMapMaxSampleCount,
-						bIsDirectionalLight,
-						VolumetricShadowMapLinkedListBuffer,
-						VolumetricShadowMapIndirectionBuffer,
-						VolumetricShadowMapSampleBuffer,
-						AdaptiveVolumetricShadowMapUniformBuffer
-					);
-
-					CreateBeerShadowMapUniformBuffer(
-						GraphBuilder,
-						TranslatedWorldOrigin,
-						TranslatedWorldPlane,
-						TranslatedWorldToShadow,
-						SplitDepths,
-						HeterogeneousVolumes::GetShadowType() == HeterogeneousVolumes::EShadowType::BeerShadowMap ? VolumetricShadowMapResolution : FIntPoint::ZeroValue,
-						DownsampleFactor,
-						NumShadowMatrices,
-						VolumetricShadowMapMaxSampleCount,
-						bIsDirectionalLight,
-						BeerShadowMapTexture,
-						BeerShadowMapUniformBuffer
-					);
+					if (HeterogeneousVolumes::GetShadowType() == HeterogeneousVolumes::EShadowType::AdaptiveVolumetricShadowMap)
+					{
+						FRDGBufferRef VolumetricShadowMapIndirectionBuffer;
+						FRDGBufferRef VolumetricShadowMapSampleBuffer;
+						CompressVolumetricShadowMap(
+							GraphBuilder,
+							View,
+							GroupCount,
+							VolumetricShadowMapResolution,
+							VolumetricShadowMapMaxSampleCount,
+							VolumetricShadowMapLinkedListBuffer,
+							VolumetricShadowMapIndirectionBuffer,
+							VolumetricShadowMapSampleBuffer
+						);
+
+						CreateAdaptiveVolumetricShadowMapUniformBuffer(
+							GraphBuilder,
+							TranslatedWorldOrigin,
+							TranslatedWorldPlane,
+							TranslatedWorldToShadow,
+							SplitDepths,
+							HeterogeneousVolumes::GetShadowType() == HeterogeneousVolumes::EShadowType::AdaptiveVolumetricShadowMap ? VolumetricShadowMapResolution : FIntPoint::ZeroValue,
+							DownsampleFactor,
+							NumShadowMatrices,
+							VolumetricShadowMapMaxSampleCount,
+							bIsDirectionalLight,
+							VolumetricShadowMapLinkedListBuffer,
+							VolumetricShadowMapIndirectionBuffer,
+							VolumetricShadowMapSampleBuffer,
+							AdaptiveVolumetricShadowMapUniformBuffer
+						);
+						BeerShadowMapUniformBuffer = HeterogeneousVolumes::CreateEmptyBeerShadowMapUniformBuffer(GraphBuilder);
+					}
+					else
+					{
+						CreateBeerShadowMapUniformBuffer(
+							GraphBuilder,
+							TranslatedWorldOrigin,
+							TranslatedWorldPlane,
+							TranslatedWorldToShadow,
+							SplitDepths,
+							HeterogeneousVolumes::GetShadowType() == HeterogeneousVolumes::EShadowType::BeerShadowMap ? VolumetricShadowMapResolution : FIntPoint::ZeroValue,
+							DownsampleFactor,
+							NumShadowMatrices,
+							VolumetricShadowMapMaxSampleCount,
+							bIsDirectionalLight,
+							BeerShadowMapTexture,
+							BeerShadowMapUniformBuffer
+						);
+						AdaptiveVolumetricShadowMapUniformBuffer = HeterogeneousVolumes::CreateEmptyAdaptiveVolumetricShadowMapUniformBuffer(GraphBuilder);
+					}
 				}
 				else
 				{
