Additional support on mobile platforms with SM5 materials
- Enable extended GBuffer and ClearCoat bottom layer normal
- Enable capsule light for point light

#rnx
#tests preflight, ran xb1 with odsc, A/B on cosmetics

diff --git a/Engine/Shaders/Private/DeferredLightingCommon.ush b/Engine/Shaders/Private/DeferredLightingCommon.ush
--- a/Engine/Shaders/Private/DeferredLightingCommon.ush
+++ b/Engine/Shaders/Private/DeferredLightingCommon.ush
@@ -313,6 +313,10 @@ FCapsuleLight GetCapsule( float3 ToLight, FDeferredLightData LightData )
 	return Capsule;
 }
 
+#ifndef MOBILE_SHADING_PATH_SUPPORT_CAPSULE_LIGHT
+#define MOBILE_SHADING_PATH_SUPPORT_CAPSULE_LIGHT 0
+#endif
+
 // UE_DEPRECATED 5.7 - Deprecated by Substrate - Uses SubstrateDeferredLighting() instead
 FLightAccumulator AccumulateDynamicLighting(
 	float3 TranslatedWorldPosition, half3 CameraVector, FGBufferData GBuffer, half AmbientOcclusion,
@@ -379,7 +383,17 @@ FLightAccumulator AccumulateDynamicLighting(
 		}
 		else
 		{
+		#if MOBILE_SHADING_PATH_SUPPORT_CAPSULE_LIGHT
+			FCapsuleLight Capsule = GetCapsule( ToLight, LightData );
+
+			#if REFERENCE_QUALITY
+				Lighting = IntegrateBxDF( GBuffer, N, V, Capsule, Shadow, SVPos );
+			#else
+				Lighting = IntegrateBxDF( GBuffer, N, V, Capsule, Shadow, LightData.bInverseSquared );
+			#endif
+		#else 
 			Lighting = EvaluateBxDF(GBuffer, N, V, L, NoL, Shadow);
+		#endif
 		}
 
 		Lighting.Specular *= LightData.SpecularScale;
diff --git a/Engine/Source/Runtime/Engine/Private/ShaderCompiler/ShaderCompiler.cpp b/Engine/Source/Runtime/Engine/Private/ShaderCompiler/ShaderCompiler.cpp
--- a/Engine/Source/Runtime/Engine/Private/ShaderCompiler/ShaderCompiler.cpp
+++ b/Engine/Source/Runtime/Engine/Private/ShaderCompiler/ShaderCompiler.cpp
@@ -3727,7 +3727,8 @@ void GlobalBeginCompileShader(
 
 	{
 		static const auto CVar = IConsoleManager::Get().FindTConsoleVariableDataInt(TEXT("r.ClearCoatNormal"));
-		SET_SHADER_DEFINE(Input.Environment, CLEAR_COAT_BOTTOM_NORMAL, CVar ? (CVar->GetValueOnAnyThread() != 0) && !bIsMobilePlatform : 0);
+		bool bCCBN = CVar && CVar->GetValueOnAnyThread() != 0 && (!bIsMobilePlatform || MobileSupportsSM5MaterialNodes((EShaderPlatform)Target.Platform));
+		SET_SHADER_DEFINE(Input.Environment, CLEAR_COAT_BOTTOM_NORMAL, bCCBN ? 1 : 0);
 	}
 
 	{
diff --git a/Engine/Source/Runtime/Engine/Private/ShaderCompiler/ShaderGenerationUtil.cpp b/Engine/Source/Runtime/Engine/Private/ShaderCompiler/ShaderGenerationUtil.cpp
--- a/Engine/Source/Runtime/Engine/Private/ShaderCompiler/ShaderGenerationUtil.cpp
+++ b/Engine/Source/Runtime/Engine/Private/ShaderCompiler/ShaderGenerationUtil.cpp
@@ -329,7 +329,7 @@ static FShaderGlobalDefines FetchShaderGlobalDefines(EShaderPlatform TargetPlatf
 
 	{
 		static const auto CVar = IConsoleManager::Get().FindTConsoleVariableDataInt(TEXT("r.ClearCoatNormal"));
-		Ret.CLEAR_COAT_BOTTOM_NORMAL = CVar ? (CVar->GetValueOnAnyThread() != 0) && !bIsMobilePlatform : 0;
+		Ret.CLEAR_COAT_BOTTOM_NORMAL = CVar ? (CVar->GetValueOnAnyThread() != 0 && (!bIsMobilePlatform || MobileSupportsSM5MaterialNodes(TargetPlatform))) : 0;
 	}
 
 	{
diff --git a/Engine/Source/Runtime/RenderCore/Private/RenderUtils.cpp b/Engine/Source/Runtime/RenderCore/Private/RenderUtils.cpp
--- a/Engine/Source/Runtime/RenderCore/Private/RenderUtils.cpp
+++ b/Engine/Source/Runtime/RenderCore/Private/RenderUtils.cpp
@@ -602,7 +602,7 @@ RENDERCORE_API bool MobileUsesExtenedGBuffer(FStaticShaderPlatform ShaderPlatfor
 		// maxDescriptorSetInputAttachments > 4
 		// maxColorAttachments > 4
 	// iOS: A8+
-	return (ShaderPlatform != SP_OPENGL_ES3_1_ANDROID) && false;
+	return (ShaderPlatform != SP_OPENGL_ES3_1_ANDROID) && MobileSupportsSM5MaterialNodes(ShaderPlatform);
 }
 
 // Required for shading models with a custom data
diff --git a/Engine/Source/Runtime/RenderCore/Private/Shader.cpp b/Engine/Source/Runtime/RenderCore/Private/Shader.cpp
--- a/Engine/Source/Runtime/RenderCore/Private/Shader.cpp
+++ b/Engine/Source/Runtime/RenderCore/Private/Shader.cpp
@@ -1736,7 +1736,7 @@ void ShaderMapAppendKey(EShaderPlatform Platform, FShaderKeyGenerator& KeyGen)
 	{
 		static const auto CVar = IConsoleManager::Get().FindTConsoleVariableDataInt(TEXT("r.ClearCoatNormal"));
 		KeyGen.AppendSeparator();
-		KeyGen.Append((CVar && CVar->GetValueOnAnyThread() != 0) ? TEXT("CCBN") : TEXT("NoCCBN"));
+		KeyGen.Append((CVar && CVar->GetValueOnAnyThread() != 0 && (!IsMobilePlatform(Platform) || MobileSupportsSM5MaterialNodes(Platform))) ? TEXT("CCBN") : TEXT("NoCCBN"));
 	}
 
 	{
diff --git a/Engine/Source/Runtime/Renderer/Private/MobileDeferredShadingPass.cpp b/Engine/Source/Runtime/Renderer/Private/MobileDeferredShadingPass.cpp
--- a/Engine/Source/Runtime/Renderer/Private/MobileDeferredShadingPass.cpp
+++ b/Engine/Source/Runtime/Renderer/Private/MobileDeferredShadingPass.cpp
@@ -280,6 +280,9 @@ public:
 
 		const bool bMobileForceDepthRead = MobileUsesFullDepthPrepass(Parameters.Platform);
 		OutEnvironment.SetDefine(TEXT("FORCE_DEPTH_TEXTURE_READS"), bMobileForceDepthRead ? 1u : 0u);
+
+		const bool bSupportCapsule = MobileSupportsSM5MaterialNodes(Parameters.Platform);
+		OutEnvironment.SetDefine(TEXT("MOBILE_SHADING_PATH_SUPPORT_CAPSULE_LIGHT"), bSupportCapsule ? 1u : 0u);
 	}
 
 	void SetParameters(FRHIBatchedShaderParameters& BatchedParameters, const FViewInfo& View, const FMaterialRenderProxy* Proxy, const FMaterial& Material)
