Integrate point light as capsule for translucencies object on mobile renderer when SM5 material node is enabled

#rnx
#tests kicked a limited build on XB1, pulled the build and played the game

diff --git a/Engine/Shaders/Private/DeferredLightingCommon.ush b/Engine/Shaders/Private/DeferredLightingCommon.ush
--- a/Engine/Shaders/Private/DeferredLightingCommon.ush
+++ b/Engine/Shaders/Private/DeferredLightingCommon.ush
@@ -313,10 +313,6 @@ FCapsuleLight GetCapsule( float3 ToLight, FDeferredLightData LightData )
 	return Capsule;
 }
 
-#ifndef MOBILE_SHADING_PATH_SUPPORT_CAPSULE_LIGHT
-#define MOBILE_SHADING_PATH_SUPPORT_CAPSULE_LIGHT 0
-#endif
-
 // UE_DEPRECATED 5.7 - Deprecated by Substrate - Uses SubstrateDeferredLighting() instead
 FLightAccumulator AccumulateDynamicLighting(
 	float3 TranslatedWorldPosition, half3 CameraVector, FGBufferData GBuffer, half AmbientOcclusion,
diff --git a/Engine/Shaders/Private/LightData.ush b/Engine/Shaders/Private/LightData.ush
--- a/Engine/Shaders/Private/LightData.ush
+++ b/Engine/Shaders/Private/LightData.ush
@@ -14,6 +14,10 @@
 
 #define MAX_RECT_ATLAS_MIP			32
 
+#ifndef MOBILE_SHADING_PATH_SUPPORT_CAPSULE_LIGHT
+#define MOBILE_SHADING_PATH_SUPPORT_CAPSULE_LIGHT 0
+#endif
+
 ///////////////////////////////////////////////////////////////////////////////////////////////////
 // Rect light texture data
 struct FRectTexture
@@ -605,6 +609,9 @@ FDeferredLightData ConvertToDeferredLight(const FLocalLightData In)
 
 FDeferredLightData ConvertToDeferredLight_Mobile(const FLocalLightData In)
 {
+#if MOBILE_SHADING_PATH_SUPPORT_CAPSULE_LIGHT
+	return ConvertToDeferredLight(In);
+#else
 	half4 OutPreviewShadowMapChannelMask = 0;
 	uint OutLightingChannelMask = LIGHTING_CHANNEL_MASK;
 	FDeferredLightData Out = ConvertToDeferredLight(In, 1.0f, OutPreviewShadowMapChannelMask, OutLightingChannelMask);	
@@ -613,6 +620,7 @@ FDeferredLightData ConvertToDeferredLight_Mobile(const FLocalLightData In)
 	Out.SoftSourceRadius= 0;
 	Out.SpecularScale = 1.0f; // If enabled -> UnpackLightSpecularScale(In);
 	Out.DiffuseScale  = 1.0f; // If enabled -> UnpackLightDiffuseScale(In);
+
 #if RECT_LIGHT_AS_SPOTLIGHT
 	// Render RectLight approximately as SpotLight
 	Out.bSpotLight	= Out.bSpotLight || Out.bRectLight;
@@ -628,6 +636,7 @@ FDeferredLightData ConvertToDeferredLight_Mobile(const FLocalLightData In)
 	}
 
 	return Out;
+#endif // MOBILE_SHADING_PATH_SUPPORT_CAPSULE_LIGHT
 }
 
 FLightShaderParameters ConvertToLightShaderParameters(FDeferredLightData In)
diff --git a/Engine/Source/Runtime/Renderer/Private/MobileBasePassRendering.h b/Engine/Source/Runtime/Renderer/Private/MobileBasePassRendering.h
--- a/Engine/Source/Runtime/Renderer/Private/MobileBasePassRendering.h
+++ b/Engine/Source/Runtime/Renderer/Private/MobileBasePassRendering.h
@@ -439,6 +439,9 @@ public:
 		OutEnvironment.SetDefine(TEXT("MOBILE_TRANSLUCENT_COLOR_TRANSMITTANCE_PROGRAMMABLE_BLENDING"), TranslucentColorTransmittanceMode == EMobileTranslucentColorTransmittanceMode::PROGRAMMABLE_BLENDING ? 1u : 0u);
 		OutEnvironment.SetDefine(TEXT("MOBILE_TRANSLUCENT_COLOR_TRANSMITTANCE_SINGLE_SRC_BLENDING"), TranslucentColorTransmittanceMode == EMobileTranslucentColorTransmittanceMode::SINGLE_SRC_BLENDING ? 1u : 0u);
 		OutEnvironment.SetDefine(TEXT("MOBILE_SSR_ENABLED"), ((!bDeferredShadingEnabled && AreMobileScreenSpaceReflectionsEnabled(Parameters.Platform)) || (bIsSLW && MobileSupportsSM5MaterialNodes(Parameters.Platform))) ? 1u : 0u);
+
+		const bool bSupportCapsule = MobileSupportsSM5MaterialNodes(Parameters.Platform);
+		OutEnvironment.SetDefine(TEXT("MOBILE_SHADING_PATH_SUPPORT_CAPSULE_LIGHT"), bSupportCapsule ? 1u : 0u);
 	}
 	
 	/** Initialization constructor. */
