[DOF] Fix artifacts around edges caused by issues in the reduce pass when producing mip levels for the background. These are caused by incorrect assumptions made in the CONFIG_WAVE_BROADCAST_REDUCTION path, which this CL disables.

diff --git a/Engine/Shaders/Private/DiaphragmDOF/DOFReduce.usf b/Engine/Shaders/Private/DiaphragmDOF/DOFReduce.usf
--- a/Engine/Shaders/Private/DiaphragmDOF/DOFReduce.usf
+++ b/Engine/Shaders/Private/DiaphragmDOF/DOFReduce.usf
@@ -42,8 +42,7 @@
 	#define CONFIG_GATHER_INPUT_LAYOUT (GATHER_INPUT_LAYOUT_RGB_COC)
 #endif
 
-#define CONFIG_WAVE_BROADCAST_REDUCTION (PLATFORM_SUPPORTS_WAVE_BROADCAST)
-
+#define CONFIG_WAVE_BROADCAST_REDUCTION 0
 
 // Configures the neighborhood analysis method to use for slight out of focus early out.
 #if COMPILER_SUPPORTS_WAVE_BIT_ORAND && (COMPILER_PSSL || XBOXONE_PROFILE)
