Fix an issue where splitscreen or stereo views with non-(0,0) aligned ViewRects would incorrectly reproject Contrast Adaptive Shading shading rate images, leading to shimmering in those views when moving the camera.

#Synced-CL 48355415

diff --git a/Engine/Shaders/Private/VariableRateShading/VRSShadingRateReproject.usf b/Engine/Shaders/Private/VariableRateShading/VRSShadingRateReproject.usf
--- a/Engine/Shaders/Private/VariableRateShading/VRSShadingRateReproject.usf
+++ b/Engine/Shaders/Private/VariableRateShading/VRSShadingRateReproject.usf
@@ -28,6 +28,9 @@ groupshared float CachedDepth[THREADGROUP_SIZE / 2][THREADGROUP_SIZE / 2];
 // like TAA. In practice, I found reprojection is only required in a really small set of scenarios since VRS already preserves 
 // silhouettes anyways. The main artifact I saw was materials that have a strip of metal that's not on the edge of a mesh and can
 // cause noticeble aliasing when it's got a bright specular highlight.
+//
+// Note that some of the above considerations do not apply for Software VRS - increased granularity leaves less allowance
+// for incorrect reprojection, which can lead to flickering along edges where rough rates have been applied incorrectly.
 float2 ReprojectBufferUV(uint3 DispatchThreadId, uint3 GroupThreadId, float2 BufferUV, bool bValidThread)
 {
 	// Only read depth once in the top left thread for each 2x2 block of thread IDs
@@ -72,7 +75,8 @@ void RescaleVariableRateShading(
 	const float epsilon = 0.0001;
 
 	const float2 HalfPixelOffset = float2(0.5, 0.5);
-	float2 CenterUV = (float2(DispatchThreadId.xy) + HalfPixelOffset) * pixelUVSize;
+
+	float2 CenterUV = (float2(DispatchThreadId.xy) + HalfPixelOffset) * pixelUVSize + ScaledUVOffset;
 	CenterUV = ReprojectBufferUV(DispatchThreadId, GroupThreadId, CenterUV, bValidThread); 
 
 	// FXC doesn't allow a thread to terminate before GroupMemoryBarrierWithGroupSync so 
@@ -82,6 +86,10 @@ void RescaleVariableRateShading(
 		return;
 	}
 
+	// When using splitscreen or stereo, the output reprojected SRI is double-wide, but the input SRIs are single-wide
+	// So we need to adjust the coords used to sample from the input SRIs to always sample from a single-wide texture
+	CenterUV = CenterUV - ScaledUVOffset;
+
 	// Map from pre-scaled UV to post-scaled UV.
 	float2 uvStart = CenterUV - pixelUVSize * 0.5 + epsilon;
 	float2 uvEnd = CenterUV + pixelUVSize * 0.5 - epsilon;
