Path Tracer: Fix two-sided foliage materials when substrate is enabled

The ray bias should be -1 for transmission sampling, otherwise the BSDF rays will self-intersect.

diff --git a/Engine/Shaders/Private/PathTracing/Material/PathTracingSubstrateFoliage.ush b/Engine/Shaders/Private/PathTracing/Material/PathTracingSubstrateFoliage.ush
--- a/Engine/Shaders/Private/PathTracing/Material/PathTracingSubstrateFoliage.ush
+++ b/Engine/Shaders/Private/PathTracing/Material/PathTracingSubstrateFoliage.ush
@@ -155,6 +155,7 @@ FMaterialSample SubstrateFoliage_SampleMaterial(
 	if (bSampledTransmit)
 	{
 		Result.AddLobeWithMIS(Data.DiffuseWeight * Payload.GetFoliageTransmissionColor(), abs(L.z) / PI, Data.LobePdf.x * Data.TransmissionProb);
+		Result.PositionBiasSign = -1;
 	}
 	else if (bSampledDiffuse)
 	{
