Fix LDS out-of-bound read when aggregating hair bounds within a macrogroup.

This caused the macrogroup to have incorrect bound (too large) when a log of hair element are part of the same macrogroup.

diff --git a/Engine/Shaders/Private/HairStrands/HairStrandsAABB.usf b/Engine/Shaders/Private/HairStrands/HairStrandsAABB.usf
--- a/Engine/Shaders/Private/HairStrands/HairStrandsAABB.usf
+++ b/Engine/Shaders/Private/HairStrands/HairStrandsAABB.usf
@@ -67,7 +67,8 @@ void Main(uint2 DispatchThreadId : SV_DispatchThreadID, uint LinearIndex : SV_Gr
 	if (ThreadId == 0)
 	{
 		FHairAABB Bound = InitHairAABB();
-		for (uint BoundIt = 0; BoundIt < MacroGroupCount; ++BoundIt)
+		const uint BoundCount = min(GROUP_SIZE, MacroGroupCount);
+		for (uint BoundIt = 0; BoundIt < BoundCount; ++BoundIt)
 		{
 			FHairAABB CurrBound;
 			CurrBound.Min = g_MinBounds[BoundIt];
