Clamp texture coordinate index in vertex factory for non manual vertex fetch path.
This now matches the logic used in the manual vertex fetch path.
#tests EngineTests, ReplayRun

diff --git a/Engine/Plugins/FX/Niagara/Shaders/Private/NiagaraMeshVertexFactory.ush b/Engine/Plugins/FX/Niagara/Shaders/Private/NiagaraMeshVertexFactory.ush
--- a/Engine/Plugins/FX/Niagara/Shaders/Private/NiagaraMeshVertexFactory.ush
+++ b/Engine/Plugins/FX/Niagara/Shaders/Private/NiagaraMeshVertexFactory.ush
@@ -26,10 +26,8 @@
 #define MANUAL_VERTEX_FETCH 0
 #endif
 
-#if MANUAL_VERTEX_FETCH
-	#define VF_ColorIndexMask_Index 0
-	#define VF_NumTexcoords_Index 1
-#endif // MANUAL_VERTEX_FETCH
+#define VF_ColorIndexMask_Index 0
+#define VF_NumTexcoords_Index 1
 
 struct FVertexFactoryInput
 {
@@ -604,9 +602,9 @@ FVertexFactoryIntermediates GetVertexFactoryIntermediates(FVertexFactoryInput In
 	// Particle translated world position.
 	Intermediates.ParticleTranslatedWorldPosition = LWCToFloat(LWCAdd(LWCGetOrigin(Intermediates.ParticleToWorld), ResolvedView.TileOffset.PreViewTranslation));
 
+	const uint NumFetchTexCoords = NiagaraMeshVF.VertexFetch_Parameters[VF_NumTexcoords_Index];
 #if NUM_MATERIAL_TEXCOORDS_VERTEX
 	#if MANUAL_VERTEX_FETCH
-		const uint NumFetchTexCoords = NiagaraMeshVF.VertexFetch_Parameters[VF_NumTexcoords_Index];
 		UNROLL
 		for (uint CoordinateIndex = 0; CoordinateIndex < NUM_MATERIAL_TEXCOORDS_VERTEX; CoordinateIndex++)
 		{
@@ -619,25 +617,25 @@ FVertexFactoryIntermediates GetVertexFactoryIntermediates(FVertexFactoryInput In
 			Intermediates.TexCoords[0] = Input.TexCoords0;
 		#elif NUM_MATERIAL_TEXCOORDS_VERTEX >= 2
 			Intermediates.TexCoords[0] = Input.TexCoords01.xy;
-			Intermediates.TexCoords[1] = Input.TexCoords01.zw;
+			Intermediates.TexCoords[1] = NumFetchTexCoords > 1 ? Input.TexCoords01.zw: Result.TexCoords[0];
 		#endif
 		#if NUM_MATERIAL_TEXCOORDS_VERTEX == 3
-			Intermediates.TexCoords[2] = Input.TexCoords2;
+			Intermediates.TexCoords[2] = NumFetchTexCoords > 2 ? Input.TexCoords2: Result.TexCoords[1];
 		#elif NUM_MATERIAL_TEXCOORDS_VERTEX >= 4
-			Intermediates.TexCoords[2] = Input.TexCoords23.xy;
-			Intermediates.TexCoords[3] = Input.TexCoords23.zw;
+			Intermediates.TexCoords[2] = NumFetchTexCoords > 2 ? Input.TexCoords23.xy : Result.TexCoords[1];
+			Intermediates.TexCoords[3] = NumFetchTexCoords > 3 ? Input.TexCoords23.zw : Result.TexCoords[2];
 		#endif
 		#if NUM_MATERIAL_TEXCOORDS_VERTEX == 5
-			Intermediates.TexCoords[4] = Input.TexCoords4;
+			Intermediates.TexCoords[4] = NumFetchTexCoords > 4 ? Input.TexCoords4 : Result.TexCoords[3];
 		#elif NUM_MATERIAL_TEXCOORDS_VERTEX >= 6
-			Intermediates.TexCoords[4] = Input.TexCoords45.xy;
-			Intermediates.TexCoords[5] = Input.TexCoords45.zw;
+			Intermediates.TexCoords[4] = NumFetchTexCoords > 4 ? Input.TexCoords45.xy : Result.TexCoords[3];
+			Intermediates.TexCoords[5] = NumFetchTexCoords > 5 ? Input.TexCoords45.zw : Result.TexCoords[4];
 		#endif
 		#if NUM_MATERIAL_TEXCOORDS_VERTEX == 7
-			Intermediates.TexCoords[6] = Input.TexCoords6;
+			Intermediates.TexCoords[6] = NumFetchTexCoords > 6 ? Input.TexCoords6 : Result.TexCoords[5];
 		#elif NUM_MATERIAL_TEXCOORDS_VERTEX >= 8
-			Intermediates.TexCoords[6] = Input.TexCoords67.xy;
-			Intermediates.TexCoords[7] = Input.TexCoords67.zw;
+			Intermediates.TexCoords[6] = NumFetchTexCoords > 6 ? Input.TexCoords67.xy : Result.TexCoords[5];
+			Intermediates.TexCoords[7] = NumFetchTexCoords > 7 ? Input.TexCoords67.zw : Result.TexCoords[6];
 		#endif
 	#endif // MANUAL_VERTEX_FETCH
 #endif // NUM_MATERIAL_TEXCOORDS_VERTEX
diff --git a/Engine/Plugins/FX/Niagara/Source/Niagara/Private/NiagaraRendererMeshes.cpp b/Engine/Plugins/FX/Niagara/Source/Niagara/Private/NiagaraRendererMeshes.cpp
--- a/Engine/Plugins/FX/Niagara/Source/Niagara/Private/NiagaraRendererMeshes.cpp
+++ b/Engine/Plugins/FX/Niagara/Source/Niagara/Private/NiagaraRendererMeshes.cpp
@@ -756,19 +756,18 @@ void FNiagaraRendererMeshes::PreparePerMeshData(FParticleMeshRenderData& Particl
 		}
 	}
 
+	const int32 NumTexCoords = VertexFactory.GetNumTexcoords();
+	const int32 ColorIndexMask = VertexFactory.GetColorIndexMask();
+	ParticleMeshRenderData.VertexFetch_Parameters = { ColorIndexMask, NumTexCoords, INDEX_NONE, INDEX_NONE };
+
 	if (RHISupportsManualVertexFetch(GShaderPlatformForFeatureLevel[FeatureLevel]))
 	{
-		const int32 NumTexCoords = VertexFactory.GetNumTexcoords();
-		const int32 ColorIndexMask = VertexFactory.GetColorIndexMask();
-
-		ParticleMeshRenderData.VertexFetch_Parameters = { ColorIndexMask, NumTexCoords, INDEX_NONE, INDEX_NONE };
 		ParticleMeshRenderData.TexCoordBufferSrv = VertexFactory.GetTextureCoordinatesSRV();
 		ParticleMeshRenderData.PackedTangentsBufferSrv = VertexFactory.GetTangentsSRV();
 		ParticleMeshRenderData.ColorComponentsBufferSrv = VertexFactory.GetColorComponentsSRV();
 	}
 	else
 	{
-		ParticleMeshRenderData.VertexFetch_Parameters = FIntVector4(INDEX_NONE);
 		ParticleMeshRenderData.TexCoordBufferSrv = FNiagaraRenderer::GetDummyFloat2Buffer();
 		ParticleMeshRenderData.PackedTangentsBufferSrv = FNiagaraRenderer::GetDummyFloat4Buffer();
 		ParticleMeshRenderData.ColorComponentsBufferSrv = FNiagaraRenderer::GetDummyFloat4Buffer();
diff --git a/Engine/Shaders/Private/GeometryCollectionVertexFactory.ush b/Engine/Shaders/Private/GeometryCollectionVertexFactory.ush
--- a/Engine/Shaders/Private/GeometryCollectionVertexFactory.ush
+++ b/Engine/Shaders/Private/GeometryCollectionVertexFactory.ush
@@ -19,12 +19,12 @@
 	#define USE_SHADER_BONE_TRANSFORM (MANUAL_VERTEX_FETCH)
 #endif
 
-#if MANUAL_VERTEX_FETCH || USE_SHADER_BONE_TRANSFORM
-	#define VF_ColorIndexMask_Index 0
-	#define VF_NumTexcoords_Index 1
-	#define FV_LightMapIndex_Index 2
-	#define VF_VertexOffset 3
+#define VF_ColorIndexMask_Index 0
+#define VF_NumTexcoords_Index 1
+#define FV_LightMapIndex_Index 2
+#define VF_VertexOffset 3
 
+#if MANUAL_VERTEX_FETCH || USE_SHADER_BONE_TRANSFORM
 	Buffer<float4> VertexFetch_BoneOriginBuffer;
 	Buffer<float4> VertexFetch_BoneLightmapBuffer;
 
@@ -259,8 +259,8 @@ FMaterialVertexParameters GetMaterialVertexParameters(
 	Result.PreSkinnedPosition = Input.Position.xyz;
 	Result.PreSkinnedNormal = TangentToLocal[2];
 
-#if MANUAL_VERTEX_FETCH && NUM_MATERIAL_TEXCOORDS_VERTEX
 	const uint NumFetchTexCoords = GeometryCollectionVF.VertexFetch_Parameters[VF_NumTexcoords_Index];
+#if MANUAL_VERTEX_FETCH && NUM_MATERIAL_TEXCOORDS_VERTEX
 	UNROLL
 	for (uint CoordinateIndex = 0; CoordinateIndex < NUM_MATERIAL_TEXCOORDS_VERTEX; CoordinateIndex++)
 	{
@@ -273,25 +273,25 @@ FMaterialVertexParameters GetMaterialVertexParameters(
 		Result.TexCoords[0] = Input.TexCoords0.xy;
 	#endif
 	#if NUM_MATERIAL_TEXCOORDS_VERTEX > 1
-		Result.TexCoords[1] = Input.TexCoords0.zw;
+		Result.TexCoords[1] = NumFetchTexCoords > 1 ? Input.TexCoords0.zw : Result.TexCoords[0];
 	#endif
 	#if NUM_MATERIAL_TEXCOORDS_VERTEX > 2
-		Result.TexCoords[2] = Input.TexCoords1.xy;
+		Result.TexCoords[2] = NumFetchTexCoords > 2 ? Input.TexCoords1.xy : Result.TexCoords[1];
 	#endif
 	#if NUM_MATERIAL_TEXCOORDS_VERTEX > 3
-		Result.TexCoords[3] = Input.TexCoords1.zw;
+		Result.TexCoords[3] = NumFetchTexCoords > 3 ? Input.TexCoords1.zw : Result.TexCoords[2];
 	#endif
 	#if NUM_MATERIAL_TEXCOORDS_VERTEX > 4
-		Result.TexCoords[4] = Input.TexCoords2.xy;
+		Result.TexCoords[4] = NumFetchTexCoords > 4 ? Input.TexCoords2.xy : Result.TexCoords[3];
 	#endif
 	#if NUM_MATERIAL_TEXCOORDS_VERTEX > 5
-		Result.TexCoords[5] = Input.TexCoords2.zw;
+		Result.TexCoords[5] = NumFetchTexCoords > 5 ? Input.TexCoords2.zw : Result.TexCoords[4];
 	#endif
 	#if NUM_MATERIAL_TEXCOORDS_VERTEX > 6
-		Result.TexCoords[6] = Input.TexCoords3.xy;
+		Result.TexCoords[6] = NumFetchTexCoords > 6 ? Input.TexCoords3.xy : Result.TexCoords[5];
 	#endif
 	#if NUM_MATERIAL_TEXCOORDS_VERTEX > 7
-		Result.TexCoords[7] = Input.TexCoords3.zw;
+		Result.TexCoords[7] = NumFetchTexCoords > 7 ? Input.TexCoords3.zw : Result.TexCoords[6];
 	#endif
 #endif  //MANUAL_VERTEX_FETCH && NUM_MATERIAL_TEXCOORDS_VERTEX
 
diff --git a/Engine/Shaders/Private/LocalVertexFactory.ush b/Engine/Shaders/Private/LocalVertexFactory.ush
--- a/Engine/Shaders/Private/LocalVertexFactory.ush
+++ b/Engine/Shaders/Private/LocalVertexFactory.ush
@@ -25,12 +25,12 @@ uint InstanceOffset;
 #define MANUAL_VERTEX_FETCH 0
 #endif
 
-#if MANUAL_VERTEX_FETCH
-	#define VF_ColorIndexMask_Index 0
-	#define VF_NumTexcoords_Index 1
-	#define FV_LightMapIndex_Index 2
-	#define VF_VertexOffset 3
+#define VF_ColorIndexMask_Index 0
+#define VF_NumTexcoords_Index 1
+#define FV_LightMapIndex_Index 2
+#define VF_VertexOffset 3
 
+#if MANUAL_VERTEX_FETCH
 	Buffer<float4> VertexFetch_InstanceOriginBuffer;
 	Buffer<float4> VertexFetch_InstanceTransformBuffer;
 	Buffer<float4> VertexFetch_InstanceLightmapBuffer;
@@ -724,8 +724,8 @@ FMaterialVertexParameters GetMaterialVertexParameters(
 		Result.PreSkinnedNormal = TangentToLocal[2]; //TangentBias(Input.TangentZ.xyz);
 	}
 
+	const uint NumFetchTexCoords = LocalVF.VertexFetch_Parameters[VF_NumTexcoords_Index];
 #if MANUAL_VERTEX_FETCH && NUM_MATERIAL_TEXCOORDS_VERTEX
-		const uint NumFetchTexCoords = LocalVF.VertexFetch_Parameters[VF_NumTexcoords_Index];
 		UNROLL
 		for (uint CoordinateIndex = 0; CoordinateIndex < NUM_MATERIAL_TEXCOORDS_VERTEX; CoordinateIndex++)
 		{
@@ -734,30 +734,30 @@ FMaterialVertexParameters GetMaterialVertexParameters(
 			Result.TexCoords[CoordinateIndex] = LocalVF.VertexFetch_TexCoordBuffer[NumFetchTexCoords * (LocalVF.VertexFetch_Parameters[VF_VertexOffset] + Input.VertexId) + ClampedCoordinateIndex];
 		}
 #elif NUM_MATERIAL_TEXCOORDS_VERTEX
-			#if NUM_MATERIAL_TEXCOORDS_VERTEX > 0
-				Result.TexCoords[0] = Input.TexCoords0.xy;
-			#endif
-			#if NUM_MATERIAL_TEXCOORDS_VERTEX > 1
-				Result.TexCoords[1] = Input.TexCoords0.zw;
-			#endif
-			#if NUM_MATERIAL_TEXCOORDS_VERTEX > 2
-				Result.TexCoords[2] = Input.TexCoords1.xy;
-			#endif
-			#if NUM_MATERIAL_TEXCOORDS_VERTEX > 3
-				Result.TexCoords[3] = Input.TexCoords1.zw;
-			#endif
-			#if NUM_MATERIAL_TEXCOORDS_VERTEX > 4
-				Result.TexCoords[4] = Input.TexCoords2.xy;
-			#endif
-			#if NUM_MATERIAL_TEXCOORDS_VERTEX > 5
-				Result.TexCoords[5] = Input.TexCoords2.zw;
-			#endif
-			#if NUM_MATERIAL_TEXCOORDS_VERTEX > 6
-				Result.TexCoords[6] = Input.TexCoords3.xy;
-			#endif
-			#if NUM_MATERIAL_TEXCOORDS_VERTEX > 7
-				Result.TexCoords[7] = Input.TexCoords3.zw;
-			#endif
+		#if NUM_MATERIAL_TEXCOORDS_VERTEX > 0
+			Result.TexCoords[0] = Input.TexCoords0.xy;
+		#endif
+		#if NUM_MATERIAL_TEXCOORDS_VERTEX > 1
+			Result.TexCoords[1] = NumFetchTexCoords > 1 ? Input.TexCoords0.zw : Result.TexCoords[0];
+		#endif
+		#if NUM_MATERIAL_TEXCOORDS_VERTEX > 2
+			Result.TexCoords[2] = NumFetchTexCoords > 2 ? Input.TexCoords1.xy : Result.TexCoords[1];
+		#endif
+		#if NUM_MATERIAL_TEXCOORDS_VERTEX > 3
+			Result.TexCoords[3] = NumFetchTexCoords > 3 ? Input.TexCoords1.zw : Result.TexCoords[2];
+		#endif
+		#if NUM_MATERIAL_TEXCOORDS_VERTEX > 4
+			Result.TexCoords[4] = NumFetchTexCoords > 4 ? Input.TexCoords2.xy : Result.TexCoords[3];
+		#endif
+		#if NUM_MATERIAL_TEXCOORDS_VERTEX > 5
+			Result.TexCoords[5] = NumFetchTexCoords > 5 ? Input.TexCoords2.zw : Result.TexCoords[4];
+		#endif
+		#if NUM_MATERIAL_TEXCOORDS_VERTEX > 6
+			Result.TexCoords[6] = NumFetchTexCoords > 6 ? Input.TexCoords3.xy : Result.TexCoords[5];
+		#endif
+		#if NUM_MATERIAL_TEXCOORDS_VERTEX > 7
+			Result.TexCoords[7] = NumFetchTexCoords > 7 ? Input.TexCoords3.zw : Result.TexCoords[6];
+		#endif
 #endif  //MANUAL_VERTEX_FETCH && NUM_MATERIAL_TEXCOORDS_VERTEX
 
 	Result.Particle.Color = half4(1,1,1,1);
