Add special Slate shader version that uses dynamic branching instead of static specializations. This new variant will be the only path, replacing most of the above after performance and correctness testing.
The purpose of the change is to reduce the total number of shaders compiled for a shipping game and the number of PSOs that must be created at runtime.

* Add CVar r.Slate.DynamicShaders (default: 0) to enable the dynamic branching path

diff --git a/Engine/Shaders/Private/MaterialTemplate.ush b/Engine/Shaders/Private/MaterialTemplate.ush
--- a/Engine/Shaders/Private/MaterialTemplate.ush
+++ b/Engine/Shaders/Private/MaterialTemplate.ush
@@ -598,7 +598,7 @@ struct FMaterialPixelParameters
 	uint WaterWaveParamIndex;
 #endif
 
-#if SHADER_TYPE == ST_GrayscaleFont || SHADER_TYPE == ST_ColorFont || SHADER_TYPE == ST_SdfFont || SHADER_TYPE == ST_MsdfFont
+#if SHADER_TYPE == ST_GrayscaleFont || SHADER_TYPE == ST_ColorFont || SHADER_TYPE == ST_SdfFont || SHADER_TYPE == ST_MsdfFont || SHADER_TYPE == ST_Dynamic
 	float4 FontSignedDistanceData;
 #endif
 
@@ -1989,7 +1989,7 @@ MaterialFloat GetDistanceCullFade()
 
 MaterialFloat4 GetFontSignedDistanceData(FMaterialPixelParameters Parameters)
 {
-#if SHADER_TYPE == ST_GrayscaleFont || SHADER_TYPE == ST_ColorFont || SHADER_TYPE == ST_SdfFont || SHADER_TYPE == ST_MsdfFont
+#if SHADER_TYPE == ST_GrayscaleFont || SHADER_TYPE == ST_ColorFont || SHADER_TYPE == ST_SdfFont || SHADER_TYPE == ST_MsdfFont || SHADER_TYPE == ST_Dynamic
 	return Parameters.FontSignedDistanceData;
 #else
 	return float4(0.0f, 0.0f, 0.0f, 0.0f);
diff --git a/Engine/Shaders/Private/SlateElementPixelShader.usf b/Engine/Shaders/Private/SlateElementPixelShader.usf
--- a/Engine/Shaders/Private/SlateElementPixelShader.usf
+++ b/Engine/Shaders/Private/SlateElementPixelShader.usf
@@ -5,6 +5,12 @@
 #include "GammaCorrectionCommon.ush"
 #include "SlateShaderCommon.ush"
 
+#pragma message("UESHADERMETADATA_VERSION F34E9FD2-616B-4620-95C5-C123D78CE6AA")
+
+#ifndef SHADER_TYPE
+#define SHADER_TYPE 0
+#endif
+
 #ifndef SAMPLE_VIRTUAL_TEXTURE
 #define SAMPLE_VIRTUAL_TEXTURE 0
 #endif
@@ -18,16 +24,25 @@
 
 /** Display gamma x:gamma curve adjustment, y:inverse gamma (1/GEngine->DisplayGamma), z:InvertAlpha, w:Contrast */
 half4 GammaAndAlphaValues;
-/** x:DrawDisabledEffect */
+
 /** Per element params */
 float4 ShaderParams;
 float4 ShaderParams2;
 
 #define USE_LEGACY_DISABLED_EFFECT 0
 
-#if USE_MATERIALS
-half4 DrawFlags;
+/** x: DrawDisabledEffect */
+/** y: ShaderType*/
+uint4 DrawFlags;
+
+uint GetSlateShaderType()
+{
+#if SHADER_TYPE == ST_Dynamic
+	return DrawFlags.y;
+#else
+	return SHADER_TYPE;
 #endif
+}
 
 /** Texture for the element */
 Texture2D ElementTexture;
@@ -94,46 +109,58 @@ void GetSlateUVs(VertexToPixelInterpolants Interpolants, inout float2 UVs[8])
 #if USE_MATERIALS
 void FillTexCoords( inout FMaterialPixelParameters Parameters, VertexToPixelInterpolants InVertex, float2 InDefaultUV )
 {
-#if SHADER_TYPE == ST_GrayscaleFont || SHADER_TYPE == ST_ColorFont || SHADER_TYPE == ST_SdfFont || SHADER_TYPE == ST_MsdfFont
-	#if NUM_MATERIAL_TEXCOORDS == 1
-		// Note: The first channel is used to look up into the texture atlas. The second channel is the custom 0-1 UV set for the text run
-		Parameters.TexCoords[0] = SLATE_UV1( InVertex );
-	#elif NUM_MATERIAL_TEXCOORDS == 2
-		// Note: The first channel is used to look up into the texture atlas. The second channel is the custom 0-1 UV set for the text run
-		Parameters.TexCoords[0] = SLATE_UV1( InVertex );
-		Parameters.TexCoords[1] = InVertex.MaterialTexCoords;
-	#endif
-#elif (SHADER_TYPE == ST_Custom || NUM_CUSTOMIZED_UVS != 0) && NUM_MATERIAL_TEXCOORDS > 0
-	// note: We assume users who customize UVs or are using completely custom elements know what they are doing with UVs and want no assumptions made about how uvs are used
-	float2 UVs[8]; GetSlateUVs(InVertex, UVs);
-	UNROLL
-	for( int CoordinateIndex = 0; CoordinateIndex < NUM_MATERIAL_TEXCOORDS; CoordinateIndex++ )
+	const uint ShaderType = GetSlateShaderType();
+
+	if (ShaderType == ST_GrayscaleFont || ShaderType == ST_ColorFont || ShaderType == ST_SdfFont || ShaderType == ST_MsdfFont)
 	{
-		Parameters.TexCoords[CoordinateIndex] = UVs[CoordinateIndex];
+		#if NUM_MATERIAL_TEXCOORDS == 1
+			// Note: The first channel is used to look up into the texture atlas. The second channel is the custom 0-1 UV set for the text run
+			Parameters.TexCoords[0] = SLATE_UV1( InVertex );
+		#elif NUM_MATERIAL_TEXCOORDS == 2
+			// Note: The first channel is used to look up into the texture atlas. The second channel is the custom 0-1 UV set for the text run
+			Parameters.TexCoords[0] = SLATE_UV1( InVertex );
+			Parameters.TexCoords[1] = InVertex.MaterialTexCoords;
+		#endif
+
+		return;
 	}
-#else
-	#if NUM_MATERIAL_TEXCOORDS == 1
-		Parameters.TexCoords[0] = InDefaultUV;
-	#elif NUM_MATERIAL_TEXCOORDS == 2
-		Parameters.TexCoords[0] = InDefaultUV;
-		Parameters.TexCoords[1] = SLATE_UV0( InVertex );
-	#elif NUM_MATERIAL_TEXCOORDS == 3
-		Parameters.TexCoords[0] = InDefaultUV;
-		Parameters.TexCoords[1] = SLATE_UV0( InVertex );
-		Parameters.TexCoords[2] = SLATE_UV1( InVertex );
-	#elif NUM_MATERIAL_TEXCOORDS >= 4
-		Parameters.TexCoords[0] = InDefaultUV;
-		Parameters.TexCoords[1] = SLATE_UV0( InVertex );
-		Parameters.TexCoords[2] = SLATE_UV1( InVertex );
 
+#if NUM_MATERIAL_TEXCOORDS > 0
+	if (ShaderType == ST_Custom || NUM_CUSTOMIZED_UVS != 0)
+	{
+		// note: We assume users who customize UVs or are using completely custom elements know what they are doing with UVs and want no assumptions made about how uvs are used
 		float2 UVs[8]; GetSlateUVs(InVertex, UVs);
 		UNROLL
-		for( int CoordinateIndex = 3; CoordinateIndex < NUM_MATERIAL_TEXCOORDS; CoordinateIndex++ )
+		for (int CoordinateIndex = 0; CoordinateIndex < NUM_MATERIAL_TEXCOORDS; CoordinateIndex++)
 		{
 			Parameters.TexCoords[CoordinateIndex] = UVs[CoordinateIndex];
 		}
-	#endif // NUM_MATERIAL_TEXCOORDS == 1
-#endif	// SHADER_TYPE
+
+		return;
+	}
+#endif // NUM_MATERIAL_TEXCOORDS > 0
+
+#if NUM_MATERIAL_TEXCOORDS == 1
+	Parameters.TexCoords[0] = InDefaultUV;
+#elif NUM_MATERIAL_TEXCOORDS == 2
+	Parameters.TexCoords[0] = InDefaultUV;
+	Parameters.TexCoords[1] = SLATE_UV0(InVertex);
+#elif NUM_MATERIAL_TEXCOORDS == 3
+	Parameters.TexCoords[0] = InDefaultUV;
+	Parameters.TexCoords[1] = SLATE_UV0(InVertex);
+	Parameters.TexCoords[2] = SLATE_UV1(InVertex);
+#elif NUM_MATERIAL_TEXCOORDS >= 4
+	Parameters.TexCoords[0] = InDefaultUV;
+	Parameters.TexCoords[1] = SLATE_UV0(InVertex);
+	Parameters.TexCoords[2] = SLATE_UV1(InVertex);
+
+	float2 UVs[8]; GetSlateUVs(InVertex, UVs);
+	UNROLL
+	for (int CoordinateIndex = 3; CoordinateIndex < NUM_MATERIAL_TEXCOORDS; CoordinateIndex++)
+	{
+		Parameters.TexCoords[CoordinateIndex] = UVs[CoordinateIndex];
+	}
+#endif // NUM_MATERIAL_TEXCOORDS == 1
 }
 
 half4 GetMaterialColor( VertexToPixelInterpolants InVertex, float2 MeshUV, float4 InFontSignedDistanceData = float4(0.0, 0.0, 0.0, 0.0) )
@@ -145,7 +172,7 @@ half4 GetMaterialColor( VertexToPixelInterpolants InVertex, float2 MeshUV, float
 	FillTexCoords( Parameters, InVertex, MeshUV );
 
 	Parameters.VertexColor = InVertex.Color;
-#if SHADER_TYPE == ST_GrayscaleFont || SHADER_TYPE == ST_ColorFont || SHADER_TYPE == ST_SdfFont || SHADER_TYPE == ST_MsdfFont
+#if SHADER_TYPE == ST_GrayscaleFont || SHADER_TYPE == ST_ColorFont || SHADER_TYPE == ST_SdfFont || SHADER_TYPE == ST_MsdfFont || SHADER_TYPE == ST_Dynamic
 	Parameters.FontSignedDistanceData = InFontSignedDistanceData;
 #endif
 
@@ -259,21 +286,24 @@ half4 GetColor( VertexToPixelInterpolants InVertex, float2 UV )
 	half4 BaseColor = GetTextureColor( InVertex, UV );
 #endif
 
-#if SHADER_TYPE == ST_Custom
-	// Make no assumptionn of what vertex color does in a custom material shader
-	FinalColor = BaseColor;
-#else
-	FinalColor = BaseColor * InVertex.Color;
-	#if USE_MATERIALS && MATERIALBLENDING_ADDITIVE
+	const uint ShaderType = GetSlateShaderType();
+
+	if (ShaderType == ST_Custom)
+	{
+		// Make no assumptionn of what vertex color does in a custom material shader
+		FinalColor = BaseColor;	
+	}
+	else
+	{
+		FinalColor = BaseColor * InVertex.Color;
+#if USE_MATERIALS && MATERIALBLENDING_ADDITIVE
 		FinalColor *= InVertex.Color.a;
-	#endif
 #endif
+	}
 
 	return FinalColor;
 }
 
-#if SHADER_TYPE == ST_SdfFont || SHADER_TYPE == ST_MsdfFont
-
 struct FontSignedDistanceSample
 {
 	MaterialFloat DistanceSample;
@@ -285,13 +315,20 @@ FontSignedDistanceSample SampleFontSignedDistance(float2 TextureCoordinate)
 {
 	FontSignedDistanceSample OutSample;
 	const MaterialFloat4 MultiDistance = Texture2DSample(ElementTexture, ElementTextureSampler, TextureCoordinate);
-#if SHADER_TYPE == ST_MsdfFont
-	OutSample.DistanceSample = max(min(MultiDistance.r, MultiDistance.g), min(max(MultiDistance.r, MultiDistance.g), MultiDistance.b));
-	OutSample.SmoothDistanceSample = MultiDistance.a;
-#else // SHADER_TYPE == ST_SdfFont
-	OutSample.DistanceSample = lerp(MultiDistance A8_SAMPLE_MASK, MultiDistance.a, ShaderParams2.x);
-	OutSample.SmoothDistanceSample = OutSample.DistanceSample;
-#endif
+
+	const uint ShaderType = GetSlateShaderType();
+
+	if (GetSlateShaderType() == ST_MsdfFont)
+	{
+		OutSample.DistanceSample = max(min(MultiDistance.r, MultiDistance.g), min(max(MultiDistance.r, MultiDistance.g), MultiDistance.b));
+		OutSample.SmoothDistanceSample = MultiDistance.a;
+	}
+	else
+	{
+		OutSample.DistanceSample = lerp(MultiDistance A8_SAMPLE_MASK, MultiDistance.a, ShaderParams2.x);
+		OutSample.SmoothDistanceSample = OutSample.DistanceSample;
+	}
+
 	const float2 ScreenSpaceTextureSize = float2(1.0, 1.0) / fwidth(TextureCoordinate);
 	OutSample.ScreenSpaceSpread = max(dot(ShaderParams.xy, ScreenSpaceTextureSize), 1.0);
 	return OutSample;
@@ -334,24 +371,23 @@ half4 GetSdfTextElementColor(VertexToPixelInterpolants InVertex)
 	return OutColor;
 }
 
-#endif
-
-
-#if SHADER_TYPE == ST_GrayscaleFont || SHADER_TYPE == ST_ColorFont
 half4 GetFontElementColor( VertexToPixelInterpolants InVertex )
 {
 	float2 AtlasTextureCoordinates = SLATE_UV0( InVertex );
 	MaterialFloat4 SampledColor = MaterialFloat4(1,1,1,1);
 
-#if SHADER_TYPE == ST_GrayscaleFont
-	#if USE_MATERIALS && (MATERIALBLENDING_ADDITIVE || MATERIALBLENDING_ALPHACOMPOSITE)
-		SampledColor *= Texture2DSample_A8( ElementTexture, ElementTextureSampler, AtlasTextureCoordinates );
-	#else
-		SampledColor.a = Texture2DSample_A8( ElementTexture, ElementTextureSampler, AtlasTextureCoordinates );
-	#endif
-#else
-	SampledColor = Texture2DSample( ElementTexture, ElementTextureSampler, AtlasTextureCoordinates );
-#endif
+	if (GetSlateShaderType() == ST_GrayscaleFont)
+	{
+		#if USE_MATERIALS && (MATERIALBLENDING_ADDITIVE || MATERIALBLENDING_ALPHACOMPOSITE)
+			SampledColor *= Texture2DSample_A8(ElementTexture, ElementTextureSampler, AtlasTextureCoordinates);
+		#else
+			SampledColor.a = Texture2DSample_A8(ElementTexture, ElementTextureSampler, AtlasTextureCoordinates);
+		#endif
+	}
+	else
+	{
+		SampledColor = Texture2DSample(ElementTexture, ElementTextureSampler, AtlasTextureCoordinates);
+	}
 
 	half4 OutColor = InVertex.Color;
 #if USE_MATERIALS
@@ -366,16 +402,12 @@ half4 GetFontElementColor( VertexToPixelInterpolants InVertex )
 
 	return OutColor;
 }
-#endif
 
-#if SHADER_TYPE == ST_Default || SHADER_TYPE == ST_Custom
 half4 GetDefaultElementColor( VertexToPixelInterpolants InVertex )
 {
 	return GetColor( InVertex, SLATE_UV0( InVertex ) * SLATE_UV1( InVertex ) );
 }
-#endif
 
-#if SHADER_TYPE == ST_Border 
 half4 GetBorderElementColor( VertexToPixelInterpolants InVertex )
 {
 	float4 TextureCoordinates = InVertex.TextureCoordinates0;
@@ -409,17 +441,12 @@ half4 GetBorderElementColor( VertexToPixelInterpolants InVertex )
 
 	return GetColor( InVertex, NewUV );
 }
-#endif
 
 
-#if SHADER_TYPE == ST_RoundedBox
 half4 GetRoundedBoxElementColor( VertexToPixelInterpolants InVertex )
 {
 	return GetRoundedBoxElementColorInternal(ShaderParams.zw, InVertex.MaterialTexCoords.xy, ShaderParams.y, ShaderParams2, GetColor(InVertex, SLATE_UV0(InVertex) * SLATE_UV1(InVertex)),  InVertex.SecondaryColor);
 }
-#endif
-
-#if SHADER_TYPE == ST_LineSegment
 
 /**
  * Computes 0-1 value that interpolates between dashes of length DashLength based on Distance with 1px of antialiasing between dashes and gaps
@@ -477,25 +504,36 @@ half4 GetLineSegmentElementColor( VertexToPixelInterpolants InVertex )
 	return Color;
 }
 
-#endif
-
 HALF4_TYPE Main( VertexToPixelInterpolants VIn ) : SV_Target0
 {
 	half4 OutColor;
-#if SHADER_TYPE == ST_Default || SHADER_TYPE == ST_Custom
-	OutColor = GetDefaultElementColor( VIn );
-#elif SHADER_TYPE == ST_RoundedBox
-	OutColor = GetRoundedBoxElementColor( VIn );
-#elif SHADER_TYPE == ST_Border
-	OutColor = GetBorderElementColor( VIn );
-#elif SHADER_TYPE == ST_GrayscaleFont || SHADER_TYPE == ST_ColorFont
-	OutColor = GetFontElementColor( VIn );
-#elif SHADER_TYPE == ST_SdfFont || SHADER_TYPE == ST_MsdfFont
-	OutColor = GetSdfTextElementColor( VIn );
-#else
-	OutColor = GetLineSegmentElementColor( VIn );
-#endif
-	
+	const uint ShaderType = GetSlateShaderType();
+
+	if (ShaderType == ST_Default || ShaderType == ST_Custom)
+	{
+		OutColor = GetDefaultElementColor(VIn);
+	}
+	else if (ShaderType == ST_RoundedBox)
+	{
+		OutColor = GetRoundedBoxElementColor( VIn );
+	}
+	else if (ShaderType == ST_Border)
+	{
+		OutColor = GetBorderElementColor(VIn);
+	}
+	else if (ShaderType == ST_GrayscaleFont || ShaderType == ST_ColorFont)
+	{
+		OutColor = GetFontElementColor(VIn);
+	}
+	else if (ShaderType == ST_SdfFont || ShaderType == ST_MsdfFont)
+	{
+		OutColor = GetSdfTextElementColor(VIn);
+	}
+	else
+	{
+		OutColor = GetLineSegmentElementColor(VIn);
+	}
+
 	// For standard slate which only has a few shaders, the permutation is fine here,
 	// for Engine Materials, we don't ever define this macro, so instead we fallback
 	// to looking at the shader constants.
diff --git a/Engine/Shaders/Private/SlateShaderCommon.ush b/Engine/Shaders/Private/SlateShaderCommon.ush
--- a/Engine/Shaders/Private/SlateShaderCommon.ush
+++ b/Engine/Shaders/Private/SlateShaderCommon.ush
@@ -11,6 +11,7 @@
 #define ST_RoundedBox		7
 #define ST_SdfFont			8
 #define ST_MsdfFont			9
+#define ST_Dynamic			10 // This variant will be the only path, replacing most of the above after performance and correctness testing (#jira UE-319026)
 
 #if USE_MATERIALS
 #include "/Engine/Generated/Material.ush"
diff --git a/Engine/Source/Runtime/SlateCore/Public/Rendering/RenderingCommon.h b/Engine/Source/Runtime/SlateCore/Public/Rendering/RenderingCommon.h
--- a/Engine/Source/Runtime/SlateCore/Public/Rendering/RenderingCommon.h
+++ b/Engine/Source/Runtime/SlateCore/Public/Rendering/RenderingCommon.h
@@ -78,6 +78,12 @@ enum class ESlateShader : uint8
 	SdfFont = 8,
 	/** Multi-channel signed distance field font shader */
 	MsdfFont = 9,
+	/** 
+	* Dynamic branching version.
+	* This variant will be the only path, replacing most of the above after performance and correctness testing.
+	* #jira UE-319026
+	*/
+	Dynamic = 10,
 };
 
 /**
diff --git a/Engine/Source/Runtime/SlateRHIRenderer/Private/SlateMaterialShader.cpp b/Engine/Source/Runtime/SlateRHIRenderer/Private/SlateMaterialShader.cpp
--- a/Engine/Source/Runtime/SlateRHIRenderer/Private/SlateMaterialShader.cpp
+++ b/Engine/Source/Runtime/SlateRHIRenderer/Private/SlateMaterialShader.cpp
@@ -101,9 +101,9 @@ void FSlateMaterialShaderPS::SetDisplayGammaAndContrast(FMeshDrawSingleShaderBin
 	ShaderBindings.Add(GammaAndAlphaValues, InGammaValues);
 }
 
-void FSlateMaterialShaderPS::SetDrawFlags(FMeshDrawSingleShaderBindings& ShaderBindings, bool bDrawDisabledEffect)
+void FSlateMaterialShaderPS::SetDrawFlags(FMeshDrawSingleShaderBindings& ShaderBindings, ESlateShader InShaderType, bool bDrawDisabledEffect)
 {
-	FVector4f InDrawFlags((bDrawDisabledEffect ? 1.f : 0.f), 0.f, 0.f, 0.f);
+	FUintVector4 InDrawFlags((bDrawDisabledEffect ? 1u : 0u), uint32(InShaderType), 0, 0);
 
 	ShaderBindings.Add(DrawFlags, InDrawFlags);
 }
@@ -129,3 +129,4 @@ IMPLEMENT_SLATE_MATERIALSHADER_TYPE(ColorFont);
 IMPLEMENT_SLATE_MATERIALSHADER_TYPE(RoundedBox);
 IMPLEMENT_SLATE_MATERIALSHADER_TYPE(SdfFont);
 IMPLEMENT_SLATE_MATERIALSHADER_TYPE(MsdfFont);
+IMPLEMENT_SLATE_MATERIALSHADER_TYPE(Dynamic);
diff --git a/Engine/Source/Runtime/SlateRHIRenderer/Private/SlateMaterialShader.h b/Engine/Source/Runtime/SlateRHIRenderer/Private/SlateMaterialShader.h
--- a/Engine/Source/Runtime/SlateRHIRenderer/Private/SlateMaterialShader.h
+++ b/Engine/Source/Runtime/SlateRHIRenderer/Private/SlateMaterialShader.h
@@ -56,7 +56,7 @@ public:
 
 	void SetAdditionalTexture(FMeshDrawSingleShaderBindings& ShaderBindings, FRHITexture* InTexture, const FSamplerStateRHIRef SamplerState );
 
-	void SetDrawFlags(FMeshDrawSingleShaderBindings& ShaderBindings, bool bDrawDisabledEffect);
+	void SetDrawFlags(FMeshDrawSingleShaderBindings& ShaderBindings, ESlateShader InShaderType, bool bDrawDisabledEffect);
 
 private:
 	LAYOUT_FIELD(FShaderParameter, GammaAndAlphaValues);
diff --git a/Engine/Source/Runtime/SlateRHIRenderer/Private/SlateRHIRenderingPolicy.cpp b/Engine/Source/Runtime/SlateRHIRenderer/Private/SlateRHIRenderingPolicy.cpp
--- a/Engine/Source/Runtime/SlateRHIRenderer/Private/SlateRHIRenderingPolicy.cpp
+++ b/Engine/Source/Runtime/SlateRHIRenderer/Private/SlateRHIRenderingPolicy.cpp
@@ -57,6 +57,14 @@ static FAutoConsoleVariableRef CVarGSlateMaterialPSOPrecache(
 	TEXT("Precache all possible required PSOs for loaded Slate Materials."),
 	ECVF_ReadOnly);
 
+static int32 GSlateDynamicShaders = 0;
+static FAutoConsoleVariableRef CVarGSlateDynamicShaders(
+	TEXT("r.Slate.DynamicShaders"),
+	GSlateDynamicShaders,
+	TEXT("Whether to use a dynamic branch path for all Slate materials instead of static specializations. ")
+	TEXT("This reduces the number of unique Slate PSOs, but may result in some GPU overhead. (Default: 0)"),
+	ECVF_Default);
+
 static const TCHAR* SlateGlobalPSOCollectorName = TEXT("SlateGlobalPSOCollector");
 static const TCHAR* SlateMaterialPSOCollectorName = TEXT("SlateMaterialPSOCollector");
 
@@ -431,34 +439,41 @@ TShaderRef<FSlateElementPS> GetTexturePixelShader(FGlobalShaderMap* ShaderMap, E
 
 bool ChooseMaterialShaderTypes(ESlateShader ShaderType, bool bUseInstancing, FMaterialShaderTypes& OutShaderTypes)
 {
-	switch (ShaderType)
+	if (GSlateDynamicShaders)
 	{
-	case ESlateShader::Default:
-		OutShaderTypes.AddShaderType<TSlateMaterialShaderPS<ESlateShader::Default>>();
-		break;
-	case ESlateShader::Border:
-		OutShaderTypes.AddShaderType<TSlateMaterialShaderPS<ESlateShader::Border>>();
-		break;
-	case ESlateShader::GrayscaleFont:
-		OutShaderTypes.AddShaderType<TSlateMaterialShaderPS<ESlateShader::GrayscaleFont>>();
-		break;
-	case ESlateShader::ColorFont:
-		OutShaderTypes.AddShaderType<TSlateMaterialShaderPS<ESlateShader::ColorFont>>();
-		break;
-	case ESlateShader::Custom:
-		OutShaderTypes.AddShaderType<TSlateMaterialShaderPS<ESlateShader::Custom>>();
-		break;
-	case ESlateShader::RoundedBox:
-		OutShaderTypes.AddShaderType<TSlateMaterialShaderPS<ESlateShader::RoundedBox>>();
-		break;
-	case ESlateShader::SdfFont:
-		OutShaderTypes.AddShaderType<TSlateMaterialShaderPS<ESlateShader::SdfFont>>();
-		break;
-	case ESlateShader::MsdfFont:
-		OutShaderTypes.AddShaderType<TSlateMaterialShaderPS<ESlateShader::MsdfFont>>();
-		break;
-	default:
-		return false;
+		OutShaderTypes.AddShaderType<TSlateMaterialShaderPS<ESlateShader::Dynamic>>();
+	}
+	else
+	{
+		switch (ShaderType)
+		{
+		case ESlateShader::Default:
+			OutShaderTypes.AddShaderType<TSlateMaterialShaderPS<ESlateShader::Default>>();
+			break;
+		case ESlateShader::Border:
+			OutShaderTypes.AddShaderType<TSlateMaterialShaderPS<ESlateShader::Border>>();
+			break;
+		case ESlateShader::GrayscaleFont:
+			OutShaderTypes.AddShaderType<TSlateMaterialShaderPS<ESlateShader::GrayscaleFont>>();
+			break;
+		case ESlateShader::ColorFont:
+			OutShaderTypes.AddShaderType<TSlateMaterialShaderPS<ESlateShader::ColorFont>>();
+			break;
+		case ESlateShader::Custom:
+			OutShaderTypes.AddShaderType<TSlateMaterialShaderPS<ESlateShader::Custom>>();
+			break;
+		case ESlateShader::RoundedBox:
+			OutShaderTypes.AddShaderType<TSlateMaterialShaderPS<ESlateShader::RoundedBox>>();
+			break;
+		case ESlateShader::SdfFont:
+			OutShaderTypes.AddShaderType<TSlateMaterialShaderPS<ESlateShader::SdfFont>>();
+			break;
+		case ESlateShader::MsdfFont:
+			OutShaderTypes.AddShaderType<TSlateMaterialShaderPS<ESlateShader::MsdfFont>>();
+			break;
+		default:
+			return false;
+		}
 	}
 
 	if (bUseInstancing)
@@ -1123,7 +1138,7 @@ FSlateRenderBatchOp* CreateSlateRenderBatchOp(
 		PixelBindings = FSlateDrawShaderBindings::Create(GraphBuilder, PixelShader);
 		PixelShader->SetMaterialShaderParameters(*PixelBindings, Scene, ViewUniformBuffer, MaterialRenderProxy, EffectiveMaterial, ShaderParams);
 		PixelShader->SetDisplayGammaAndContrast(*PixelBindings, FinalGamma, FinalContrast);
-		PixelShader->SetDrawFlags(*PixelBindings, bDrawDisabled);
+		PixelShader->SetDrawFlags(*PixelBindings, ShaderType, bDrawDisabled);
 
 		auto* MaskResource = static_cast<TSlateTexture<FTextureRHIRef>*>(MaterialShaderResource->GetTextureMaskResource());
 
@@ -1708,14 +1723,26 @@ void AddSlateDrawElementsPass(
 	INC_DWORD_STAT_BY(STAT_SlateStencilClips, ClippingCreateContext.NumStencils);
 }
 
-static TArray<ESlateShader> SlateShaderTypesToPrecache =
+TConstArrayView<ESlateShader> GetSlateShaderTypesToPrecache()
 {
-	ESlateShader::Default,
-	ESlateShader::Border,
-	ESlateShader::GrayscaleFont,
-	ESlateShader::Custom,
-	ESlateShader::RoundedBox
-};
+	static const ESlateShader DynamicShaderTypes[] =
+	{ 
+		ESlateShader::Dynamic 
+	};
+
+	static const ESlateShader StaticShaderTypes[] =
+	{
+		ESlateShader::Default,
+		ESlateShader::Border,
+		ESlateShader::GrayscaleFont,
+		ESlateShader::Custom,
+		ESlateShader::RoundedBox
+	};
+
+	return GSlateDynamicShaders 
+		? TConstArrayView<ESlateShader>(DynamicShaderTypes)
+		: TConstArrayView<ESlateShader>(StaticShaderTypes);
+}
 
 static EPixelFormat GetDefaultBackBufferPixelFormat()
 {
@@ -1848,8 +1875,8 @@ void FSlateMaterialPSOCollector::CollectPSOInitializers(
 		return;
 	}
 
-	bool bRequired = false;	
-	for (ESlateShader ShaderType : SlateShaderTypesToPrecache)
+	const bool bRequired = false;
+	for (ESlateShader ShaderType : GetSlateShaderTypesToPrecache())
 	{
 		int32 UseInstancingCount = ShaderType == ESlateShader::Custom ? 2 : 1;
 		for (int32 UseInstancingIndex = 0; UseInstancingIndex < UseInstancingCount; ++UseInstancingIndex)
@@ -1870,8 +1897,8 @@ void FSlateMaterialPSOCollector::CollectPSOInitializers(
 			
 			TShaderRef<FSlateMaterialShaderVS> VertexShader;
 			TShaderRef<FSlateMaterialShaderPS> PixelShader;
-			Shaders.TryGetVertexShader(VertexShader); 
-			Shaders.TryGetPixelShader(PixelShader);						
+			Shaders.TryGetVertexShader(VertexShader);
+			Shaders.TryGetPixelShader(PixelShader);
 			if (!VertexShader.IsValid() || !PixelShader.IsValid())
 			{
 				continue;
