SkyLight - improved sky light illuminance view.

#tests PC editor dx12
#rnx
#DisabledValidations Preflight

diff --git a/Engine/Shaders/Private/ReflectionEnvironmentShaders.usf b/Engine/Shaders/Private/ReflectionEnvironmentShaders.usf
--- a/Engine/Shaders/Private/ReflectionEnvironmentShaders.usf
+++ b/Engine/Shaders/Private/ReflectionEnvironmentShaders.usf
@@ -966,6 +966,10 @@ void RenderRealTimeReflectionHeightFogPS(
 int CubeMipIndexToSampleIlluminance;
 float3 SkyLightCaptureWorlPos;
 
+int CurrentCubemapResolution;
+int SkyLightCubemapMaxResolution;
+int HardwareCubemapMaxResolution;
+
 TextureCube SkyLightCubeTexture;
 SamplerState SkyLightCubeSampler;
 
@@ -1066,8 +1070,18 @@ void VisualizeIlluminanceMeterPS(
 		//float SunIlluminance = dot(SunLightIlluminance.rgb, (1.0/3.0).xxx);
 		float TotalIlluminance = SkyIlluminance + SunIlluminance;
 
-		FShaderPrintContext Context = InitShaderPrintContext(true, uint2(110, 70));
-		Print(Context, TEXT("Sky light illuminance = "), InitFontColor(0.8, 0.8, 0.8));
+		FFontColor DefaultFontColor = InitFontColor(0.9, 0.9, 0.9);
+		FShaderPrintContext Context = InitShaderPrintContext(true, uint2(110, 20));
+		Print(Context, TEXT("Current sky light resolution = "), DefaultFontColor);
+		Print(Context, CurrentCubemapResolution, DefaultFontColor);
+		Print(Context, TEXT("(Hardware max resolution = "), DefaultFontColor);
+		Print(Context, HardwareCubemapMaxResolution, DefaultFontColor);
+		Print(Context, TEXT(", r.SkyLight.CubemapMaxResolution = "), DefaultFontColor);
+		Print(Context, SkyLightCubemapMaxResolution, DefaultFontColor);
+		Print(Context, TEXT(")"), DefaultFontColor);
+
+		Context = InitShaderPrintContext(true, uint2(110, 70));
+		Print(Context, TEXT("Sky light illuminance = "), DefaultFontColor);
 		Print(Context, SkyLightIlluminance.r, InitFontColor(1.0, 0.8, 0.8));
 		Print(Context, SkyLightIlluminance.g, InitFontColor(0.8, 1.0, 0.8));
 		Print(Context, SkyLightIlluminance.b, InitFontColor(0.8, 0.8, 1.0));
@@ -1076,7 +1090,7 @@ void VisualizeIlluminanceMeterPS(
 
 		Newline(Context);
 
-		Print(Context, TEXT("Sun light illuminance = "), InitFontColor(0.8, 0.8, 0.8));
+		Print(Context, TEXT("Sun light illuminance = "), DefaultFontColor);
 		Print(Context, SunLightIlluminance.r, InitFontColor(1.0, 0.8, 0.8));
 		Print(Context, SunLightIlluminance.g, InitFontColor(0.8, 1.0, 0.8));
 		Print(Context, SunLightIlluminance.b, InitFontColor(0.8, 0.8, 1.0));
@@ -1085,7 +1099,7 @@ void VisualizeIlluminanceMeterPS(
 
 		Newline(Context);
 
-		Print(Context, TEXT("Total illuminance     = "), InitFontColor(0.8, 0.8, 0.8));
+		Print(Context, TEXT("Total illuminance     = "), DefaultFontColor);
 		Print(Context, SkyLightIlluminance.r + SunLightIlluminance.r, InitFontColor(1.0, 0.8, 0.8));
 		Print(Context, SkyLightIlluminance.g + SunLightIlluminance.g, InitFontColor(0.8, 1.0, 0.8));
 		Print(Context, SkyLightIlluminance.b + SunLightIlluminance.b, InitFontColor(0.8, 0.8, 1.0));
diff --git a/Engine/Source/Runtime/Renderer/Private/IlluminanceMeter.cpp b/Engine/Source/Runtime/Renderer/Private/IlluminanceMeter.cpp
--- a/Engine/Source/Runtime/Renderer/Private/IlluminanceMeter.cpp
+++ b/Engine/Source/Runtime/Renderer/Private/IlluminanceMeter.cpp
@@ -72,6 +72,9 @@ class FPrintIlluminanceMeterPS : public FGlobalShader
 		SHADER_PARAMETER_TEXTURE(TextureCube, SkyLightCubeTexture)
 		SHADER_PARAMETER_SAMPLER(SamplerState, SkyLightCubeSampler)
 		SHADER_PARAMETER(int32, CubeMipIndexToSampleIlluminance)
+		SHADER_PARAMETER(int32, SkyLightCubemapMaxResolution)
+		SHADER_PARAMETER(int32, HardwareCubemapMaxResolution)
+		SHADER_PARAMETER(int32, CurrentCubemapResolution)
 		SHADER_PARAMETER(FVector3f, SkyLightCaptureWorlPos)
 		SHADER_PARAMETER_STRUCT_INCLUDE(ShaderPrint::FShaderParameters, ShaderPrintParameters)
 		RENDER_TARGET_BINDING_SLOTS()
@@ -216,6 +219,12 @@ FScreenPassTexture ProcessAndRenderIlluminanceMeter(FRDGBuilder& GraphBuilder, c
 		PassParameters->SkyLightCubeSampler = TStaticSamplerState<SF_Trilinear, AM_Clamp, AM_Clamp, AM_Clamp>::GetRHI();
 		PassParameters->CubeMipIndexToSampleIlluminance = IlluminanceMeterCubeTexture->Desc.NumMips - 1;
 		PassParameters->SkyLightCaptureWorlPos = FVector3f(SkyLight->CapturePosition);
+
+		PassParameters->CurrentCubemapResolution = SkyLightTextureRHI ? SkyLightTextureRHI->GetDesc().Extent.X : 0;
+		PassParameters->HardwareCubemapMaxResolution = GetMaxCubeTextureDimension();
+		static auto* SkyLightCubemapMaxResolutionCVar = IConsoleManager::Get().FindConsoleVariable(TEXT("r.SkyLight.CubemapMaxResolution"));
+		PassParameters->SkyLightCubemapMaxResolution = SkyLightCubemapMaxResolutionCVar ? SkyLightCubemapMaxResolutionCVar->GetInt() : 0;
+
 		ShaderPrint::SetParameters(GraphBuilder, View.ShaderPrintData, PassParameters->ShaderPrintParameters);
 		PassParameters->RenderTargets[0] = FRenderTargetBinding(ScreenPassSceneColor.Texture, ERenderTargetLoadAction::ELoad);
 
