ACES: NaN fixes for CombineLUTs pass.

* Was corrupting (0, 0, 0) LUT value with ACES 2.0 tonemapping.

diff --git a/Engine/Shaders/Private/ACES/ACESOutputTransform.ush b/Engine/Shaders/Private/ACES/ACESOutputTransform.ush
--- a/Engine/Shaders/Private/ACES/ACESOutputTransform.ush
+++ b/Engine/Shaders/Private/ACES/ACESOutputTransform.ush
@@ -183,7 +183,7 @@ float3 XYZ_to_Hellwig2022_JMh(  float3 XYZ,
 
     // Computing the hue angle math h
     // Unclear why this isn't matching the python version
-    float hr = atan2(b, a);
+    float hr = satan2(b, a);
     float h = wrap_to_360(radians_to_degrees(hr));
 
     // Step 6 - Computing achromatic responses for the stimulus
diff --git a/Engine/Shaders/Private/ACES/ACESUtilities.ush b/Engine/Shaders/Private/ACES/ACESUtilities.ush
--- a/Engine/Shaders/Private/ACES/ACESUtilities.ush
+++ b/Engine/Shaders/Private/ACES/ACESUtilities.ush
@@ -59,6 +59,11 @@ float sdiv(float a, float b)
     return b == 0.0 ? 0.0 : a / b;
 }
 
+float satan2(float y, float x)
+{
+    return (x == 0.0 && y == 0.0) ? 0.0 : atan2(y, x);
+}
+
 #else // UE_ACES_SAFE_FUNCTIONS
 
 #define spow pow
@@ -68,4 +73,6 @@ float sdiv(float a, float b)
     return a / b;
 }
 
+#define satan2 atan2
+
 #endif // UE_SAFE_FUNCTIONS_SLOW
