Use stencil bit to mark foreground samples in EditorPrimitives, and treat those as if at ZNear while compositing

#FailedValidations Preflight

diff --git a/Engine/Shaders/Private/PostProcessCompositePrimitives.usf b/Engine/Shaders/Private/PostProcessCompositePrimitives.usf
--- a/Engine/Shaders/Private/PostProcessCompositePrimitives.usf
+++ b/Engine/Shaders/Private/PostProcessCompositePrimitives.usf
@@ -75,9 +75,11 @@ SamplerState UndistortingDisplacementSampler;
 #if MSAA_SAMPLE_COUNT > 1
 Texture2DMS<float4, MSAA_SAMPLE_COUNT> EditorPrimitivesColor;
 Texture2DMS<float, MSAA_SAMPLE_COUNT> EditorPrimitivesDepth;
+Texture2DMS<uint2, MSAA_SAMPLE_COUNT> EditorPrimitivesStencil;
 #else
 Texture2D EditorPrimitivesColor;
 Texture2D<float> EditorPrimitivesDepth;
+Texture2D<uint2> EditorPrimitivesStencil;
 #endif
 
 Texture2D ColorTexture;
@@ -351,8 +353,15 @@ void MainCompositePostProcessPrimitivesPS(
 
 		EditorPrimitiveColor = EditorPrimitivesColor.Load(EditorPrimitivesPixelPos, TargetSampleIndex);
 		float SampleEditorDeviceZ = EditorPrimitivesDepth.Load(EditorPrimitivesPixelPos, TargetSampleIndex).r;
+		uint SampleEditorStencil = EditorPrimitivesStencil.Load(EditorPrimitivesPixelPos, TargetSampleIndex).g;
 		float SampleSceneDeviceZ = Texture2DSampleLevel(DepthTexture, DepthSampler, SampleDepthUV, 0).r;
 
+		bool bIsForegroundSample = SampleEditorStencil & EDITOR_STENCIL_FLAG_FOREGROUND;
+		if (bIsForegroundSample)
+		{
+			SampleEditorDeviceZ = NearDepthValue;
+		}
+
 		if (SampleEditorDeviceZ > 0.0 && bCompositeAnyNonNullDepth)
 		{
 			EditorPrimitiveColor.a = 1;
@@ -380,6 +389,13 @@ void MainCompositePostProcessPrimitivesPS(
 
 			float4 Sample = EditorPrimitivesColor.Load(EditorPrimitivesPixelPos, SampleIndex);
 			float SampleEditorDeviceZ = EditorPrimitivesDepth.Load(EditorPrimitivesPixelPos, SampleIndex).r;
+			uint SampleEditorStencil = EditorPrimitivesStencil.Load(EditorPrimitivesPixelPos, SampleIndex).g;
+
+			bool bIsForegroundSample = SampleEditorStencil & EDITOR_STENCIL_FLAG_FOREGROUND;
+			if (bIsForegroundSample)
+			{
+				SampleEditorDeviceZ = NearDepthValue;
+			}
 
 			// Check if any color was applied to this pixel. Note: This prevents actual black pixels from being visible.
 			float Weight = Sample.a;
@@ -420,6 +436,13 @@ void MainCompositePostProcessPrimitivesPS(
 
 		float SceneDeviceZ = Texture2DSampleLevel(DepthTexture, DepthSampler, SampleDepthUV, 0).r;
 		float EditorDeviceZ = EditorPrimitivesDepth.Load(int3(EditorPrimitivesPixelPos, 0)).r;
+		uint EditorStencil = EditorPrimitivesStencil.Load(int3(EditorPrimitivesPixelPos, 0)).g;
+
+		bool bIsForegroundSample = EditorStencil & EDITOR_STENCIL_FLAG_FOREGROUND;
+		if (bIsForegroundSample)
+		{
+			EditorDeviceZ = NearDepthValue;
+		}
 		
 		EditorPrimitiveColor = EditorPrimitivesColor.Load(int3(EditorPrimitivesPixelPos, 0));
 		DepthMask = ComputeDepthMask(SceneDeviceZ, EditorDeviceZ);
diff --git a/Engine/Source/Runtime/Renderer/Private/PostProcess/PostProcessCompositeDebugPrimitives.cpp b/Engine/Source/Runtime/Renderer/Private/PostProcess/PostProcessCompositeDebugPrimitives.cpp
--- a/Engine/Source/Runtime/Renderer/Private/PostProcess/PostProcessCompositeDebugPrimitives.cpp
+++ b/Engine/Source/Runtime/Renderer/Private/PostProcess/PostProcessCompositeDebugPrimitives.cpp
@@ -211,7 +211,8 @@ static void ComposeDebugPrimitives(
 	PassParameters->ColorSampler = PointClampSampler;
 	PassParameters->DepthTexture = SceneDepth.Texture;
 	PassParameters->DepthSampler = PointClampSampler;
-	PassParameters->EditorPrimitivesDepth = Context.DebugPrimitiveDepth;
+	PassParameters->EditorPrimitivesDepth = GraphBuilder.CreateSRV(FRDGTextureSRVDesc::CreateWithPixelFormat(Context.DebugPrimitiveDepth, PF_DepthStencil));
+	PassParameters->EditorPrimitivesStencil = GraphBuilder.CreateSRV(FRDGTextureSRVDesc::CreateWithPixelFormat(Context.DebugPrimitiveDepth, PF_X24_G8));
 	PassParameters->EditorPrimitivesColor = Context.DebugPrimitiveColor;
 		
 	PassParameters->PassSvPositionToViewportUV = FScreenTransform::SvPositionToViewportUV(Context.OutputRT.ViewRect);
diff --git a/Engine/Source/Runtime/Renderer/Private/PostProcess/PostProcessCompositeEditorPrimitives.cpp b/Engine/Source/Runtime/Renderer/Private/PostProcess/PostProcessCompositeEditorPrimitives.cpp
--- a/Engine/Source/Runtime/Renderer/Private/PostProcess/PostProcessCompositeEditorPrimitives.cpp
+++ b/Engine/Source/Runtime/Renderer/Private/PostProcess/PostProcessCompositeEditorPrimitives.cpp
@@ -156,7 +156,8 @@ static void RenderForegroundEditorPrimitives(FRHICommandList& RHICmdList, const
 
 	// Draw a second time the foreground primitive with depth test to have proper depth test between foreground primitives.
 	{
-		DrawRenderState.SetDepthStencilState(TStaticDepthStencilState<true, CF_DepthNearOrEqual>::GetRHI());
+		DrawRenderState.SetDepthStencilState(TStaticDepthStencilState<true, CF_DepthNearOrEqual, true, CF_Always, SO_Keep, SO_Keep, SO_Replace, true, CF_Always, SO_Keep, SO_Keep, SO_Replace, 0, 0xFF>::GetRHI());
+		DrawRenderState.SetStencilRef(FCompositePostProcessPrimitivesPS::EditorStencilFlagForeground);
 
 		View.EditorSimpleElementCollector.DrawBatchedElements(RHICmdList, DrawRenderState, View, EBlendModeFilter::OpaqueAndMasked, SDPG_Foreground);
 
@@ -442,7 +443,8 @@ FScreenPassTexture AddEditorPrimitivePass(
 			PassParameters->DepthTexture = GSystemTextures.GetDepthDummy(GraphBuilder);
 		}
 		PassParameters->DepthSampler = PointClampSampler;
-		PassParameters->EditorPrimitivesDepth = EditorPrimitiveDepth;
+		PassParameters->EditorPrimitivesDepth = GraphBuilder.CreateSRV(FRDGTextureSRVDesc::CreateWithPixelFormat(EditorPrimitiveDepth, PF_DepthStencil));
+		PassParameters->EditorPrimitivesStencil = GraphBuilder.CreateSRV(FRDGTextureSRVDesc::CreateWithPixelFormat(EditorPrimitiveDepth, PF_X24_G8));
 		PassParameters->EditorPrimitivesColor = EditorPrimitiveColor;
 		
 		PassParameters->PassSvPositionToViewportUV = FScreenTransform::SvPositionToViewportUV(Output.ViewRect);
diff --git a/Engine/Source/Runtime/Renderer/Private/PostProcess/PostProcessCompositePrimitivesCommon.h b/Engine/Source/Runtime/Renderer/Private/PostProcess/PostProcessCompositePrimitivesCommon.h
--- a/Engine/Source/Runtime/Renderer/Private/PostProcess/PostProcessCompositePrimitivesCommon.h
+++ b/Engine/Source/Runtime/Renderer/Private/PostProcess/PostProcessCompositePrimitivesCommon.h
@@ -130,7 +130,8 @@ public:
 
 		SHADER_PARAMETER_RDG_TEXTURE(Texture2D, UndistortingDisplacementTexture)
 		SHADER_PARAMETER_SAMPLER(SamplerState,  UndistortingDisplacementSampler)
-		SHADER_PARAMETER_RDG_TEXTURE(Texture2D, EditorPrimitivesDepth)
+		SHADER_PARAMETER_RDG_TEXTURE_SRV(Texture2D<float>, EditorPrimitivesDepth)
+		SHADER_PARAMETER_RDG_TEXTURE_SRV(Texture2D<uint2>, EditorPrimitivesStencil)
 		SHADER_PARAMETER_RDG_TEXTURE(Texture2D, EditorPrimitivesColor)
 		SHADER_PARAMETER_RDG_TEXTURE(Texture2D, ColorTexture)
 		SHADER_PARAMETER_SAMPLER(SamplerState,  ColorSampler)
@@ -151,6 +152,8 @@ public:
 		RENDER_TARGET_BINDING_SLOTS()
 	END_SHADER_PARAMETER_STRUCT()
 
+	static constexpr uint32 EditorStencilFlagForeground = 0x1;
+
 	static bool ShouldCompilePermutation(const FPermutationDomain& PermutationVector, const EShaderPlatform Platform)
 	{
 		const int32 SampleCount = PermutationVector.Get<FSampleCountDimension>();
@@ -178,6 +181,8 @@ public:
 	static void ModifyCompilationEnvironment(const FGlobalShaderPermutationParameters& Parameters, FShaderCompilerEnvironment& OutEnvironment)
 	{
 		FCompositePrimitiveShaderBase::ModifyCompilationEnvironment(Parameters, OutEnvironment);
+
+		OutEnvironment.SetDefine(TEXT("EDITOR_STENCIL_FLAG_FOREGROUND"), EditorStencilFlagForeground);
 	}
 };
 
diff --git a/Engine/Source/Runtime/Renderer/Private/SceneTextures.cpp b/Engine/Source/Runtime/Renderer/Private/SceneTextures.cpp
--- a/Engine/Source/Runtime/Renderer/Private/SceneTextures.cpp
+++ b/Engine/Source/Runtime/Renderer/Private/SceneTextures.cpp
@@ -785,7 +785,8 @@ void FSceneTextures::InitializeViewFamily(FRDGBuilder& GraphBuilder, FViewFamily
 		const FRDGTextureDesc ColorDesc(FRDGTextureDesc::Create2D(Config.Extent, PF_B8G8R8A8, FClearValueBinding::Transparent, TexCreate_ShaderResource | TexCreate_RenderTargetable, 1, Config.EditorPrimitiveNumSamples));
 		SceneTextures.EditorPrimitiveColor = GraphBuilder.CreateTexture(ColorDesc, TEXT("Editor.PrimitivesColor"));
 
-		const FRDGTextureDesc DepthDesc(FRDGTextureDesc::Create2D(Config.Extent, PF_DepthStencil, FClearValueBinding::DepthFar, TexCreate_ShaderResource | TexCreate_DepthStencilTargetable, 1, Config.EditorPrimitiveNumSamples));
+		FRDGTextureDesc DepthDesc(FRDGTextureDesc::Create2D(Config.Extent, PF_DepthStencil, FClearValueBinding::DepthFar, TexCreate_ShaderResource | TexCreate_DepthStencilTargetable, 1, Config.EditorPrimitiveNumSamples));
+		DepthDesc.AliasableFormats.Add(PF_X24_G8);
 		SceneTextures.EditorPrimitiveDepth = GraphBuilder.CreateTexture(DepthDesc, TEXT("Editor.PrimitivesDepth"));
 	}
 #endif
