Fix missing EDITOR_STENCIL_FLAG_FOREGROUND define

diff --git a/Engine/Shaders/Private/PostProcessCompositePrimitives.usf b/Engine/Shaders/Private/PostProcessCompositePrimitives.usf
--- a/Engine/Shaders/Private/PostProcessCompositePrimitives.usf
+++ b/Engine/Shaders/Private/PostProcessCompositePrimitives.usf
@@ -9,6 +9,11 @@
 #define MSAA_SAMPLE_COUNT 1
 #endif
 
+#ifndef EDITOR_STENCIL_FLAG_FOREGROUND
+#define EDITOR_STENCIL_FLAG_FOREGROUND 0
+#endif
+
+
 #define CONFIG_SAMPLES 6
 
 // K = Center of the nearest input pixel.
