[Stochastic lighting] Fix upsample weights computation using incorrect sample depth

#tests tested in editor

diff --git a/Engine/Shaders/Private/StochasticLighting/StochasticLightingTileClassification.usf b/Engine/Shaders/Private/StochasticLighting/StochasticLightingTileClassification.usf
--- a/Engine/Shaders/Private/StochasticLighting/StochasticLightingTileClassification.usf
+++ b/Engine/Shaders/Private/StochasticLighting/StochasticLightingTileClassification.usf
@@ -342,7 +342,7 @@ float4 ComputeUpsampleWeights(uint2 GroupThreadId, uint2 ScreenCoord, float2 Scr
 
 		float3 Position00 = mul(float4(GetScreenPositionForProjectionType(ScreenPosition00, CornerDepths.x), CornerDepths.x, 1), View.ScreenToTranslatedWorld).xyz;
 		float3 Position10 = mul(float4(GetScreenPositionForProjectionType(ScreenPosition10, CornerDepths.y), CornerDepths.y, 1), View.ScreenToTranslatedWorld).xyz;
-		float3 Position01 = mul(float4(GetScreenPositionForProjectionType(ScreenPosition01, CornerDepths.z), CornerDepths.x, 1), View.ScreenToTranslatedWorld).xyz;
+		float3 Position01 = mul(float4(GetScreenPositionForProjectionType(ScreenPosition01, CornerDepths.z), CornerDepths.z, 1), View.ScreenToTranslatedWorld).xyz;
 		float3 Position11 = mul(float4(GetScreenPositionForProjectionType(ScreenPosition11, CornerDepths.w), CornerDepths.w, 1), View.ScreenToTranslatedWorld).xyz;
 
 		float4 PlaneDistances;
