SkyAtmosphere - fix for aerial perspective affecting the look of the sky when in per pixel tracing mode. (it should only affect AP on opqaue)

#tests pc editor dx12
#rnx

diff --git a/Engine/Shaders/Private/SkyAtmosphere.usf b/Engine/Shaders/Private/SkyAtmosphere.usf
--- a/Engine/Shaders/Private/SkyAtmosphere.usf
+++ b/Engine/Shaders/Private/SkyAtmosphere.usf
@@ -1073,8 +1073,26 @@ void RenderSkyAtmosphereRayMarchingPS(
 		return;
 	}
 
+#if RENDERSKY_ENABLED
+	if (DeviceZ != FarDepthValue)
+	{
+		// Offset start of integration only if we are going ot hit an opaque pixel (aerial perspective part) and only if we are not going to go past it.
+		const float3 DepthBufferTranslatedWorldPos = GetScreenTranslatedWorldPos(SVPos, DeviceZ).xyz;
+		if (length(DepthBufferTranslatedWorldPos - GetCameraTranslatedWorldPos()) > AerialPerspectiveStartDepthKm * KILOMETER_TO_CENTIMETER)
+		{
+			WorldPos += WorldDir * AerialPerspectiveStartDepthKm;
+		}
+		else
+		{
+			OutLuminance = PrepareOutput(float3(0.0f, 0.0f, 0.0f), float3(1.0f, 1.0f, 1.0f));
+			clip(-1.0f);
+			return;
+		}
+	}
+#else
 	// Apply the start depth offset after moving to the top of atmosphere for consistency (and to avoid wrong out-of-atmosphere test resulting in black pixels).
 	WorldPos += WorldDir * AerialPerspectiveStartDepthKm;
+#endif
 
 	SamplingSetup Sampling = (SamplingSetup)0;
 	{
