Fixed DiffuseIndirectComposite compile errors from code moved out of define

diff --git a/Engine/Shaders/Private/DiffuseIndirectComposite.usf b/Engine/Shaders/Private/DiffuseIndirectComposite.usf
--- a/Engine/Shaders/Private/DiffuseIndirectComposite.usf
+++ b/Engine/Shaders/Private/DiffuseIndirectComposite.usf
@@ -350,6 +350,8 @@ FSubstrateDeferredLighting SubstrateIndirectLighting(
 }
 #endif // SUBSTRATE_GBUFFER_FORMAT==1
 
+#if DIM_APPLY_DIFFUSE_INDIRECT == DIM_APPLY_DIFFUSE_INDIRECT_SSGI
+
 float3 UpsampleLegacySSGI(uint2 PixelPos, float2 SceneBufferUV, FLumenMaterialData Material)
 {
 	// Sample the output of the denoiser.
@@ -425,6 +427,8 @@ float3 UpsampleLegacySSGI(uint2 PixelPos, float2 SceneBufferUV, FLumenMaterialDa
 	return Sample.SceneColor.rgb;
 }
 
+#endif
+
 ///////////////////////////////////////////////////////////////////////////////////////////////////////////////////////////////////////
 
 void MainPS(
