Fix regression in transform gizmo occlusion dithering during CompositeEditorPrimitives

#rnx

diff --git a/Engine/Shaders/Private/PostProcessCompositePrimitives.usf b/Engine/Shaders/Private/PostProcessCompositePrimitives.usf
--- a/Engine/Shaders/Private/PostProcessCompositePrimitives.usf
+++ b/Engine/Shaders/Private/PostProcessCompositePrimitives.usf
@@ -9,8 +9,8 @@
 #define MSAA_SAMPLE_COUNT 1
 #endif
 
-#ifndef EDITOR_STENCIL_FLAG_FOREGROUND
-#define EDITOR_STENCIL_FLAG_FOREGROUND 0
+#ifndef EDITOR_STENCIL_FLAG_FOREGROUND_NO_DITHERING
+#define EDITOR_STENCIL_FLAG_FOREGROUND_NO_DITHERING 0
 #endif
 
 
@@ -361,8 +361,8 @@ void MainCompositePostProcessPrimitivesPS(
 		uint SampleEditorStencil = EditorPrimitivesStencil.Load(EditorPrimitivesPixelPos, TargetSampleIndex).g;
 		float SampleSceneDeviceZ = Texture2DSampleLevel(DepthTexture, DepthSampler, SampleDepthUV, 0).r;
 
-		bool bIsForegroundSample = SampleEditorStencil & EDITOR_STENCIL_FLAG_FOREGROUND;
-		if (bIsForegroundSample)
+		bool bIsForegroundNoDitheringSample = SampleEditorStencil & EDITOR_STENCIL_FLAG_FOREGROUND_NO_DITHERING;
+		if (bIsForegroundNoDitheringSample)
 		{
 			SampleEditorDeviceZ = NearDepthValue;
 		}
@@ -396,8 +396,8 @@ void MainCompositePostProcessPrimitivesPS(
 			float SampleEditorDeviceZ = EditorPrimitivesDepth.Load(EditorPrimitivesPixelPos, SampleIndex).r;
 			uint SampleEditorStencil = EditorPrimitivesStencil.Load(EditorPrimitivesPixelPos, SampleIndex).g;
 
-			bool bIsForegroundSample = SampleEditorStencil & EDITOR_STENCIL_FLAG_FOREGROUND;
-			if (bIsForegroundSample)
+			bool bIsForegroundNoDitheringSample = SampleEditorStencil & EDITOR_STENCIL_FLAG_FOREGROUND_NO_DITHERING;
+			if (bIsForegroundNoDitheringSample)
 			{
 				SampleEditorDeviceZ = NearDepthValue;
 			}
@@ -443,8 +443,8 @@ void MainCompositePostProcessPrimitivesPS(
 		float EditorDeviceZ = EditorPrimitivesDepth.Load(int3(EditorPrimitivesPixelPos, 0)).r;
 		uint EditorStencil = EditorPrimitivesStencil.Load(int3(EditorPrimitivesPixelPos, 0)).g;
 
-		bool bIsForegroundSample = EditorStencil & EDITOR_STENCIL_FLAG_FOREGROUND;
-		if (bIsForegroundSample)
+		bool bIsForegroundNoDitheringSample = EditorStencil & EDITOR_STENCIL_FLAG_FOREGROUND_NO_DITHERING;
+		if (bIsForegroundNoDitheringSample)
 		{
 			EditorDeviceZ = NearDepthValue;
 		}
diff --git a/Engine/Source/Runtime/Renderer/Private/PostProcess/PostProcessCompositeEditorPrimitives.cpp b/Engine/Source/Runtime/Renderer/Private/PostProcess/PostProcessCompositeEditorPrimitives.cpp
--- a/Engine/Source/Runtime/Renderer/Private/PostProcess/PostProcessCompositeEditorPrimitives.cpp
+++ b/Engine/Source/Runtime/Renderer/Private/PostProcess/PostProcessCompositeEditorPrimitives.cpp
@@ -157,10 +157,11 @@ static void RenderForegroundEditorPrimitives(FRHICommandList& RHICmdList, const
 	// Draw a second time the foreground primitive with depth test to have proper depth test between foreground primitives.
 	{
 		DrawRenderState.SetDepthStencilState(TStaticDepthStencilState<true, CF_DepthNearOrEqual, true, CF_Always, SO_Keep, SO_Keep, SO_Replace, true, CF_Always, SO_Keep, SO_Keep, SO_Replace, 0, 0xFF>::GetRHI());
-		DrawRenderState.SetStencilRef(FCompositePostProcessPrimitivesPS::EditorStencilFlagForeground);
 
+		DrawRenderState.SetStencilRef(FCompositePostProcessPrimitivesPS::EditorStencilFlagForegroundNoDithering);
 		View.EditorSimpleElementCollector.DrawBatchedElements(RHICmdList, DrawRenderState, View, EBlendModeFilter::OpaqueAndMasked, SDPG_Foreground);
 
+		DrawRenderState.SetStencilRef(0);
 		DrawDynamicMeshPass(View, RHICmdList,
 			[&View, &DrawRenderState](FDynamicPassMeshDrawListContext* DynamicMeshPassContext)
 		{
diff --git a/Engine/Source/Runtime/Renderer/Private/PostProcess/PostProcessCompositePrimitivesCommon.h b/Engine/Source/Runtime/Renderer/Private/PostProcess/PostProcessCompositePrimitivesCommon.h
--- a/Engine/Source/Runtime/Renderer/Private/PostProcess/PostProcessCompositePrimitivesCommon.h
+++ b/Engine/Source/Runtime/Renderer/Private/PostProcess/PostProcessCompositePrimitivesCommon.h
@@ -152,7 +152,9 @@ public:
 		RENDER_TARGET_BINDING_SLOTS()
 	END_SHADER_PARAMETER_STRUCT()
 
-	static constexpr uint32 EditorStencilFlagForeground = 0x1;
+	// SDPG_World means "with occlusion", SDPG_Foreground means "without occlusion". But, SDPG_Foreground can be drawn fully opaque, or with dithering on 'occluded' areas (e.g. gizmos).
+	// The stencil flag marks which areas should *not* receive dithering.
+	static constexpr uint32 EditorStencilFlagForegroundNoDithering = 0x1;
 
 	static bool ShouldCompilePermutation(const FPermutationDomain& PermutationVector, const EShaderPlatform Platform)
 	{
@@ -182,7 +184,7 @@ public:
 	{
 		FCompositePrimitiveShaderBase::ModifyCompilationEnvironment(Parameters, OutEnvironment);
 
-		OutEnvironment.SetDefine(TEXT("EDITOR_STENCIL_FLAG_FOREGROUND"), EditorStencilFlagForeground);
+		OutEnvironment.SetDefine(TEXT("EDITOR_STENCIL_FLAG_FOREGROUND_NO_DITHERING"), EditorStencilFlagForegroundNoDithering);
 	}
 };
 
