Fix a bug with the work distributor shader code that happens when the wave lane count is larger than the compute shader's thread group size.

diff --git a/Engine/Shaders/Private/WorkDistribution.ush b/Engine/Shaders/Private/WorkDistribution.ush
--- a/Engine/Shaders/Private/WorkDistribution.ush
+++ b/Engine/Shaders/Private/WorkDistribution.ush
@@ -222,7 +222,7 @@ groupshared uint WorkBoundary[ THREADGROUP_SIZE ];
 template< typename FTask >
 void DistributeWork( FTask Task, uint GroupIndex, uint NumWorkItems )
 {
-	const uint LaneCount	= WaveGetLaneCount();
+	const uint LaneCount	= min(WaveGetLaneCount(), THREADGROUP_SIZE);
 	const uint LaneIndex	= GroupIndex &  ( LaneCount - 1 );
 	const uint QueueOffset	= GroupIndex & ~( LaneCount - 1 );
 
