[DOF] Added missing ResolveView to DOF shaders.
In 44062835, a change was made to ViewportUVToBufferUV to use ResolvedView. This lead to undefined behaviour, since ResolvedView was not set.

diff --git a/Engine/Shaders/Private/DiaphragmDOF/DOFCocTileDilate.usf b/Engine/Shaders/Private/DiaphragmDOF/DOFCocTileDilate.usf
--- a/Engine/Shaders/Private/DiaphragmDOF/DOFCocTileDilate.usf
+++ b/Engine/Shaders/Private/DiaphragmDOF/DOFCocTileDilate.usf
@@ -45,6 +45,8 @@ RWTexture2D<float4>	TileOutput_Background;
 [numthreads(DEFAULT_GROUP_BORDER_SIZE, DEFAULT_GROUP_BORDER_SIZE, 1)]
 void CocDilateMainCS(uint2 DispatchThreadId : SV_DispatchThreadID) 
 {
+	ResolvedView = ResolveView();
+
 	const uint CocTileLayout = COC_TILE_LAYOUT_FGD_SEP_BGD;
 	const uint Radius = DIM_DILATE_RADIUS;
 
diff --git a/Engine/Shaders/Private/DiaphragmDOF/DOFCocTileFlatten.usf b/Engine/Shaders/Private/DiaphragmDOF/DOFCocTileFlatten.usf
--- a/Engine/Shaders/Private/DiaphragmDOF/DOFCocTileFlatten.usf
+++ b/Engine/Shaders/Private/DiaphragmDOF/DOFCocTileFlatten.usf
@@ -126,6 +126,8 @@ void CocFlattenMainCS(
 	uint2 DispatchThreadId : SV_DispatchThreadID,
 	uint GroupIndex : SV_GroupIndex)
 {
+	ResolvedView = ResolveView();
+
 	FCocTileSample TileSample;
 
 	#if DIM_DO_COC_GATHER4
diff --git a/Engine/Shaders/Private/DiaphragmDOF/DOFDownsample.usf b/Engine/Shaders/Private/DiaphragmDOF/DOFDownsample.usf
--- a/Engine/Shaders/Private/DiaphragmDOF/DOFDownsample.usf
+++ b/Engine/Shaders/Private/DiaphragmDOF/DOFDownsample.usf
@@ -32,6 +32,8 @@ RWTexture2D<float>	OutDownsampledGatherInput_SeparateCoc;
 [numthreads(DEFAULT_GROUP_BORDER_SIZE, DEFAULT_GROUP_BORDER_SIZE, 1)]
 void DownsampleCS(uint2 DispatchThreadId : SV_DispatchThreadID) 
 {
+	ResolvedView = ResolveView();
+
 	// Compute scene buffer uv, but carefully to handle view size % 2 == 1,
 	// because this pass is running at half resolution.
 	float2 BufferUV = (2.0 * (DispatchThreadId + 0.5)) * GatherInputSize.zw;
diff --git a/Engine/Shaders/Private/DiaphragmDOF/DOFGatherPass.usf b/Engine/Shaders/Private/DiaphragmDOF/DOFGatherPass.usf
--- a/Engine/Shaders/Private/DiaphragmDOF/DOFGatherPass.usf
+++ b/Engine/Shaders/Private/DiaphragmDOF/DOFGatherPass.usf
@@ -152,6 +152,8 @@ void GatherMainCS(
 	uint2 DispatchThreadId : SV_DispatchThreadID,
 	uint2 GroupId : SV_GroupID) 
 {
+	ResolvedView = ResolveView();
+
 	float4 Debug = 0;
 
 	// Get output's UVs.
diff --git a/Engine/Shaders/Private/DiaphragmDOF/DOFHybridScatterPixelShader.usf b/Engine/Shaders/Private/DiaphragmDOF/DOFHybridScatterPixelShader.usf
--- a/Engine/Shaders/Private/DiaphragmDOF/DOFHybridScatterPixelShader.usf
+++ b/Engine/Shaders/Private/DiaphragmDOF/DOFHybridScatterPixelShader.usf
@@ -117,6 +117,8 @@ void ScatterMainPS(
 	in float4 SvPosition : SV_Position,
 	out float4 OutColor : SV_Target0) 
 {
+	ResolvedView = ResolveView();
+
 	OutColor = 0;
 
 	const float2 BokehCenter0 = BokehCenterScreenAndClipPos.xy;
diff --git a/Engine/Shaders/Private/DiaphragmDOF/DOFHybridScatterVertexShader.usf b/Engine/Shaders/Private/DiaphragmDOF/DOFHybridScatterVertexShader.usf
--- a/Engine/Shaders/Private/DiaphragmDOF/DOFHybridScatterVertexShader.usf
+++ b/Engine/Shaders/Private/DiaphragmDOF/DOFHybridScatterVertexShader.usf
@@ -83,6 +83,8 @@ void ScatterMainVS(
 #endif
 	out float4 OutPosition : SV_POSITION)
 {
+	ResolvedView = ResolveView();
+
 	#if SCATTERING_GROUP_PACKING
 		// Number scattering group that can be packed per instance.
 		const uint ScatteringGroupPerInstance = VERTEX_SHADER_INVOCATION_TARGET / VERTEX_SHADER_INVOCATION_PER_SCATTERING_GROUP;
diff --git a/Engine/Shaders/Private/DiaphragmDOF/DOFPostfiltering.usf b/Engine/Shaders/Private/DiaphragmDOF/DOFPostfiltering.usf
--- a/Engine/Shaders/Private/DiaphragmDOF/DOFPostfiltering.usf
+++ b/Engine/Shaders/Private/DiaphragmDOF/DOFPostfiltering.usf
@@ -160,6 +160,8 @@ void PostfilterMainCS(
 	uint2 DispatchThreadId : SV_DispatchThreadID,
 	uint2 GroupId : SV_GroupID)
 {
+	ResolvedView = ResolveView();
+
 	float2 InputBufferUV = (DispatchThreadId + 0.5) * ConvolutionInputSize.zw;
 	
 	BRANCH
diff --git a/Engine/Shaders/Private/DiaphragmDOF/DOFRecombine.usf b/Engine/Shaders/Private/DiaphragmDOF/DOFRecombine.usf
--- a/Engine/Shaders/Private/DiaphragmDOF/DOFRecombine.usf
+++ b/Engine/Shaders/Private/DiaphragmDOF/DOFRecombine.usf
@@ -692,6 +692,8 @@ void RecombineMainCS(
 	uint2 GroupThreadId : SV_GroupThreadID,
 	uint GroupThreadIndex : SV_GroupIndex) 
 {
+	ResolvedView = ResolveView();
+
 	float4 Debug = 0;
 
 	// Setup input parameters.
diff --git a/Engine/Shaders/Private/DiaphragmDOF/DOFReduce.usf b/Engine/Shaders/Private/DiaphragmDOF/DOFReduce.usf
--- a/Engine/Shaders/Private/DiaphragmDOF/DOFReduce.usf
+++ b/Engine/Shaders/Private/DiaphragmDOF/DOFReduce.usf
@@ -308,6 +308,8 @@ void ReduceCS(
 	uint2 GroupId : SV_GroupID,
 	uint GroupThreadIndex : SV_GroupIndex) 
 {
+	ResolvedView = ResolveView();
+
 	// Init LDS.
 	#if (DIM_HYBRID_SCATTER_FGD || DIM_HYBRID_SCATTER_BGD) && SCATTER_ALLOC_METHOD == SCATTER_ALLOC_METHOD_ATOMIC
 	{
diff --git a/Engine/Shaders/Private/DiaphragmDOF/DOFSetup.usf b/Engine/Shaders/Private/DiaphragmDOF/DOFSetup.usf
--- a/Engine/Shaders/Private/DiaphragmDOF/DOFSetup.usf
+++ b/Engine/Shaders/Private/DiaphragmDOF/DOFSetup.usf
@@ -399,6 +399,8 @@ void SetupCS(
 	uint GroupThreadIndex : SV_GroupIndex,
 	uint2 DispatchThreadId : SV_DispatchThreadID)
 {
+	ResolvedView = ResolveView();
+
 	FDOFSetupParameters Parameters;
 	Parameters.bOutputFullRes = true;
 	Parameters.bOutputHalfRes = true;
@@ -445,6 +447,8 @@ void SetupCS(
 	uint GroupThreadIndex : SV_GroupIndex,
 	uint2 DispatchThreadId : SV_DispatchThreadID)
 {
+	ResolvedView = ResolveView();
+
 	FDOFSetupParameters Parameters;
 	Parameters.bOutputFullRes = true;
 	Parameters.bOutputHalfRes = false;
@@ -477,6 +481,8 @@ void SetupCS(
 	uint GroupThreadIndex : SV_GroupIndex,
 	uint2 DispatchThreadId : SV_DispatchThreadID)
 {
+	ResolvedView = ResolveView();
+
 	FDOFSetupParameters Parameters;
 	Parameters.bOutputFullRes = false;
 	Parameters.bOutputHalfRes = true;
