Path Tracer: Make sure we don't advance the random number generator if scattering is not present

This helps avoid a difference in behavior between dx12 and vulkan RHIs

#rnx

diff --git a/Engine/Shaders/Private/PathTracing/PathTracingCore.ush b/Engine/Shaders/Private/PathTracing/PathTracingCore.ush
--- a/Engine/Shaders/Private/PathTracing/PathTracingCore.ush
+++ b/Engine/Shaders/Private/PathTracing/PathTracingCore.ush
@@ -714,7 +714,14 @@ FPackedPathTracingPayload TraceTransparentRay(inout FPathState PathState, int Bo
 			float4 ScatterData = PathState.GetSigmaSAndPhase();
 			float3 MediumAlbedo = select(SigmaT > 0, ScatterData.xyz / SigmaT, 0.0);
 			float3 ColorChannelPdf = PathState.PathThroughput * MediumAlbedo;
-			float ScatterT = !bLastBounce ? SampleSpectralTransmittance(PathState.RandSequence.Get1D(), SigmaT, ColorChannelPdf) : -1;
+			float ScatterT = -1; 
+			BRANCH
+			if (!bLastBounce && any(ColorChannelPdf > 0))
+			{
+				// Only sample scatter hits if we have some scattering albedo and this is not the last bounce (which only cares about emission)
+				float RandSample = PathState.RandSequence.Get1D();
+				ScatterT = SampleSpectralTransmittance(RandSample, SigmaT, ColorChannelPdf);
+			}
 			float HitT = PackedPayload.IsHit() ? PackedPayload.HitT : PathState.Ray.TMax;
 			if (ScatterT < 0)
 			{
