Fix incorrect lighting on hair pixel with full coverage and having scatter scene lighting enabled.

diff --git a/Engine/Shaders/Private/HairStrands/HairStrandsVisibilityDepthPS.usf b/Engine/Shaders/Private/HairStrands/HairStrandsVisibilityDepthPS.usf
--- a/Engine/Shaders/Private/HairStrands/HairStrandsVisibilityDepthPS.usf
+++ b/Engine/Shaders/Private/HairStrands/HairStrandsVisibilityDepthPS.usf
@@ -26,7 +26,8 @@ struct FHairSampleData
 {
 	float DeviceZ;
 	uint LightChannelMask;
-	bool CastContactShadow;
+	bool bCastContactShadow;
+	bool bScatterSceneLight;
 };
 
 FHairSampleData ExtractSampleData(uint3 PixelCoord)
@@ -41,7 +42,8 @@ FHairSampleData ExtractSampleData(uint3 PixelCoord)
 		const FHairSample Sample = UnpackHairSample(HairSampleData[SampleIndex]);
 		Out.DeviceZ = max(Out.DeviceZ, Sample.Depth); // Inverse-Z
 		Out.LightChannelMask |= Sample.LightChannelMask;
-		Out.CastContactShadow = Out.CastContactShadow || HasHairFlags(Sample.Flags, HAIR_FLAGS_CAST_CONTACT_SHADOW);
+		Out.bCastContactShadow = Out.bCastContactShadow || HasHairFlags(Sample.Flags, HAIR_FLAGS_CAST_CONTACT_SHADOW);
+		Out.bScatterSceneLight = Out.bScatterSceneLight || HasHairFlags(Sample.Flags, HAIR_FLAGS_SCATTER_SCENE_LIGHT);
 	}
 	return Out;
 }
@@ -88,7 +90,7 @@ void MainPS(
 	// Set the material mode to hair for GBuffer/Material inspection
 	DummyHeader.SetMaterialMode(HEADER_MATERIALMODE_HAIR);
 	// Force PerObjectData contact shadow to respect hair settings, in order to have correct screen cast later
-	DummyHeader.SetCastContactShadow(SampleData.CastContactShadow ? 1 : 0);
+	DummyHeader.SetCastContactShadow(SampleData.bCastContactShadow ? 1 : 0);
 	Output0 = DummyHeader.State;
 
 #else
@@ -96,7 +98,7 @@ void MainPS(
 	Output1 = 0;
 
 	// Force PerObjectData contact shadow to respect hair settings, in order to have correct screen cast later
-	Output2 = float4(0,0,0,SampleData.CastContactShadow ? 1 : 0);
+	Output2 = float4(0,0,0,SampleData.bCastContactShadow ? 1 : 0);
 #endif
 	OutColor = 0;
 
@@ -107,6 +109,12 @@ void MainPS(
 		discard;
 	}
 	#endif
+
+	// Do not override GBuffer data for pixel containing hair sample which will gather light from scene color. Otherwise the pulled value will be black
+	if (SampleData.bScatterSceneLight)
+	{
+		discard;
+	}
 }
 #endif
 
