[Holdout] Fix Holdout does not render correctly in refractive objects. Previous implementation is using IOR to modify opacity and is not needed when IOR = 1.0. Make ray traced holdout with IOR render when opacity = 0.

diff --git a/Engine/Shaders/Private/BasePassPixelShader.usf b/Engine/Shaders/Private/BasePassPixelShader.usf
--- a/Engine/Shaders/Private/BasePassPixelShader.usf
+++ b/Engine/Shaders/Private/BasePassPixelShader.usf
@@ -2517,15 +2517,15 @@ void FPixelShaderInOut_MainPS(
 			float MaterialIOR = DielectricF0RGBToIor(ComputeF0(Specular, BaseColor, Metallic));
 		#endif
 		
-		// Match path tracing first intersection. Multiple intersections are biased.
+		// When material IOR is active and not 1.0f, it is fully opaque and match path tracing.
 		#if SUBSTRATE_TRANSLUCENT_MATERIAL
 			#if !MATERIAL_SHADINGMODEL_THIN_TRANSLUCENT && !SUBSTRATE_BLENDING_ALPHAHOLDOUT
-				float DualBlendSurfaceCoverageMix = 1 - DualBlendSurfaceCoverage * (1 - DualBlendSurfaceCoverage);
-				DualBlendSurfaceCoverage = select(MaterialIOR != 0.0, DualBlendSurfaceCoverageMix, DualBlendSurfaceCoverage);
+				float DualBlendSurfaceCoverageMix = 1.0f;
+				DualBlendSurfaceCoverage = select(MaterialIOR != 0.0 && MaterialIOR != 1.0f, DualBlendSurfaceCoverageMix, DualBlendSurfaceCoverage);
 			#endif
 		#elif MATERIALBLENDING_TRANSLUCENT
-			float OpacityMix =  1 - Opacity * (1 - Opacity);
-			Opacity = select(MaterialIOR != 0.0, OpacityMix, Opacity);
+			float OpacityMix = 1.0f;
+			Opacity = select(MaterialIOR != 0.0 && MaterialIOR != 1.0f, OpacityMix, Opacity);
 		#endif
 	#endif 
 	
