Fixed rect light data on mobile renderer supporting SM5 material nodes

#rnx

diff --git a/Engine/Shaders/Private/LightData.ush b/Engine/Shaders/Private/LightData.ush
--- a/Engine/Shaders/Private/LightData.ush
+++ b/Engine/Shaders/Private/LightData.ush
@@ -620,13 +620,17 @@ FDeferredLightData ConvertToDeferredLight(const FLocalLightData In)
 
 FDeferredLightData ConvertToDeferredLight_Mobile(const FLocalLightData In)
 {
-#if MOBILE_SHADING_PATH_SUPPORT_CAPSULE_LIGHT
-	return ConvertToDeferredLight(In);
-#else
 	half4 OutPreviewShadowMapChannelMask = 0;
 	uint OutLightingChannelMask = LIGHTING_CHANNEL_MASK;
 	FDeferredLightData Out = ConvertToDeferredLight(In, 1.0f, OutPreviewShadowMapChannelMask, OutLightingChannelMask);	
 
+#if MOBILE_SHADING_PATH_SUPPORT_CAPSULE_LIGHT
+	if (!Out.bRectLight)
+	{
+		return ConvertToDeferredLight(In);
+	}
+#endif
+
 	// Override value for performance on mobile
 	Out.SoftSourceRadius= 0;
 	Out.SpecularScale = 1.0f; // If enabled -> UnpackLightSpecularScale(In);
@@ -647,7 +651,6 @@ FDeferredLightData ConvertToDeferredLight_Mobile(const FLocalLightData In)
 	}
 
 	return Out;
-#endif // MOBILE_SHADING_PATH_SUPPORT_CAPSULE_LIGHT
 }
 
 FLightShaderParameters ConvertToLightShaderParameters(FDeferredLightData In)
diff --git a/Engine/Shaders/Private/MobileDeferredShading.usf b/Engine/Shaders/Private/MobileDeferredShading.usf
--- a/Engine/Shaders/Private/MobileDeferredShading.usf
+++ b/Engine/Shaders/Private/MobileDeferredShading.usf
@@ -327,6 +327,14 @@ void MobileRadialLightPS(
 	FLightAccumulator DirectLighting = (FLightAccumulator)0;
 	FDeferredLightData LightData = InitDeferredLightFromLightParameters(GetRootLightShaderParameters(), RADIAL_LIGHT_TYPE);
 		
+#if MOBILE_SUPPORTS_SM5_MATERIAL_NODES
+	if(LightData.bRectLight && GBuffer.ShadingModelID == SHADINGMODELID_SUBSURFACE_PROFILE)
+	{
+		// SSS has issues with rect light. Will be fixed in upcoming version
+		 GBuffer.ShadingModelID = SHADINGMODELID_DEFAULT_LIT;
+	}
+#endif
+
 #if !SIMPLE_LIGHT
 	// Affect the light color by any other shadow/transmittance before the lighting is evaluated
 	half Attenuation = ComputeLightProfileMultiplier(TranslatedWorldPosition, LightData.TranslatedWorldPosition, -LightData.Direction, LightData.Tangent, LightData.IESAtlasIndex);
