Mobile: Fixed typo in a substrate UB name

diff --git a/Engine/Shaders/Private/MobileBasePassPixelShader.usf b/Engine/Shaders/Private/MobileBasePassPixelShader.usf
--- a/Engine/Shaders/Private/MobileBasePassPixelShader.usf
+++ b/Engine/Shaders/Private/MobileBasePassPixelShader.usf
@@ -906,7 +906,7 @@ void Main(
 		#else
 		const bool bForceFullyRough = View.RenderingReflectionCaptureMask > 0; // Can mobile capture reflection probe? Is this needed for mobile?
 		#endif
-		FSubstrateIntegrationSettings Settings = InitSubstrateIntegrationSettings(bForceFullyRough, SubstrateStruct.bRoughDiffuse, SubstrateStruct.PeelLayersAboveDepth, Substrate.bRoughnessTracking);
+		FSubstrateIntegrationSettings Settings = InitSubstrateIntegrationSettings(bForceFullyRough, SubstrateStruct.bRoughDiffuse, SubstrateStruct.PeelLayersAboveDepth, SubstrateStruct.bRoughnessTracking);
 
 
 		// TODO: needed? Remove?
