[Lumen translucency GI volume] Smooth out the transition from the furthest GI volume froxel to unshadowed skylight

#tests tested in editor

diff --git a/Engine/Shaders/Private/BasePassPixelShader.usf b/Engine/Shaders/Private/BasePassPixelShader.usf
--- a/Engine/Shaders/Private/BasePassPixelShader.usf
+++ b/Engine/Shaders/Private/BasePassPixelShader.usf
@@ -409,7 +409,14 @@ void GetSkyLighting(FMaterialPixelParameters MaterialParameters, VTPageTableResu
 	if (IsLumenTranslucencyGIEnabled())
 	{
 		const float3 VolumeUV = ComputeTranslucencyGIVolumeUV(MaterialParameters.AbsoluteWorldPosition, ResolvedView.WorldToClip);
-		if(VolumeUV.z < 1.0)
+
+		float3 DiffuseLighting = 0;
+		float3 SubsurfaceLighting = 0;
+
+		float InvFadeLength = LumenGIVolumeStruct.TranslucencyGIGridSize.z / 0.5;
+		float FadeAmount = saturate((1.0 - VolumeUV.z) * InvFadeLength);
+
+		if(FadeAmount > 0.0)
 		{
 			// Lumen Dynamic GI + shadowed Skylight
 			FTwoBandSHVectorRGB TranslucencyGISH = GetTranslucencyGIVolumeLighting(MaterialParameters.AbsoluteWorldPosition, ResolvedView.WorldToClip, true);
@@ -422,32 +429,36 @@ void GetSkyLighting(FMaterialPixelParameters MaterialParameters, VTPageTableResu
 			TranslucencyGISH1.B.V = TranslucencyGISH.B.V.x;
 
 			FOneBandSHVector DiffuseTransferSH = CalcDiffuseTransferSH1(1);
-			OutDiffuseLighting += max(float3(0,0,0), DotSH1(TranslucencyGISH1, DiffuseTransferSH)) / PI;
+			DiffuseLighting = max(float3(0,0,0), DotSH1(TranslucencyGISH1, DiffuseTransferSH)) / PI;
 
 		#else
 			// Diffuse convolution
 			FTwoBandSHVector DiffuseTransferSH = CalcDiffuseTransferSH(WorldNormal, 1);
-			OutDiffuseLighting += max(half3(0,0,0), DotSH(TranslucencyGISH, DiffuseTransferSH)) / PI;
+			DiffuseLighting = max(half3(0,0,0), DotSH(TranslucencyGISH, DiffuseTransferSH)) / PI;
 
 			#if SHADINGMODEL_REQUIRES_BACKFACE_LIGHTING
 			if (bEvaluateBackface)
 			{
 				FTwoBandSHVector SubsurfaceTransferSH = CalcDiffuseTransferSH(-WorldNormal, 1);
-				OutSubsurfaceLighting += max(half3(0,0,0), DotSH(TranslucencyGISH, SubsurfaceTransferSH)) / PI;
+				SubsurfaceLighting = max(half3(0,0,0), DotSH(TranslucencyGISH, SubsurfaceTransferSH)) / PI;
 			}
 			#endif
 		#endif
 		}
-		else
+		
+		if (FadeAmount < 1.0)
 		{
-			OutDiffuseLighting += GetEffectiveSkySHDiffuse(WorldNormal) * ResolvedView.SkyLightColor.rgb;
+			DiffuseLighting = lerp(GetEffectiveSkySHDiffuse(WorldNormal) * ResolvedView.SkyLightColor.rgb, DiffuseLighting, FadeAmount);
 			#if SHADINGMODEL_REQUIRES_BACKFACE_LIGHTING
 			if (bEvaluateBackface)
 			{
-				OutDiffuseLighting += GetEffectiveSkySHDiffuse(-WorldNormal) * ResolvedView.SkyLightColor.rgb;
+				SubsurfaceLighting = lerp(GetEffectiveSkySHDiffuse(-WorldNormal) * ResolvedView.SkyLightColor.rgb, SubsurfaceLighting, FadeAmount);
 			}
 			#endif
 		}
+
+		OutDiffuseLighting += DiffuseLighting;
+		OutSubsurfaceLighting += SubsurfaceLighting;
 	}
 	else
 #endif
