[Lit Wireframe] Fix issue where MeshEdges doesn't composite correctly when MSAA is disabled in the editor.

#tests Tested in Vulkan on Windows

diff --git a/Engine/Shaders/Private/MeshEdges.usf b/Engine/Shaders/Private/MeshEdges.usf
--- a/Engine/Shaders/Private/MeshEdges.usf
+++ b/Engine/Shaders/Private/MeshEdges.usf
@@ -28,6 +28,7 @@ float4 SampleOffsetArray[MSAA_SAMPLE_COUNT];
 SCREEN_PASS_TEXTURE_VIEWPORT(Output)
 
 float Opacity;
+float Gamma;
 
 // Return true if A is closer than B
 bool IsCloser(float A, float B)
@@ -55,7 +56,7 @@ void ComposeMeshEdgesPS(
 	float ShadedDeviceZ;
 	#if MSAA_SAMPLE_COUNT > 1
 	{
-		WireframeColor = float4(WireframeColorTexture.Load(WireframeSamplePosition, TargetSampleIndex).rgb, 1.0);
+		WireframeColor = WireframeColorTexture.Load(WireframeSamplePosition, TargetSampleIndex);
 		WireframeDeviceZ = WireframeDepthTexture.Load(WireframeSamplePosition, TargetSampleIndex).r;
 
 		const float2 SampleOffset = SampleOffsetArray[TargetSampleIndex].xy;
@@ -96,16 +97,9 @@ void ComposeMeshEdgesPS(
 		
 	OutDepth = WireframeDeviceZ;
 	
-	// Fixes the gamma of editor primitives
-	{
-		// Bring out of premultiplied.
-		WireframeColor.rgb /= max(WireframeColor.a, 0.0001f);
-		// Fix gamma.
-		WireframeColor.rgb = pow(WireframeColor.rgb, 2.2f);
-		// Bring back to premultiplied
-		WireframeColor.rgb *= WireframeColor.a;
-	}
+	// WireframeColor already has gamma applied, and will get gamma again, so undo the first one here
+	WireframeColor.rgb = pow(WireframeColor.rgb, rcp(Gamma));
 
-	OutColor = WireframeColor;
+	OutColor = float4(WireframeColor.rgb, 1);
 	OutColor.a *= Opacity;
 }
diff --git a/Engine/Source/Runtime/Renderer/Private/MeshEdgesRendering.cpp b/Engine/Source/Runtime/Renderer/Private/MeshEdgesRendering.cpp
--- a/Engine/Source/Runtime/Renderer/Private/MeshEdgesRendering.cpp
+++ b/Engine/Source/Runtime/Renderer/Private/MeshEdgesRendering.cpp
@@ -71,6 +71,7 @@ class FComposeMeshEdgesPS : public FGlobalShader
 
 		SHADER_PARAMETER_STRUCT(FScreenPassTextureViewportParameters, Output)
 		
+		SHADER_PARAMETER(float, Gamma)
 		SHADER_PARAMETER(float, Opacity)
 		
 		RENDER_TARGET_BINDING_SLOTS()
@@ -457,6 +458,7 @@ void ComposeMeshEdges(FRDGBuilder& GraphBuilder,
 	PassParameters->DepthTexture = SceneDepth.Texture;
 	PassParameters->DepthSampler = PointClampSampler;
 	PassParameters->DepthTextureJitter = FVector2f(View.TemporalJitterPixels);
+	PassParameters->Gamma = GEngine->DisplayGamma;
 	PassParameters->Opacity = ViewSettings.Opacity;
 
 	for (int32 i = 0; i < int32(NumMSAASamples); i++)
