Fix incorrect culling of hidden line in ShaderPrint when r.ShaderPrint.DrawOccludedLines is enabled.

diff --git a/Engine/Shaders/Private/ShaderPrintDrawPrimitive.usf b/Engine/Shaders/Private/ShaderPrintDrawPrimitive.usf
--- a/Engine/Shaders/Private/ShaderPrintDrawPrimitive.usf
+++ b/Engine/Shaders/Private/ShaderPrintDrawPrimitive.usf
@@ -101,8 +101,10 @@ void ShaderDrawDebugPS(
 	const float Depth = DepthTexture.SampleLevel(DepthSampler, DepthUV, 0);
 	const bool bIsHidden = Position.z < Depth; // Reverse-Z
 
-	if (bIsHidden)
+	if (bIsHidden && !bDrawOccludedLines)
+	{
 		discard;
+	}
 
 	const float ColorScale = bIsHidden ? 0.4f : 1;
 	const float CheckerboardScale = ((PixelCoord.x%2) ^ (PixelCoord.y%2)) > 0 ? 1 : 0.0;
