Path Tracer: Improve volumetric lighting model to better match the deferred renderer

The previous implementation did not properly account for directionality greater than 1. It also did not match the light clipping around 45 degrees past the terminator.

diff --git a/Engine/Shaders/Private/PathTracing/Material/PathTracingMedium.ush b/Engine/Shaders/Private/PathTracing/Material/PathTracingMedium.ush
--- a/Engine/Shaders/Private/PathTracing/Material/PathTracingMedium.ush
+++ b/Engine/Shaders/Private/PathTracing/Material/PathTracingMedium.ush
@@ -3,7 +3,7 @@
 /*=============================================================================================
 PathTracingMedium.usf: Volumetric phase function
  * To support atmosphere, we represent a blend of two components: rayleigh scattering and an HG lobe.
- * To support volumetric clouds we also include a dual HG lobe that we track seperately to simplify handling cases where multiple volumes overlap.
+ * To support volumetric clouds we also include a dual HG lobe that we track separately to simplify handling cases where multiple volumes overlap.
  * See FVolumeShadedResult for the logic around how this struct is tracked during volume tracing.
    
 ===============================================================================================*/
@@ -112,26 +112,30 @@ FMaterialEval Medium_EvalMaterial(
 }
 
 
+
 FMaterialSample Volumetric_SampleMaterial(
 	float3 V_World,
 	FPathTracingPayload Payload,
 	float3 RandSample)
 {
 
-	const float Directionality = Payload.Anisotropy;
-	// phase function = (1 + Directionality * CosTheta) / 4Pi
-	// The expression below gives exact importance sampling
-	const float CosTheta = Directionality > 0 ? (-1.0 + sqrt(Pow2(1 - Directionality) + 4.0 * Directionality * RandSample.x)) * rcp(Directionality) : -1.0 + 2.0 * RandSample.x;
+	const float Directionality = Payload.Anisotropy + 1;
+	// Exact importance sampling of the Pdf below
+	const float CosTheta = 
+		Directionality > 1 ? (-1.0 + Directionality * 0.5 + (2 - 0.5 * Directionality) * sqrt(RandSample.x)) : (
+		Directionality > 0 ? (-1.0 + sqrt(Pow4(1 - Directionality) + (4.0 - Directionality) * Directionality * (2 + Directionality * Directionality) * RandSample.x)) * rcp(2 * Directionality) :
+			-1.0 + 2.0 * RandSample.x
+	);
 	const float SinTheta = sqrt(saturate(1.0 - CosTheta * CosTheta));
 	const float Phi = RandSample.y * (2.0 * PI);
 
 	const float3 L_World = TangentToWorld(float3(cos(Phi) * SinTheta, sin(Phi) * SinTheta, CosTheta), Payload.WorldNormal);
 
 	const float3 Color = Payload.GetBaseColor() * Payload.BSDFOpacity;
-	const float Inv4PI = 0.25 / PI;
-	const float Pdf = Inv4PI + Inv4PI * Directionality * CosTheta;
-
-	return CreateMaterialSample(L_World, Color, Pdf, sign(CosTheta), 1.0, PATHTRACER_SCATTER_DIFFUSE);
+	const float Pdf = Directionality > 1
+		? ((4 - 2 * Directionality) + 4 * CosTheta) / (PI * Pow2(Directionality - 4))
+		: (1 + (2 * Directionality) * CosTheta) / (((2 * PI) - (0.5 * PI) * Directionality) * (2 + Directionality * Directionality));
+	return CreateMaterialSample(L_World, Color, max(Pdf, 0), sign(CosTheta), 1.0, PATHTRACER_SCATTER_DIFFUSE);
 }
 
 FMaterialEval Volumetric_EvalMaterial(
@@ -141,10 +145,16 @@ FMaterialEval Volumetric_EvalMaterial(
 	float2 DiffuseSpecularScale
 )
 {
-	const float Directionality = Payload.Anisotropy;
+	const float Directionality = Payload.Anisotropy + 1;
 	const float3 Color = Payload.GetBaseColor() * Payload.BSDFOpacity;
 	const float CosTheta = dot(L_World, Payload.WorldNormal);
-	const float Inv4PI = 0.25 / PI;
-	const float Pdf = Inv4PI + Inv4PI * Directionality * CosTheta;
-	return CreateMaterialEval(Color * DiffuseSpecularScale.x, Pdf);
+	// The PDF below has the following properties:
+	// When Directionality == 0: uniform sampling of the sphere (all cosines from -1 to 1 are equaly likely)
+	// When Directionality == 1: sample cosine linearly -1/2 to +1 
+	// When Directionality == 2: sample cosine linearly from 0 to +1 (regular cosine lobe)
+	// In between these values, the PDF changes smoothly and is always a linear function, so the CDF is quadratic and can be directly inverted (see above)
+	const float Pdf = Directionality > 1
+		? ((4 - 2 * Directionality) + 4 * CosTheta) / (PI * Pow2(Directionality - 4))
+		: (1 + (2 * Directionality) * CosTheta) / (((2 * PI) - (0.5 * PI) * Directionality) * (2 + Directionality * Directionality));
+	return CreateMaterialEval(Color * DiffuseSpecularScale.x, max(Pdf, 0));
 }
diff --git a/Engine/Shaders/Private/PathTracing/PathTracingMaterialHitShader.usf b/Engine/Shaders/Private/PathTracing/PathTracingMaterialHitShader.usf
--- a/Engine/Shaders/Private/PathTracing/PathTracingMaterialHitShader.usf
+++ b/Engine/Shaders/Private/PathTracing/PathTracingMaterialHitShader.usf
@@ -75,7 +75,11 @@ float AdjustMaterialRoughness(float Roughness, float PathRoughness)
 	return max(Roughness, PathRoughness);
 }
 
-#define PATH_TRACER_USE_SUBSTRATE_GLASS_LOBE    (MATERIALBLENDING_TRANSLUCENT && (SUBSTRATE_BLENDING_TRANSLUCENT_GREYTRANSMITTANCE || SUBSTRATE_BLENDING_TRANSLUCENT_COLOREDTRANSMITTANCE))
+#define PATH_TRACER_USE_SUBSTRATE_GLASS_LOBE    	(MATERIALBLENDING_TRANSLUCENT && (SUBSTRATE_BLENDING_TRANSLUCENT_GREYTRANSMITTANCE || SUBSTRATE_BLENDING_TRANSLUCENT_COLOREDTRANSMITTANCE))
+
+// Translucent shading cases that use a different BSDF response than normal (implemented with SH lighting in the raster pipeline)
+#define PATH_TRACER_USE_VOLUMETRIC_SHADING_MODEL 				(TRANSLUCENCY_LIGHTING_SURFACE_LIGHTINGVOLUME || TRANSLUCENCY_LIGHTING_VOLUMETRIC_PERVERTEX_NONDIRECTIONAL || TRANSLUCENCY_LIGHTING_VOLUMETRIC_PERVERTEX_DIRECTIONAL || TRANSLUCENCY_LIGHTING_VOLUMETRIC_NONDIRECTIONAL || TRANSLUCENCY_LIGHTING_VOLUMETRIC_DIRECTIONAL)
+#define PATH_TRACER_USE_VOLUMETRIC_SHADING_MODEL_DIRECTIONAL 	(TRANSLUCENCY_LIGHTING_SURFACE_LIGHTINGVOLUME || TRANSLUCENCY_LIGHTING_VOLUMETRIC_PERVERTEX_DIRECTIONAL || TRANSLUCENCY_LIGHTING_VOLUMETRIC_DIRECTIONAL)
 
 #if REFRACTION_USE_INDEX_OF_REFRACTION || REFRACTION_USE_INDEX_OF_REFRACTION_FROM_F0
 
@@ -271,7 +275,7 @@ RAY_TRACING_ENTRY_CLOSEST_HIT(PathTracingMaterialCHS,
 	 * Set common material attributes for both full and simplified materials
 	 **/
 
-#if (TRANSLUCENCY_LIGHTING_VOLUMETRIC_PERVERTEX_NONDIRECTIONAL || TRANSLUCENCY_LIGHTING_VOLUMETRIC_PERVERTEX_DIRECTIONAL || TRANSLUCENCY_LIGHTING_VOLUMETRIC_NONDIRECTIONAL || TRANSLUCENCY_LIGHTING_VOLUMETRIC_DIRECTIONAL)
+#if PATH_TRACER_USE_VOLUMETRIC_SHADING_MODEL
 	Payload.ShadingModelID = SHADINGMODELID_UNLIT;
 #else
 	Payload.ShadingModelID = GetMaterialShadingModel(PixelMaterialInputs);
@@ -348,7 +352,7 @@ RAY_TRACING_ENTRY_CLOSEST_HIT(PathTracingMaterialCHS,
 #endif
 	Payload.WorldGeoNormal = GeoNormalSign * GeoNormal;
 	Payload.WorldSmoothNormal = MaterialParameters.TwoSidedSign * TransformTangentNormalToWorld(MaterialParameters.TangentToWorld, float3(0, 0, 1));
-#if (TRANSLUCENCY_LIGHTING_VOLUMETRIC_PERVERTEX_NONDIRECTIONAL || TRANSLUCENCY_LIGHTING_VOLUMETRIC_PERVERTEX_DIRECTIONAL || TRANSLUCENCY_LIGHTING_VOLUMETRIC_NONDIRECTIONAL || TRANSLUCENCY_LIGHTING_VOLUMETRIC_DIRECTIONAL)
+#if PATH_TRACER_USE_VOLUMETRIC_SHADING_MODEL
 	// no need to adjust normal on volumetric shading models
 #else
 	Payload.WorldSmoothNormal = AdjustShadingNormal(Payload.WorldSmoothNormal, Payload.WorldGeoNormal, WorldRayDirection());
@@ -795,7 +799,7 @@ RAY_TRACING_ENTRY_CLOSEST_HIT(PathTracingMaterialCHS,
 	Payload.WorldNormal *= MaterialParameters.TwoSidedSign;
 #endif
 
-#if (TRANSLUCENCY_LIGHTING_VOLUMETRIC_PERVERTEX_NONDIRECTIONAL || TRANSLUCENCY_LIGHTING_VOLUMETRIC_PERVERTEX_DIRECTIONAL || TRANSLUCENCY_LIGHTING_VOLUMETRIC_NONDIRECTIONAL || TRANSLUCENCY_LIGHTING_VOLUMETRIC_DIRECTIONAL)
+#if PATH_TRACER_USE_VOLUMETRIC_SHADING_MODEL
 	// volumetric shading doesn't need to adjust normal
 #else
 	Payload.WorldNormal = AdjustShadingNormal(Payload.WorldNormal, Payload.WorldGeoNormal, WorldRayDirection());
@@ -844,14 +848,14 @@ RAY_TRACING_ENTRY_CLOSEST_HIT(PathTracingMaterialCHS,
 	Payload.WorldTangent = normalize(MaterialParameters.WorldNormal);
 #endif
 
-#if (TRANSLUCENCY_LIGHTING_VOLUMETRIC_PERVERTEX_NONDIRECTIONAL || TRANSLUCENCY_LIGHTING_VOLUMETRIC_PERVERTEX_DIRECTIONAL || TRANSLUCENCY_LIGHTING_VOLUMETRIC_NONDIRECTIONAL || TRANSLUCENCY_LIGHTING_VOLUMETRIC_DIRECTIONAL)
+#if PATH_TRACER_USE_VOLUMETRIC_SHADING_MODEL
 	if (Payload.ShadingModelID == SHADINGMODELID_UNLIT)
 	{
-		// encode the directionality using the specular field
-	#if TRANSLUCENCY_LIGHTING_VOLUMETRIC_PERVERTEX_DIRECTIONAL || TRANSLUCENCY_LIGHTING_VOLUMETRIC_DIRECTIONAL
-		Payload.Anisotropy = GetMaterialTranslucencyDirectionalLightingIntensity();
+		// encode the directionality using Anisotropy: remap [0,2] range to [-1,1]
+	#if PATH_TRACER_USE_VOLUMETRIC_SHADING_MODEL_DIRECTIONAL
+		Payload.Anisotropy = clamp(GetMaterialTranslucencyDirectionalLightingIntensity() - 1, -1, 1);
 	#else
-		Payload.Anisotropy = 0;
+		Payload.Anisotropy = -1;
 	#endif
 
 	#if USE_LEGACY_CODEPATH_FOR_SUBSTRATE
