[DOF] Added new T-interval sample clamping algorithm to the TAA prepass, to help reduce noise in areas with a lot of sample, like specular highlights. This method is similar to the variance clamping method, where a group of neighbouring samples is used to build a probability distribution, which then determines the likelyhood of the current sample being part of that distribution. The T-interval method goes one step further and models the uncertainty of that distribution, due to it being reconstructed from a low samplecount.

diff --git a/Engine/Shaders/Private/TemporalAA.usf b/Engine/Shaders/Private/TemporalAA.usf
--- a/Engine/Shaders/Private/TemporalAA.usf
+++ b/Engine/Shaders/Private/TemporalAA.usf
@@ -55,6 +55,9 @@
 	// Min max samples that are within distance from output pixel. 
 	#define HISTORY_CLAMPING_BOX_SAMPLE_DISTANCE 2
 
+	// Like the variance method, but uses Student-T distribution to model uncertainty in variance and mean.
+	#define HISTORY_CLAMPING_BOX_T_INTERVAL 3
+
 /** Qualities */
 	#define TAA_QUALITY_LOW 0
 	#define TAA_QUALITY_MEDIUM 1
@@ -230,6 +233,7 @@
 		#define AA_BICUBIC 1
 		#define AA_FILTERED 1
 		#define AA_YCOCG 1
+		#define AA_HISTORY_CLAMPING_BOX (HISTORY_CLAMPING_BOX_T_INTERVAL)
 
 	#elif TAA_QUALITY == TAA_QUALITY_MEDIUM_HIGH
 		#if AA_UPSAMPLE
@@ -751,6 +755,13 @@ FTAAHistoryPayload MaxPayload(in FTAAHistoryPayload Payload0, in FTAAHistoryPayl
 	return Payload0;
 }
 
+FTAAHistoryPayload Pow2Payload(in FTAAHistoryPayload Payload)
+{
+	Payload.Color = Pow2(Payload.Color);
+	Payload.CocRadius = Pow2(Payload.CocRadius);
+	return Payload;
+}
+
 FTAAHistoryPayload MinPayload3(in FTAAHistoryPayload Payload0, in FTAAHistoryPayload Payload1, in FTAAHistoryPayload Payload2)
 {
 	Payload0.Color = min3(Payload0.Color, Payload1.Color, Payload2.Color);
@@ -1701,7 +1712,7 @@ void ComputeNeighborhoodBoundingbox(
 	FTAAHistoryPayload NeighborMin;
 	FTAAHistoryPayload NeighborMax;
 
-	#if AA_HISTORY_CLAMPING_BOX == HISTORY_CLAMPING_BOX_VARIANCE
+	#if AA_HISTORY_CLAMPING_BOX == HISTORY_CLAMPING_BOX_VARIANCE || AA_HISTORY_CLAMPING_BOX == HISTORY_CLAMPING_BOX_T_INTERVAL
 	{
 		#if AA_SAMPLES == 9
 			const uint SampleIndexes[9] = kSquareIndexes3x3;
@@ -1711,25 +1722,52 @@ void ComputeNeighborhoodBoundingbox(
 			#error Unknown number of samples.
 		#endif
 
-		float4 m1 = 0;
-		float4 m2 = 0;
+		FTAAHistoryPayload m1 = (FTAAHistoryPayload)0;
+		FTAAHistoryPayload m2 = (FTAAHistoryPayload)0;
 		for( uint i = 0; i < AA_SAMPLES; i++ )
 		{
-			float4 SampleColor = Neighbors[ SampleIndexes[i] ];
+			FTAAHistoryPayload Sample = Neighbors[ SampleIndexes[i] ];
 
-			m1 += SampleColor;
-			m2 += Pow2( SampleColor );
+			m1 = AddPayload(m1, Sample);
+			m2 = AddPayload(m2, Pow2Payload(Sample));
 		}
 
-		m1 *= (1.0 / AA_SAMPLES);
-		m2 *= (1.0 / AA_SAMPLES);
+		m1 = MulPayload(m1, (1.0 / AA_SAMPLES));
+		m2 = MulPayload(m2, (1.0 / AA_SAMPLES));
+
+		float4 StdDevColor = sqrt( abs(m2.Color - m1.Color * m1.Color) );
+		float StdDevCocRadius = sqrt( abs(m2.CocRadius - m1.CocRadius * m1.CocRadius) );
+		#if AA_HISTORY_CLAMPING_BOX == HISTORY_CLAMPING_BOX_VARIANCE
+			NeighborMin.Color = m1.Color - 1.25 * StdDevColor;
+			NeighborMin.CocRadius = m1.CocRadius - 1.25 * StdDevCocRadius;
+			NeighborMax.Color = m1.Color + 1.25 * StdDevColor;
+			NeighborMax.CocRadius = m1.CocRadius + 1.25 * StdDevCocRadius;
+		#elif AA_HISTORY_CLAMPING_BOX == HISTORY_CLAMPING_BOX_T_INTERVAL
+			// We estimate the mean and stddev based on a low sample count
+			// so there is some uncertainty in these estimations.
+			// We model this uncertainty with a T-distribution,
+			// and use this to get a confidence interval.
 
-		float4 StdDev = sqrt( abs(m2 - m1 * m1) );
-		NeighborMin = m1 - 1.25 * StdDev;
-		NeighborMax = m1 + 1.25 * StdDev;
+		#if AA_SAMPLES == 9
+			const float StudentTValue = 2.306f; // 95% confidence on 9 samples (8 degrees of freedom)
+		#elif AA_SAMPLES == 5
+			const float StudentTValue = 2.776f; // 95% confidence on 5 samples (4 degrees of freedom)
+		#else
+			#error Unknown number of samples.
+		#endif
+		
+			// Prediction Interval Formula
+			float4 ErrorMarginColor = StdDevColor * StudentTValue * sqrt(1.0f + (1.0f/AA_SAMPLES));
+			float ErrorMarginCocRadius = StdDevCocRadius * StudentTValue * sqrt(1.0f + (1.0f/AA_SAMPLES));
+
+			NeighborMin.Color = m1.Color - ErrorMarginColor;
+			NeighborMin.CocRadius = m1.CocRadius - ErrorMarginCocRadius;
+			NeighborMax.Color = m1.Color + ErrorMarginColor;
+			NeighborMax.CocRadius = m1.CocRadius + ErrorMarginCocRadius;
+		#endif
 
-		NeighborMin = min( NeighborMin, IntermediaryResult.Filtered );
-		NeighborMax = max( NeighborMax, IntermediaryResult.Filtered );
+		NeighborMin = MinPayload(NeighborMin, IntermediaryResult.Filtered);
+		NeighborMax = MaxPayload(NeighborMax, IntermediaryResult.Filtered);
 	}
 	#elif AA_HISTORY_CLAMPING_BOX == HISTORY_CLAMPING_BOX_SAMPLE_DISTANCE
 	// Do color clamping only within a radius.
