[DOF] Disable 47806104 for now, falling back to HISTORY_CLAMPING_BOX_MIN_MAX. Moving high contrast edges in animated content have increased ghosting that may be solvable with better depth-based rejection.

diff --git a/Engine/Shaders/Private/TemporalAA.usf b/Engine/Shaders/Private/TemporalAA.usf
--- a/Engine/Shaders/Private/TemporalAA.usf
+++ b/Engine/Shaders/Private/TemporalAA.usf
@@ -233,7 +233,7 @@
 		#define AA_BICUBIC 1
 		#define AA_FILTERED 1
 		#define AA_YCOCG 1
-		#define AA_HISTORY_CLAMPING_BOX (HISTORY_CLAMPING_BOX_T_INTERVAL)
+		#define AA_HISTORY_CLAMPING_BOX (HISTORY_CLAMPING_BOX_MIN_MAX)
 
 	#elif TAA_QUALITY == TAA_QUALITY_MEDIUM_HIGH
 		#if AA_UPSAMPLE
