Path Tracer: Replace the diffuse lobe in thin glass with a scattering lobe that accounts for the phase function

The scattering lobe is modelled with the Hanrahan-Krueger single scatter BSDF for now, which should be appropriate as long as the optical depth of the slab is low. For now there doesn't seem to be a simple method to make the lobe energy preserving as taking into account the optical depth would require three lookups at a minimum. After investigating the energy loss in a furnace setup, it seems like the single-scatter assumptions is reasonable and does not lead to too much energy loss in practice.

#rnx

diff --git a/Engine/Shaders/Private/BRDF.ush b/Engine/Shaders/Private/BRDF.ush
--- a/Engine/Shaders/Private/BRDF.ush
+++ b/Engine/Shaders/Private/BRDF.ush
@@ -818,4 +818,37 @@ FBeckmannDesc GGXToBeckmann(float2 GGXSafeRoughness)
 	Out.Sigma		= Out.Alpha * InvSqrt2;
 	Out.Rho			= 0.0f;
 	return Out;
+}
+
+// Hanrahan-Krueger BRDF (not including the phase function and albedo):
+//   [Siggraph 1993, "Reflection from Layered Surfaces due to Subsurface Scattering", Pat Hanrahan, Wolfgang Krueger ]
+//   http://www.graphics.stanford.edu/papers/subsurface/
+// The same BRDF is used in the following paper which we reference for the more modern notation and for having numbered equations:
+//   [EGSR 2014, "A Physically-Based BSDF for Modeling the Appearance of Paper", Marios Papas, Krystle de Mesa, and Henrik Wann Jensen ]
+//   http://graphics.ucsd.edu/~henrik/papers/paper_appearance.pdf
+float3 HanrahanKrueger_Weight(float NoV, float NoL, float3 Tr)
+{
+	float3 SigmaT = -log(max(Tr, 1e-5));
+	if (NoV > 0)
+	{
+		if (NoL > 0)
+		{
+			// Section 4.4, Eq 8
+			return (1 - exp(-SigmaT * (rcp(NoL) + rcp(NoV)))) * NoL / (NoV + NoL);
+		}
+		else if (NoL < 0)
+		{
+			if (abs(NoV + NoL) > 1e-5)
+			{
+				// Section 4.5, Eq 11
+				return (exp(-SigmaT / NoV) - exp(SigmaT / NoL)) * -NoL / (NoV + NoL);
+			}
+			else
+			{
+				// Section 4.5, Eq 12
+				return SigmaT * exp(-SigmaT / NoV) / NoV;
+			}
+		}
+	}
+	return 0;
 }
\ No newline at end of file
diff --git a/Engine/Shaders/Private/PathTracing/Material/PathTracingSubstrateThinGlass.ush b/Engine/Shaders/Private/PathTracing/Material/PathTracingSubstrateThinGlass.ush
--- a/Engine/Shaders/Private/PathTracing/Material/PathTracingSubstrateThinGlass.ush
+++ b/Engine/Shaders/Private/PathTracing/Material/PathTracingSubstrateThinGlass.ush
@@ -89,7 +89,6 @@ struct FSubstrateThinGlassSlabData
 	}
 };
 
-
 FMaterialSample SubstrateThinGlass_SampleMaterial(
 	float3 V_World,
 	FPathTracingPayload Payload,
@@ -99,17 +98,23 @@ FMaterialSample SubstrateThinGlass_SampleMaterial(
 	FSubstrateThinGlassSlabData Data = (FSubstrateThinGlassSlabData)0;
 	Data.PrepareBSDF(V_World, Payload);
 
-	float3 L = 0, H = 0, V = Data.V;
+	float3 L = 0, H = 0, V = Data.V, ScatterWeight = 0;
 	float OutRoughness = 1;
 
-	bool bSampleDiffuse = false, bSampleCloth = false, bSampleGlass = false, bIsReflection = RandSample.x < Data.LobeCdf.z;
+	bool bSampleScatter = false, bSampleCloth = false, bSampleGlass = false, bIsReflection = RandSample.x < Data.LobeCdf.z;
 	
 	if (RandSample.x < Data.LobeCdf.x)
 	{
-		bSampleDiffuse = true;
+		bSampleScatter = true;
 		RandSample.x = RescaleRandomNumber(RandSample.x, 0.0, Data.LobeCdf.x);
-		// diffuse lobe
-		L = CosineSampleHemisphere(RandSample.xy).xyz;
+		// scattering lobe
+		L = TangentToWorld(ImportanceSampleHenyeyGreensteinPhase(RandSample.xy, Payload.PhaseG).xyz, -V);
+
+		ScatterWeight = HanrahanKrueger_Weight(V.z, L.z, Data.TransmissionColor);
+		bIsReflection = L.z > 0;
+
+		// keep L to positive side for the remainder of the calculations
+		L.z = abs(L.z);
 		H = normalize(L + V);
 	}
 	else if (RandSample.x < Data.LobeCdf.y)
@@ -160,14 +165,14 @@ FMaterialSample SubstrateThinGlass_SampleMaterial(
 
 	const float PositionBiasSign = bIsReflection ? 1.0 : -1.0;
 	const float3 L_World = normalize(mul(float3(L.xy, PositionBiasSign * L.z), Data.Basis)); // flip reflection to other side if we are refracting
-	FMaterialSample Result = CreateMaterialSample(L_World, 0.0, 0.0, PositionBiasSign, OutRoughness, bSampleDiffuse ? PATHTRACER_SCATTER_DIFFUSE : (bIsReflection ? PATHTRACER_SCATTER_SPECULAR : PATHTRACER_SCATTER_REFRACT));
+	FMaterialSample Result = CreateMaterialSample(L_World, 0.0, 0.0, PositionBiasSign, OutRoughness, bSampleScatter ? PATHTRACER_SCATTER_DIFFUSE : (bIsReflection ? PATHTRACER_SCATTER_SPECULAR : PATHTRACER_SCATTER_REFRACT));
 
 	// Specular Lobes (either reflected or transmitted)
 	const float2 GGXResult0 = GGXEvalReflection(L, V, H, bIsReflection ? Data.AlphaR0 : Data.AlphaT0);
 	const float2 GGXResult1 = GGXEvalReflection(L, V, H, bIsReflection ? Data.AlphaR1 : Data.AlphaT1);
-	const float4 ClothResult = Data.Fuzz.Eval(L, V, H, Payload.FuzzColor);
+	const float4 ClothResult = bIsReflection ? Data.Fuzz.Eval(L, V, H, Payload.FuzzColor) : 0;
 	const float ClothPdf = ClothResult.w;
-	const float DiffusePdf = NoL / PI;
+	const float ScatterPdf = HenyeyGreensteinPhase(Payload.PhaseG, dot(V_World, L_World));
 	if (bSampleGlass)
 	{
 		FThinSlabWeights SlabResult = ComputeThinSlabWeights(Data.TransmissionColor, VoH, Payload.Ior, Data.F0);
@@ -183,14 +188,14 @@ FMaterialSample SubstrateThinGlass_SampleMaterial(
 
 		if (bIsReflection)
 		{
-			Result.Pdf += Data.LobePdf.x * DiffusePdf;
 			Result.Pdf += Data.LobePdf.y * ClothPdf;
 		}
+		Result.Pdf += Data.LobePdf.x * ScatterPdf;
 	}
-	else if (bSampleDiffuse)
+	else if (bSampleScatter)
 	{
-		// Diffuse Lobe
-		Result.AddLobeWithMIS(Data.DiffuseWeight * Payload.DiffuseColor, DiffusePdf, Data.LobePdf.x);
+		// Scattering Lobe
+		Result.AddLobeWithMIS(Data.DiffuseWeight * Payload.DiffuseColor * ScatterWeight, ScatterPdf, Data.LobePdf.x);
 		Result.Pdf += Data.LobePdf.y * ClothPdf;
 		Result.Pdf += Data.LobePdf.z * (1.0 - Payload.RoughnessData.z) * GGXResult0.y;
 		Result.Pdf += Data.LobePdf.z * (      Payload.RoughnessData.z) * GGXResult1.y;
@@ -201,7 +206,7 @@ FMaterialSample SubstrateThinGlass_SampleMaterial(
 		const float ShadowTerminator = ShadowTerminatorTerm(L_World, N_World, Payload.WorldSmoothNormal);
 		Result.AddLobeWithMIS(ClothResult.xyz * ShadowTerminator, ClothPdf, Data.LobePdf.y);
 
-		Result.Pdf += Data.LobePdf.x * DiffusePdf;
+		Result.Pdf += Data.LobePdf.x * ScatterPdf;
 		Result.Pdf += Data.LobePdf.z * (1.0 - Payload.RoughnessData.z) * GGXResult0.y;
 		Result.Pdf += Data.LobePdf.z * (      Payload.RoughnessData.z) * GGXResult1.y;
 	}
@@ -235,6 +240,13 @@ FMaterialEval SubstrateThinGlass_EvalMaterial(
 	const float3 V = Data.V;
 	float3 L = mul(Data.Basis, L_World);
 
+	FMaterialEval Result = NullMaterialEval();
+
+	// Scattering Lobe
+	const float ScatterPdf = HenyeyGreensteinPhase(Payload.PhaseG, dot(V_World, L_World));
+	const float3 ScatterWeight = HanrahanKrueger_Weight(V.z, L.z, Data.TransmissionColor);
+	Result.AddLobeWithMIS(DiffuseSpecularScale.x * Data.DiffuseWeight * Payload.DiffuseColor * ScatterWeight, ScatterPdf, Data.LobePdf.x);
+
 	const bool bIsReflection = L.z >= 0;
 	L.z = abs(L.z); // push L to the same side as V
 
@@ -245,9 +257,6 @@ FMaterialEval SubstrateThinGlass_EvalMaterial(
 	const float NoH = saturate(H.z);
 
 
-	// Diffuse Lobe
-	const float DiffusePdf = NoL / PI;
-
 	// Specular Lobes (either reflected or transmitted)
 	const float2 GGXResult0 = GGXEvalReflection(L, V, H, bIsReflection ? Data.AlphaR0 : Data.AlphaT0);
 	const float2 GGXResult1 = GGXEvalReflection(L, V, H, bIsReflection ? Data.AlphaR1 : Data.AlphaT1);
@@ -260,18 +269,12 @@ FMaterialEval SubstrateThinGlass_EvalMaterial(
 
 	const float LobePdf = bIsReflection ? Data.LobePdf.z : Data.LobePdf.w;
 
-	FMaterialEval Result = NullMaterialEval();
-
 	Result.AddLobeWithMIS(DiffuseSpecularScale.y * Data.Fuzz.Attenuation * (1 - Payload.RoughnessData.z) * GGXResult0.x * SlabW0 * SlabTint, GGXResult0.y, LobePdf * (1 - Payload.RoughnessData.z));
 	Result.AddLobeWithMIS(DiffuseSpecularScale.y * Data.Fuzz.Attenuation * (    Payload.RoughnessData.z) * GGXResult1.x * SlabW1 * SlabTint, GGXResult1.y, LobePdf * (    Payload.RoughnessData.z));
 
 
 	if (bIsReflection)
 	{
-		{
-			// Diffuse Lobe
-			Result.AddLobeWithMIS(DiffuseSpecularScale.x * Data.DiffuseWeight * Payload.DiffuseColor, DiffusePdf, Data.LobePdf.x);
-		}
 		{
 			// Cloth Lobe
 			const float ShadowTerminator = ShadowTerminatorTerm(L_World, N_World, Payload.WorldSmoothNormal);
