Add support for ACES 2.0 in SDR and plumb HDR display info from RHI

- Refactor PostProcessCombineLUTs.usf to make all tonemappers available to all outputs.
- Add SDR/HDR tonemapping method as separate settings to the PostProcessVolume, so we know which one to use if an SDR/HDR change happens.
- If Filmic is selected for HDR, ACES 2.0 is used to invert and remap. This may be preferable to real ACES for some applications. This will be the default HDR behavior starting with UE 5.8.
- Add r.LUT.TonemappingMethod to override the LUT tonemapping method and sidestep Post Process Volume settings.
- Add r.LUT.TonemappingSDRInvEOTF to control how linear values are remapped into nonlinear space.
- Set r.HDR.Aces.Version to 2 by default since we're trying to steer users away from 1.3.

diff --git a/Engine/Shaders/Private/PostProcessCombineLUTs.usf b/Engine/Shaders/Private/PostProcessCombineLUTs.usf
--- a/Engine/Shaders/Private/PostProcessCombineLUTs.usf
+++ b/Engine/Shaders/Private/PostProcessCombineLUTs.usf
@@ -1,7 +1,7 @@
 // Copyright Epic Games, Inc. All Rights Reserved.
 
 /*=============================================================================
-	LUTBlender.usf: Filter pixel shader source.   
+	PostProcessCombineLUTs.usf: Filter pixel shader source.
 
 	Note: Any changes to device encoding logic in CombineLUTsCommon() must
 	also be made in PostProcessDeviceEncodingOnly.usf.
@@ -31,6 +31,22 @@ DECLARE_SCALAR_ARRAY(float, LUTWeights, 5);
 half3 ColorScale;
 half4 OverlayColor;
 
+float3 LimitingRgbToXYZ0;
+float3 LimitingRgbToXYZ1;
+float3 LimitingRgbToXYZ2;
+float3x3 GetLimitingRgbToXYZ()
+{
+	return float3x3(LimitingRgbToXYZ0, LimitingRgbToXYZ1, LimitingRgbToXYZ2);
+}
+
+float3 LimitingXYZToRgb0;
+float3 LimitingXYZToRgb1;
+float3 LimitingXYZToRgb2;
+float3x3 GetLimitingXYZToRgb()
+{
+	return float3x3(LimitingXYZToRgb0, LimitingXYZToRgb1, LimitingXYZToRgb2);
+}
+
 uint bUseMobileTonemapper;
 uint bIsTemperatureWhiteBalance;
 
@@ -80,10 +96,16 @@ float ACESGamutCompression;
 Texture2D<float> ACESReachTable;
 Texture2D<float> ACESGamutTable;
 Texture2D<float> ACESGammaTable;
+Texture2D<float> ACESReachTableSDR;
+Texture2D<float> ACESGamutTableSDR;
+Texture2D<float> ACESGammaTableSDR;
+uint bScaleWhite;
+uint bScaleWhiteSDR;
 
 static const float DefaultTemperature = 6500.0;
 static const float DefaultTint = 0.0;
 static const float Epsilon = 1e-8;
+static const float Bt1886Gamma = 2.4;
 
 float3 ColorCorrect( float3 WorkingColor,
 	float4 ColorSaturation,
@@ -156,73 +178,11 @@ float BlueCorrection;
 float ExpandGamut;
 float ToneCurveAmount;
 
-float4 CombineLUTsCommon(float3 LUTEncodedColor)
+float3 ComputeFilmColorNoGamma(float3 ColorAP1)
 {
-	float4 OutColor = 0;
-	
-	const float3x3 AP1_2_sRGB = mul( XYZ_2_sRGB_MAT, mul( D60_2_D65_CAT, AP1_2_XYZ_MAT ) );
 	const float3x3 AP0_2_AP1 = mul( XYZ_2_AP1_MAT, AP0_2_XYZ_MAT );
 	const float3x3 AP1_2_AP0 = mul( XYZ_2_AP0_MAT, AP1_2_XYZ_MAT );
 
-	const float3x3 AP1_2_Output  = OuputGamutMappingMatrix( OutputGamut );
-
-	float3 LinearColor;
-	// Decode texture values as ST-2084 (Dolby PQ)
-	BRANCH
-	if (GetOutputDevice() >= TONEMAPPER_OUTPUT_ACES1000nitST2084)
-	{
-		// Since ST2084 returns linear values in nits, divide by a scale factor to convert
-		// the reference nit result to be 1.0 in linear.
-		// (for efficiency multiply by precomputed inverse)
-		LinearColor = ST2084ToLinear(LUTEncodedColor) * LinearToNitsScaleInverse;
-	}
-	// Decode log values
-	else
-	{
-		LinearColor = LogToLin(LUTEncodedColor) - LogToLin(0);
-	}
-
-	float3 BalancedColor = LinearColor;
-	
-	BRANCH
-	if( abs(WhiteTemp - DefaultTemperature) > Epsilon || abs(WhiteTint - DefaultTint) > Epsilon )
-	{
-		BalancedColor = WhiteBalance(LinearColor, WhiteTemp, WhiteTint, bIsTemperatureWhiteBalance, (float3x3)WorkingColorSpace.ToXYZ, (float3x3)WorkingColorSpace.FromXYZ);
-	}
-
-	float3 ColorAP1 = mul( (float3x3)WorkingColorSpace.ToAP1, BalancedColor );
-
-	// Expand bright saturated colors outside the sRGB gamut to fake wide gamut rendering.
- 	float  LumaAP1 = dot( ColorAP1, AP1_RGB2Y );
-	float3 ChromaAP1 = ColorAP1 / LumaAP1;
-
-	float ChromaDistSqr = dot( ChromaAP1 - 1, ChromaAP1 - 1 );
-	float ExpandAmount = ( 1 - exp2( -4 * ChromaDistSqr ) ) * ( 1 - exp2( -4 * ExpandGamut * LumaAP1*LumaAP1 ) );
-
-	// Bizarre matrix but this expands sRGB to between P3 and AP1
-	// CIE 1931 chromaticities:	x		y
-	//				Red:		0.6965	0.3065
-	//				Green:		0.245	0.718
-	//				Blue:		0.1302	0.0456
-	//				White:		0.3127	0.329
-	const float3x3 Wide_2_XYZ_MAT = 
-	{
-		0.5441691,  0.2395926,  0.1666943,
-		0.2394656,  0.7021530,  0.0583814,
-		-0.0023439,  0.0361834,  1.0552183,
-	};
-
-	const float3x3 Wide_2_AP1 = mul( XYZ_2_AP1_MAT, Wide_2_XYZ_MAT );
-	const float3x3 ExpandMat = mul( Wide_2_AP1, AP1_2_sRGB );
-
-	float3 ColorExpand = mul( ExpandMat, ColorAP1 );
-	ColorAP1 = lerp( ColorAP1, ColorExpand, ExpandAmount );
-
-	ColorAP1 = ColorCorrectAll( ColorAP1 );
-
-	// Store for Linear HDR output without tone curve
-	float3 GradedColor = mul( (float3x3)WorkingColorSpace.FromAP1, ColorAP1 );
-
 	const float3x3 BlueCorrect =
 	{
 		0.9404372683, -0.0183068787, 0.0778696104,
@@ -231,7 +191,7 @@ float4 CombineLUTsCommon(float3 LUTEncodedColor)
 	};
 	const float3x3 BlueCorrectInv =
 	{
-		1.06318,     0.0233956, -0.0865726,
+		 1.06318,     0.0233956, -0.0865726,
 		-0.0106337,   1.20632,   -0.19569,
 		-0.000590887, 0.00105248, 0.999538
 	};
@@ -291,114 +251,250 @@ float4 CombineLUTsCommon(float3 LUTEncodedColor)
 
 	// blend with custom LDR color, used for Fade track in Cinematics
 	float3 FilmColorNoGamma = lerp( FilmColor * ColorScale, OverlayColor.rgb, OverlayColor.a );
-	// Apply Fade track to linear outputs also
+	return FilmColorNoGamma;
+}
+
+float3 ComputeGradedColor(float3 ColorAP1)
+{
+	// Store for Linear HDR output without tone curve
+	float3 GradedColor = mul( (float3x3)WorkingColorSpace.FromAP1, ColorAP1 );
+
+	// blend with custom LDR color, used for Fade track in Cinematics
 	GradedColor = lerp(GradedColor * ColorScale, OverlayColor.rgb, OverlayColor.a);
 
+	return GradedColor;
+}
 
-	// Apply "gamma" curve adjustment.
-	FilmColor = pow( max(0, FilmColorNoGamma), InverseGamma.y );
-		
-	// Note: Any changes to the following device encoding logic must also be made 
-	// in PostProcessDeviceEncodingOnly.usf's corresponding pixel shader.
-	half3 OutDeviceColor = 0;
-	BRANCH
+float3 ComputeOutputACES(
+	float3 GradedColor,
+	float PeakLuminance,
+	float3x3 LimitingRgbToXYZ,
+	float3x3 LimitingXYZToRgb,
+	float3x3 AP1_2_Output)
+{
+#if USE_ACES_2
+	float3 ODTColor = ACESOutputTransform(
+		GradedColor * ACESSceneColorMultiplier,
+		(float3x3)WorkingColorSpace.ToAP0,
+		PeakLuminance,
+		ACESReachTable,
+		ACESGamutTable,
+		ACESGammaTable,
+		bScaleWhite,
+		LimitingRgbToXYZ,
+		LimitingXYZToRgb,
+		GetColorSpace(OutputGamut));
+	return ODTColor;
+#else
+	FACESTonemapParams AcesParams = ComputeACESTonemapParams(ACESMinMaxData, ACESMidData, ACESCoefsLow_0, ACESCoefsHigh_0, ACESCoefsLow_4, ACESCoefsHigh_4, ACESSceneColorMultiplier, ACESGamutCompression);
+	float3 ODTColor = ACESOutputTransform( GradedColor, (float3x3)WorkingColorSpace.ToAP0, AcesParams);
+	float3 OutDeviceColor = mul(AP1_2_Output, ODTColor);
+	return OutDeviceColor;
+#endif
+}
 
-	// sRGB, user specified gamut
-	if( GetOutputDevice() == TONEMAPPER_OUTPUT_sRGB)
+float3 ComputeInvOutputACES(float3 ODTColor, float3x3 LimitingRgbToXYZ, float3x3 LimitingXYZToRgb)
+{
+#if USE_ACES_2
+	float3 GradedColorMultiplied = ACESInvOutputTransform(
+		ODTColor,
+		(float3x3)WorkingColorSpace.FromAP0,
+		LinearToNitsScale,
+		ACESReachTableSDR,
+		ACESGamutTableSDR,
+		ACESGammaTableSDR,
+		bScaleWhiteSDR,
+		LimitingRgbToXYZ,
+		LimitingXYZToRgb,
+		GetColorSpace(OutputGamut));
+	return GradedColorMultiplied / ACESSceneColorMultiplier;
+#else
+	// Magenta makes it obvious this isn't implemented.
+	return float3(1., 0., 1.);
+#endif
+}
+
+/** Values of output device (matching C++'s ETonemappingMethod) */
+#define TONEMAPPING_METHOD_Filmic 0
+#define TONEMAPPING_METHOD_ACES 1
+uint TonemappingMethod;
+
+#define TONEMAPPING_SDR_INVEOTF_SRGB 0
+#define TONEMAPPING_SDR_INVEOTF_BT1886 1
+uint TonemappingSDRInvEOTF;
+
+float3 ApplyToneCurve(float3 ColorAP1, float PeakLuminance, float3x3 AP1_2_Output)
+{
+	float3 OutputGamutColor;
+	if (TonemappingMethod == TONEMAPPING_METHOD_ACES)
 	{
-		// Convert from sRGB to specified output gamut	
-		//float3 OutputGamutColor = mul( AP1_2_Output, mul( sRGB_2_AP1, FilmColor ) );
+		float3 GradedColor = ComputeGradedColor(ColorAP1);
+		OutputGamutColor = ComputeOutputACES(GradedColor, PeakLuminance, GetLimitingRgbToXYZ(), GetLimitingXYZToRgb(), AP1_2_Output) * LinearToNitsScaleInverse;
+	}
+	// TonemappingMethod == TONEMAPPING_METHOD_Filmic
+	else
+	{
+		float3 FilmColorNoGamma = ComputeFilmColorNoGamma(ColorAP1);
+
+		// Apply "gamma" curve adjustment.
+		float3 FilmColor = pow( max(0, FilmColorNoGamma), InverseGamma.y );
 
 		// FIXME: Workaround for UE-29935, pushing all colors with a 0 component to black output
 		// Default parameters seem to cancel out (sRGB->XYZ->AP1->XYZ->sRGB), so should be okay for a temp fix
-		float3 OutputGamutColor = WorkingColorSpace.bIsSRGB ? FilmColor : mul( AP1_2_Output, mul( (float3x3)WorkingColorSpace.ToAP1, FilmColor ) );
+		OutputGamutColor = WorkingColorSpace.bIsSRGB && (OutputGamut == TONEMAPPER_GAMUT_sRGB_D65) ?
+			FilmColor :
+			mul( AP1_2_Output, mul( (float3x3)WorkingColorSpace.ToAP1, FilmColor ) );
+
+		// "Automatic HDR" using ACES to remap filmic.
+		if (PeakLuminance > LinearToNitsScale)
+		{
+			// Rec. 2100 specifies a dim environment similar to various viewing standards, e.g. BT.2035, for consuming Rec. 709 w/BT.1886.
+			// sRGB is meant to be viewed in a bright office, so its EOTF leaves values a bit inflated in comparison to BT.1886.
+			// This adjustment converts from "bright" linear space to "dim" linear space.
+			if (TonemappingSDRInvEOTF == TONEMAPPING_SDR_INVEOTF_SRGB)
+			{
+				// Apply inverse of sRGB EOTF to generate nonlinear values that would get sent to an SDR display.
+				OutputGamutColor = LinearToSrgb(OutputGamutColor);
+				// Apply BT.1886 EOTF to linearize the values into "dim" linear space.
+				OutputGamutColor = pow(OutputGamutColor, Bt1886Gamma);
+			}
+
+			float3 GradedColor = ComputeInvOutputACES(OutputGamutColor * LinearToNitsScale, (float3x3)WorkingColorSpace.ToXYZ, (float3x3)WorkingColorSpace.FromXYZ);
+			OutputGamutColor = ComputeOutputACES(GradedColor, PeakLuminance, GetLimitingRgbToXYZ(), GetLimitingXYZToRgb(), AP1_2_Output) * LinearToNitsScaleInverse;
+		}
+	}
+
+	return OutputGamutColor;
+}
+
+float4 CombineLUTsCommon(float3 LUTEncodedColor)
+{
+	float4 OutColor = 0;
+
+	const float3x3 AP1_2_sRGB = mul( XYZ_2_sRGB_MAT, mul( D60_2_D65_CAT, AP1_2_XYZ_MAT ) );
+	const float3x3 AP1_2_Output = OutputGamutMappingMatrix( OutputGamut );
+
+	float3 LinearColor;
+	// Decode texture values as ST-2084 (Dolby PQ)
+	BRANCH
+	if (GetOutputDevice() >= TONEMAPPER_OUTPUT_ACES1000nitST2084)
+	{
+		// Since ST2084 returns linear values in nits, divide by a scale factor to convert
+		// the reference nit result to be 1.0 in linear.
+		// (for efficiency multiply by precomputed inverse)
+		LinearColor = ST2084ToLinear(LUTEncodedColor) * LinearToNitsScaleInverse;
+	}
+	// Decode log values
+	else
+	{
+		LinearColor = LogToLin(LUTEncodedColor) - LogToLin(0);
+	}
+
+	float3 BalancedColor = LinearColor;
+
+	BRANCH
+	if( abs(WhiteTemp - DefaultTemperature) > Epsilon || abs(WhiteTint - DefaultTint) > Epsilon )
+	{
+		BalancedColor = WhiteBalance(LinearColor, WhiteTemp, WhiteTint, bIsTemperatureWhiteBalance, (float3x3)WorkingColorSpace.ToXYZ, (float3x3)WorkingColorSpace.FromXYZ);
+	}
+
+	float3 ColorAP1 = mul( (float3x3)WorkingColorSpace.ToAP1, BalancedColor );
+
+	// Expand bright saturated colors outside the sRGB gamut to fake wide gamut rendering.
+ 	float  LumaAP1 = dot( ColorAP1, AP1_RGB2Y );
+	float3 ChromaAP1 = ColorAP1 / LumaAP1;
+
+	float ChromaDistSqr = dot( ChromaAP1 - 1, ChromaAP1 - 1 );
+	float ExpandAmount = ( 1 - exp2( -4 * ChromaDistSqr ) ) * ( 1 - exp2( -4 * ExpandGamut * LumaAP1*LumaAP1 ) );
+
+	// Bizarre matrix but this expands sRGB to between P3 and AP1
+	// CIE 1931 chromaticities:	x		y
+	//				Red:		0.6965	0.3065
+	//				Green:		0.245	0.718
+	//				Blue:		0.1302	0.0456
+	//				White:		0.3127	0.329
+	const float3x3 Wide_2_XYZ_MAT = 
+	{
+		0.5441691,  0.2395926,  0.1666943,
+		0.2394656,  0.7021530,  0.0583814,
+		-0.0023439,  0.0361834,  1.0552183,
+	};
+
+	const float3x3 Wide_2_AP1 = mul( XYZ_2_AP1_MAT, Wide_2_XYZ_MAT );
+	const float3x3 ExpandMat = mul( Wide_2_AP1, AP1_2_sRGB );
+
+	float3 ColorExpand = mul( ExpandMat, ColorAP1 );
+	ColorAP1 = lerp( ColorAP1, ColorExpand, ExpandAmount );
+
+	ColorAP1 = ColorCorrectAll( ColorAP1 );
+
+	// Note: Any changes to the following device encoding logic must also be made 
+	// in PostProcessDeviceEncodingOnly.usf's corresponding pixel shader.
+	half3 OutDeviceColor = 0;
 
-		// Apply conversion to sRGB (this must be an exact sRGB conversion else darks are bad).
-		OutDeviceColor = LinearToSrgb( OutputGamutColor );
+	// sRGB, user specified gamut
+	BRANCH if( GetOutputDevice() == TONEMAPPER_OUTPUT_sRGB)
+	{
+		float3 OutputGamutColor = ApplyToneCurve(ColorAP1, LinearToNitsScale, AP1_2_Output);
+
+		BRANCH if (TonemappingSDRInvEOTF == TONEMAPPING_SDR_INVEOTF_SRGB)
+		{
+			OutDeviceColor = LinearToSrgb(OutputGamutColor);
+		}
+		else // TonemappingSDRInvEOTF == TONEMAPPING_SDR_INVEOTF_BT1886
+		{
+			OutDeviceColor = pow(OutputGamutColor, 1. / Bt1886Gamma);
+		}
 	}
 
 	// Rec 709, user specified gamut
-	else if( GetOutputDevice() == TONEMAPPER_OUTPUT_Rec709)
+	else BRANCH if( GetOutputDevice() == TONEMAPPER_OUTPUT_Rec709)
 	{
-		// Convert from the working color space to specified output gamut
-		float3 OutputGamutColor = mul( AP1_2_Output, mul( (float3x3)WorkingColorSpace.ToAP1, FilmColor ) );
+		float3 OutputGamutColor = ApplyToneCurve(ColorAP1, LinearToNitsScale, AP1_2_Output);
 
 		// Didn't profile yet if the branching version would be faster (different linear segment).
-		OutDeviceColor = LinearTo709Branchless( OutputGamutColor );
+		OutDeviceColor = LinearTo709Branchless(OutputGamutColor);
 	}
 
-	// ACES transform with PQ/2084 encoding, user specified gamut 
-	else if( GetOutputDevice() == TONEMAPPER_OUTPUT_ACES1000nitST2084 || GetOutputDevice() == TONEMAPPER_OUTPUT_ACES2000nitST2084)
+	// Transform with PQ/2084 encoding, user specified gamut
+	else BRANCH if( GetOutputDevice() == TONEMAPPER_OUTPUT_ACES1000nitST2084 || GetOutputDevice() == TONEMAPPER_OUTPUT_ACES2000nitST2084)
 	{
-#if USE_ACES_2
-		// TODO: The display limiting color space should be user-exposed as part of display characterization.
-		float3 ODTColor = ACESOutputTransform(
-			GradedColor * ACESSceneColorMultiplier,
-			(float3x3)WorkingColorSpace.ToAP0,
-			OutputMaxLuminance,
-			ACESReachTable,
-			ACESGamutTable,
-			ACESGammaTable,
-			false, // bScaleWhite
-			GetColorSpace(TONEMAPPER_GAMUT_Rec2020_D65),
-			GetColorSpace(TONEMAPPER_GAMUT_Rec2020_D65));
-#else
-		FACESTonemapParams AcesParams = ComputeACESTonemapParams(ACESMinMaxData, ACESMidData, ACESCoefsLow_0, ACESCoefsHigh_0, ACESCoefsLow_4, ACESCoefsHigh_4, ACESSceneColorMultiplier, ACESGamutCompression);
-		float3 ODTColor = ACESOutputTransform( GradedColor, (float3x3)WorkingColorSpace.ToAP0, AcesParams);
-		ODTColor = mul(AP1_2_Output, ODTColor);
-#endif
+		float3 OutputGamutColor = ApplyToneCurve(ColorAP1, OutputMaxLuminance, AP1_2_Output);
 
 		// Apply conversion to ST-2084 (Dolby PQ)
-		OutDeviceColor = LinearToST2084( ODTColor );
+		OutDeviceColor = LinearToST2084(OutputGamutColor * LinearToNitsScale);
 	}
 
-	// ACES transform to linear ScRGB
-	else if( GetOutputDevice() == TONEMAPPER_OUTPUT_ACES1000nitScRGB || GetOutputDevice() == TONEMAPPER_OUTPUT_ACES2000nitScRGB)
+	// Transform to linear ScRGB
+	else BRANCH if( GetOutputDevice() == TONEMAPPER_OUTPUT_ACES1000nitScRGB || GetOutputDevice() == TONEMAPPER_OUTPUT_ACES2000nitScRGB)
 	{
-#if USE_ACES_2
-		// TODO: The display limiting color space should be user-exposed as part of display characterization.
-		float3 ODTColor = ACESOutputTransform(
-			GradedColor * ACESSceneColorMultiplier,
-			(float3x3)WorkingColorSpace.ToAP0,
-			OutputMaxLuminance,
-			ACESReachTable,
-			ACESGamutTable,
-			ACESGammaTable,
-			false, // bScaleWhite
-			GetColorSpace(TONEMAPPER_GAMUT_Rec2020_D65),
-			GetColorSpace(TONEMAPPER_GAMUT_sRGB_D65));
-		OutDeviceColor = ODTColor / UE_SCRGB_WHITE_NITS;
-#else
-		FACESTonemapParams AcesParams = ComputeACESTonemapParams(ACESMinMaxData, ACESMidData, ACESCoefsLow_0, ACESCoefsHigh_0, ACESCoefsLow_4, ACESCoefsHigh_4, ACESSceneColorMultiplier, ACESGamutCompression);
-		float3 ODTColor = ACESOutputTransform( GradedColor, (float3x3)WorkingColorSpace.ToAP0, AcesParams);
-		OutDeviceColor = mul(AP1_2_sRGB, ODTColor / UE_SCRGB_WHITE_NITS);
-#endif
+		float3 OutputGamutColor = ApplyToneCurve(ColorAP1, OutputMaxLuminance, AP1_2_Output);
+		OutDeviceColor = (LinearToNitsScale / UE_SCRGB_WHITE_NITS) * OutputGamutColor;
 	}
-	
-	else if( GetOutputDevice() == TONEMAPPER_OUTPUT_LinearEXR)
+
+	else BRANCH if( GetOutputDevice() == TONEMAPPER_OUTPUT_LinearEXR)
 	{
+		float3 GradedColor = ComputeGradedColor(ColorAP1);
 		float3 OutputGamutColor = mul( AP1_2_Output, mul( (float3x3)WorkingColorSpace.ToAP1, GradedColor ) );
 		OutDeviceColor = LinearToST2084( OutputGamutColor );
 	}
 	// Linear HDR, including all color correction, but no tone curve
-	else if( GetOutputDevice() == TONEMAPPER_OUTPUT_NoToneCurve)
+	else BRANCH if( GetOutputDevice() == TONEMAPPER_OUTPUT_NoToneCurve)
 	{
+		float3 GradedColor = ComputeGradedColor(ColorAP1);
 		OutDeviceColor = GradedColor;
 	}
 	// "Linear" including all color correction and the tone curve, but no device gamma
-	else if (GetOutputDevice() == TONEMAPPER_OUTPUT_WithToneCurve)
+	else BRANCH if (GetOutputDevice() == TONEMAPPER_OUTPUT_WithToneCurve)
 	{
-		float3 OutputGamutColor = mul( AP1_2_Output, mul( (float3x3)WorkingColorSpace.ToAP1, FilmColorNoGamma ) );
-
-		OutDeviceColor = OutputGamutColor;
+		OutDeviceColor = ApplyToneCurve(ColorAP1, LinearToNitsScale, AP1_2_Output);
 	}
 
 	// OutputDevice == TONEMAPPER_OUTPUT_ExplicitGammaMapping
 	// Gamma 2.2, user specified gamut
 	else
 	{
-		// Convert from the working color space to specified output gamut
-		float3 OutputGamutColor = mul( AP1_2_Output, mul( (float3x3)WorkingColorSpace.ToAP1, FilmColor ) );
+		float3 OutputGamutColor = ApplyToneCurve(ColorAP1, LinearToNitsScale, AP1_2_Output);
 
 		// This is different than the prior "gamma" curve adjustment (but reusing the variable).
 		// For displays set to a gamma colorspace.
diff --git a/Engine/Shaders/Private/PostProcessDeviceEncodingOnly.usf b/Engine/Shaders/Private/PostProcessDeviceEncodingOnly.usf
--- a/Engine/Shaders/Private/PostProcessDeviceEncodingOnly.usf
+++ b/Engine/Shaders/Private/PostProcessDeviceEncodingOnly.usf
@@ -40,6 +40,7 @@ float ACESGamutCompression;
 Texture2D<float> ACESReachTable;
 Texture2D<float> ACESGamutTable;
 Texture2D<float> ACESGammaTable;
+uint bScaleWhite;
 
 uint GetOutputDevice()
 {
@@ -67,7 +68,7 @@ float4 DeviceEncodingOnlyCommonPS(
 	half4 SceneColor = SampleSceneColor(SceneUV);
 
 	const float3x3 AP1_2_sRGB = mul( XYZ_2_sRGB_MAT, mul( D60_2_D65_CAT, AP1_2_XYZ_MAT ) );
-	const float3x3 AP1_2_Output  = OuputGamutMappingMatrix( OutputGamut );
+	const float3x3 AP1_2_Output  = OutputGamutMappingMatrix( OutputGamut );
 
 	// Apply "gamma" curve adjustment.
 	float3 FilmColor = pow(max(0, SceneColor.rgb), InverseGamma.y);
@@ -105,7 +106,9 @@ float4 DeviceEncodingOnlyCommonPS(
 	else if( GetOutputDevice() == TONEMAPPER_OUTPUT_ACES1000nitST2084 || GetOutputDevice() == TONEMAPPER_OUTPUT_ACES2000nitST2084)
 	{
 #if USE_ACES_2
+		FColorSpace EncodingColorSpace = GetColorSpace(TONEMAPPER_GAMUT_Rec2020_D65);
 		// TODO: The display limiting color space should be user-exposed as part of display characterization.
+		FColorSpace LimitingColorSpace = GetColorSpace(TONEMAPPER_GAMUT_Rec2020_D65);
 		float3 ODTColor = ACESOutputTransform(
 			FilmColor * ACESSceneColorMultiplier,
 			(float3x3)WorkingColorSpace.ToAP0,
@@ -113,9 +116,10 @@ float4 DeviceEncodingOnlyCommonPS(
 			ACESReachTable,
 			ACESGamutTable,
 			ACESGammaTable,
-			false, // bScaleWhite
-			GetColorSpace(TONEMAPPER_GAMUT_Rec2020_D65),
-			GetColorSpace(TONEMAPPER_GAMUT_Rec2020_D65));
+			bScaleWhite,
+			LimitingColorSpace.RGBtoXYZ,
+			LimitingColorSpace.XYZtoRGB,
+			EncodingColorSpace);
 #else
 		FACESTonemapParams AcesParams = ComputeACESTonemapParams(ACESMinMaxData, ACESMidData, ACESCoefsLow_0, ACESCoefsHigh_0, ACESCoefsLow_4, ACESCoefsHigh_4, ACESSceneColorMultiplier, ACESGamutCompression);
 		float3 ODTColor = ACESOutputTransform( FilmColor, (float3x3)WorkingColorSpace.ToAP0, AcesParams);
@@ -131,6 +135,7 @@ float4 DeviceEncodingOnlyCommonPS(
 	{
 #if USE_ACES_2
 		// TODO: The display limiting color space should be user-exposed as part of display characterization.
+		FColorSpace LimitingColorSpace = GetColorSpace(TONEMAPPER_GAMUT_Rec2020_D65);
 		float3 ODTColor = ACESOutputTransform(
 			FilmColor * ACESSceneColorMultiplier,
 			(float3x3)WorkingColorSpace.ToAP0,
@@ -138,8 +143,9 @@ float4 DeviceEncodingOnlyCommonPS(
 			ACESReachTable,
 			ACESGamutTable,
 			ACESGammaTable,
-			false, // bScaleWhite
-			GetColorSpace(TONEMAPPER_GAMUT_Rec2020_D65),
+			bScaleWhite,
+			LimitingColorSpace.RGBtoXYZ,
+			LimitingColorSpace.XYZtoRGB,
 			GetColorSpace(TONEMAPPER_GAMUT_sRGB_D65));
 		OutDeviceColor = ODTColor / UE_SCRGB_WHITE_NITS;
 #else
diff --git a/Engine/Shaders/Private/TonemapCommon.ush b/Engine/Shaders/Private/TonemapCommon.ush
--- a/Engine/Shaders/Private/TonemapCommon.ush
+++ b/Engine/Shaders/Private/TonemapCommon.ush
@@ -328,7 +328,8 @@ float3 ACESOutputTransform(
 	Texture2D<float> GamutTable,
 	Texture2D<float> GammaTable,
 	bool bScaleWhite,
-	FColorSpace LimitingColorSpace,
+	float3x3 LIMIT_RGB_TO_XYZ,
+	float3x3 LIMIT_XYZ_TO_RGB,
 	FColorSpace EncodingColorSpace
 )
 {
@@ -336,8 +337,8 @@ float3 ACESOutputTransform(
 
 	ODTParams PARAMS = init_ODTParams(
 		PeakLuminance,
-		LimitingColorSpace.RGBtoXYZ,
-		LimitingColorSpace.XYZtoRGB,
+		LIMIT_RGB_TO_XYZ,
+		LIMIT_XYZ_TO_RGB,
 		EncodingColorSpace.RGBtoXYZ,
 		EncodingColorSpace.XYZtoRGB,
 		ReachTable,
@@ -367,7 +368,8 @@ float3 ACESInvOutputTransform(
 	Texture2D<float> GamutTable,
 	Texture2D<float> GammaTable,
 	const bool bScaleWhite,
-	FColorSpace LimitingColorSpace,
+	float3x3 LIMIT_RGB_TO_XYZ,
+	float3x3 LIMIT_XYZ_TO_RGB,
 	FColorSpace EncodingColorSpace
 )
 {
@@ -375,8 +377,8 @@ float3 ACESInvOutputTransform(
 
 	ODTParams PARAMS = init_ODTParams(
 		PeakLuminance,
-		LimitingColorSpace.RGBtoXYZ,
-		LimitingColorSpace.XYZtoRGB,
+		LIMIT_RGB_TO_XYZ,
+		LIMIT_XYZ_TO_RGB,
 		EncodingColorSpace.RGBtoXYZ,
 		EncodingColorSpace.XYZtoRGB,
 		ReachTable,
@@ -473,7 +475,7 @@ static const float3x3 GamutMappingIdentityMatrix = { 1.0, 0.0, 0.0, 0.0, 1.0, 0.
 //
 // Gamut conversion matrices
 //
-float3x3 OuputGamutMappingMatrix( uint OutputGamut )
+float3x3 OutputGamutMappingMatrix( uint OutputGamut )
 {
 	// Gamut mapping matrices used later
 	const float3x3 AP1_2_sRGB    = mul( XYZ_2_sRGB_MAT, mul( D60_2_D65_CAT, AP1_2_XYZ_MAT ) );
@@ -500,7 +502,7 @@ float3x3 OuputGamutMappingMatrix( uint OutputGamut )
 	return GamutMappingMatrix; 
 }
 
-float3x3 OuputInverseGamutMappingMatrix( uint OutputGamut )
+float3x3 OutputInverseGamutMappingMatrix( uint OutputGamut )
 {
 	// Gamut mapping matrices used later
 	const float3x3 sRGB_2_AP1    = mul( XYZ_2_AP1_MAT, mul( D65_2_D60_CAT, sRGB_2_XYZ_MAT ) );
diff --git a/Engine/Source/Runtime/Engine/Classes/Engine/Scene.h b/Engine/Source/Runtime/Engine/Classes/Engine/Scene.h
--- a/Engine/Source/Runtime/Engine/Classes/Engine/Scene.h
+++ b/Engine/Source/Runtime/Engine/Classes/Engine/Scene.h
@@ -69,6 +69,15 @@ enum ETemperatureMethod : int
 	TEMP_MAX,
 };
 
+/** Used by FPostProcessSettings to determine the tonemapping method. */
+UENUM()
+enum class ETonemappingMethod : uint8
+{
+	Filmic UMETA(DisplayName = "Filmic"),
+	StandardACES UMETA(DisplayName = "Standard ACES", ToolTip = "Film settings do not apply to this method."),
+	MAX,
+};
+
 
 UENUM() 
 enum class ELightUnits : uint8
@@ -763,6 +772,10 @@ struct FPostProcessSettings
 	UPROPERTY(EditAnywhere, BlueprintReadWrite, Category = Overrides, meta = (PinHiddenByDefault, InlineEditConditionToggle))
 	uint8 bOverride_ToneCurveAmount : 1;
 
+	UPROPERTY(EditAnywhere, BlueprintReadWrite, Category=Overrides, meta=(PinHiddenByDefault, InlineEditConditionToggle))
+	uint8 bOverride_TonemappingMethodSDR:1;
+	UPROPERTY(EditAnywhere, BlueprintReadWrite, Category=Overrides, meta=(PinHiddenByDefault, InlineEditConditionToggle))
+	uint8 bOverride_TonemappingMethodHDR:1;
 	UPROPERTY(EditAnywhere, BlueprintReadWrite, Category=Overrides, meta=(PinHiddenByDefault, InlineEditConditionToggle))
 	uint8 bOverride_FilmSlope:1;
 	UPROPERTY(EditAnywhere, BlueprintReadWrite, Category=Overrides, meta=(PinHiddenByDefault, InlineEditConditionToggle))
@@ -1533,6 +1546,12 @@ struct FPostProcessSettings
 	UPROPERTY(interp, BlueprintReadWrite, Category = "Color Grading|Misc", meta = (ClampMin = "0.0", UIMax = "1.0", editcondition = "bOverride_ToneCurveAmount"))
 	float ToneCurveAmount;
 
+	/** Controls the method of tonemapping in SDR. */
+	UPROPERTY(interp, BlueprintReadWrite, Category="Film", meta=(editcondition = "bOverride_TonemappingMethodSDR", DisplayName = "SDR Method"))
+	ETonemappingMethod TonemappingMethodSDR;
+	/** Controls the method of tonemapping in HDR. */
+	UPROPERTY(interp, BlueprintReadWrite, Category="Film", meta=(editcondition = "bOverride_TonemappingMethodHDR", DisplayName = "HDR Method"))
+	ETonemappingMethod TonemappingMethodHDR;
 	/** Controls the overall steepness of the tonemapper curve.  Larger values increase scene contrast and smaller values reduce contrast. */
 	UPROPERTY(interp, BlueprintReadWrite, Category="Film", meta=(UIMin = "0.0", UIMax = "1.0", editcondition = "bOverride_FilmSlope", DisplayName = "Slope"))
 	float FilmSlope;
diff --git a/Engine/Source/Runtime/Engine/Private/PostProcessUtils.cpp b/Engine/Source/Runtime/Engine/Private/PostProcessUtils.cpp
--- a/Engine/Source/Runtime/Engine/Private/PostProcessUtils.cpp
+++ b/Engine/Source/Runtime/Engine/Private/PostProcessUtils.cpp
@@ -53,6 +53,8 @@ bool FPostProcessUtils::OverridePostProcessSettings(FPostProcessSettings& ThisFr
 		UE_SET_PP(ExpandGamut);
 		UE_SET_PP(ToneCurveAmount);
 
+		UE_SET_PP(TonemappingMethodSDR);
+		UE_SET_PP(TonemappingMethodHDR);
 		UE_SET_PP(FilmSlope);
 		UE_SET_PP(FilmToe);
 		UE_SET_PP(FilmShoulder);
@@ -358,6 +360,8 @@ bool FPostProcessUtils::BlendPostProcessSettings(FPostProcessSettings& ThisFrom,
 		UE_LERP_PP(ExpandGamut);
 		UE_LERP_PP(ToneCurveAmount);
 
+		UE_SET_PP(TonemappingMethodSDR);
+		UE_SET_PP(TonemappingMethodHDR);
 		UE_LERP_PP(FilmSlope);
 		UE_LERP_PP(FilmToe);
 		UE_LERP_PP(FilmShoulder);
diff --git a/Engine/Source/Runtime/Engine/Private/Scene.cpp b/Engine/Source/Runtime/Engine/Private/Scene.cpp
--- a/Engine/Source/Runtime/Engine/Private/Scene.cpp
+++ b/Engine/Source/Runtime/Engine/Private/Scene.cpp
@@ -437,6 +437,9 @@ FPostProcessSettings::FPostProcessSettings()
 	ExpandGamut = 1.0f;
 	ToneCurveAmount = 1.0;
 
+	TonemappingMethodSDR = ETonemappingMethod::Filmic;
+	TonemappingMethodHDR = ETonemappingMethod::Filmic;
+
 	// default values:
 
 	// ACES settings
diff --git a/Engine/Source/Runtime/Engine/Private/SceneView.cpp b/Engine/Source/Runtime/Engine/Private/SceneView.cpp
--- a/Engine/Source/Runtime/Engine/Private/SceneView.cpp
+++ b/Engine/Source/Runtime/Engine/Private/SceneView.cpp
@@ -1618,6 +1618,8 @@ void FSceneView::OverridePostProcessSettings(const FPostProcessSettings& Src, fl
 		LERP_PP(ExpandGamut);
 		LERP_PP(ToneCurveAmount);
 
+		SET_PP(TonemappingMethodSDR);
+		SET_PP(TonemappingMethodHDR);
 		LERP_PP(FilmSlope);
 		LERP_PP(FilmToe);
 		LERP_PP(FilmShoulder);
diff --git a/Engine/Source/Runtime/Engine/Private/Slate/SceneViewport.cpp b/Engine/Source/Runtime/Engine/Private/Slate/SceneViewport.cpp
--- a/Engine/Source/Runtime/Engine/Private/Slate/SceneViewport.cpp
+++ b/Engine/Source/Runtime/Engine/Private/Slate/SceneViewport.cpp
@@ -385,17 +385,14 @@ static void ComputeSceneViewportHDRMetaData(FHDRMetaData& OutHDRMetaData, const
 {
 	// @todo vreditor switch: This code needs to be called when switching between stereo/non when going immersive.  Seems to always work out that way anyway though? (Probably due to resize)
 	IStereoRenderTargetManager* const StereoRenderTargetManager = RetrieveStereoRenderTargetManager(bIsStereoRenderingAllowed);
-	if (StereoRenderTargetManager != nullptr)
+	if (StereoRenderTargetManager && StereoRenderTargetManager->HDRGetMetaDataForStereo(OutHDRMetaData.DisplayOutputFormat, OutHDRMetaData.DisplayColorGamut, OutHDRMetaData.bHDRSupported))
 	{
 		// TODO: Get the real values for stereo.
 		OutHDRMetaData.MinimumLuminanceInNits = 0.f;
-		OutHDRMetaData.MaximumLuminanceInNits = 100.f;
-		OutHDRMetaData.MaximumFullFrameLuminanceInNits = 100.f;
-		OutHDRMetaData.LimitingColorSpace = UE::Color::FColorSpace::GetSRGB();
-		if (StereoRenderTargetManager->HDRGetMetaDataForStereo(OutHDRMetaData.DisplayOutputFormat, OutHDRMetaData.DisplayColorGamut, OutHDRMetaData.bHDRSupported))
-		{
-			return;
-		}
+		OutHDRMetaData.MaximumLuminanceInNits = OutHDRMetaData.bHDRSupported ? 1499.f : 100.f;
+		OutHDRMetaData.MaximumFullFrameLuminanceInNits = OutHDRMetaData.bHDRSupported ? 799.f : 100.f;
+		OutHDRMetaData.LimitingColorSpace = OutHDRMetaData.bHDRSupported ? UE::Color::FColorSpace::GetP3D65() : UE::Color::FColorSpace::GetSRGB();
+		return;
 	}
 
 	HDRGetMetaDataForWindow(
@@ -2153,7 +2150,7 @@ void FSceneViewport::InitRHI(FRHICommandListBase& RHICmdList)
 
 		if (BufferedSlateHandles.Num() < NumBufferedFrames)
 		{
-			//add sufficient entires for buffering.
+			//add sufficient entries for buffering.
 			for (int32 i = BufferedSlateHandles.Num(); i < NumBufferedFrames; i++)
 			{
 				BufferedSlateHandles.Add(new FSlateRenderTargetRHI(nullptr, 0, 0)); 
diff --git a/Engine/Source/Runtime/RenderCore/Private/RenderCore.cpp b/Engine/Source/Runtime/RenderCore/Private/RenderCore.cpp
--- a/Engine/Source/Runtime/RenderCore/Private/RenderCore.cpp
+++ b/Engine/Source/Runtime/RenderCore/Private/RenderCore.cpp
@@ -409,8 +409,8 @@ static TAutoConsoleVariable<int32> CVarHDRDisplayMaxLuminance(
 	ECVF_RenderThreadSafe
 );
 
-static TAutoConsoleVariable<bool> CVarHDROverrideOSMaxLuminance(
-	TEXT("r.HDR.OverrideOSMaxLuminance"),
+static TAutoConsoleVariable<bool> CVarHDRDisplayOverrideOSMaxLuminance(
+	TEXT("r.HDR.Display.OverrideOSMaxLuminance"),
 #if PLATFORM_USES_FIXED_HDR_SETTING
 	false,
 #else
@@ -435,7 +435,7 @@ static TAutoConsoleVariable<float> CVarHDRAcesGamutCompression(
 
 static TAutoConsoleVariable<int32> CVarHDRAcesVersion(
 	TEXT("r.HDR.Aces.Version"),
-	1,
+	2,
 	TEXT("Version of the ACES transform:\n")
 	TEXT("1: ACES 1.3 (default)\n")
 	TEXT("2: ACES 2.0\n"),
@@ -1068,7 +1068,6 @@ void ConfigureACESTonemapParams(FACESTonemapParams& OutACESTonemapParams, float
 
 	OutACESTonemapParams.ACESSceneColorMultiplier = CVarHDRAcesColorMultiplier.GetValueOnAnyThread();
 	OutACESTonemapParams.ACESGamutCompression = CVarHDRAcesGamutCompression.GetValueOnAnyThread();
-
 }
 
 // Converts PQ signal to linear values, see https://www.itu.int/rec/R-REC-BT.2124-0-201901-I/en, conversion 3 
diff --git a/Engine/Source/Runtime/Renderer/Private/PostProcess/PostProcessCombineLUTs.cpp b/Engine/Source/Runtime/Renderer/Private/PostProcess/PostProcessCombineLUTs.cpp
--- a/Engine/Source/Runtime/Renderer/Private/PostProcess/PostProcessCombineLUTs.cpp
+++ b/Engine/Source/Runtime/Renderer/Private/PostProcess/PostProcessCombineLUTs.cpp
@@ -50,6 +50,24 @@ TAutoConsoleVariable<int32> CVarLUTUpdateEveryFrame(
 	TEXT("Controls whether the tonemapping LUT pass is executed every frame."),
 	ECVF_RenderThreadSafe);
 
+TAutoConsoleVariable<int32> CVarLUTTonemappingMethod(
+	TEXT("r.LUT.TonemappingMethod"), -1,
+	TEXT("Set the tonemapping method to transform colors for the output device.\n")
+	TEXT("-1: Use PostProcessSettings (default)\n")
+	TEXT("0: Filmic\n")
+	TEXT("1: Standard ACES"),
+	ECVF_RenderThreadSafe);
+
+TAutoConsoleVariable<int32> CVarLUTTonemappingSDRInvEOTF(
+	TEXT("r.LUT.TonemappingSDRInvEOTF"), -1,
+	TEXT("Set the SDR EOTF for tonemapping to target for the output device.\n")
+	TEXT("The default for ACES is BT.1886 to prioritize compatibility between SDR and Rec. 2100 for HDR.\n")
+	TEXT("The default for Filmic is sRGB to prioritize symmetry with GPU SRGB conversion although HDR should still align.\n")
+	TEXT("-1: Use default\n")
+	TEXT("0: sRGB\n")
+	TEXT("1: BT.1886"),
+	ECVF_RenderThreadSafe);
+
 // Including the neutral one at index 0
 const uint32 GMaxLUTBlendCount = 5;
 
@@ -106,9 +124,14 @@ BEGIN_SHADER_PARAMETER_STRUCT(FACESTonemapShaderParameters, )
 	SHADER_PARAMETER(float, ACESCoefsHigh_4)
 	SHADER_PARAMETER(float, ACESSceneColorMultiplier)
 	SHADER_PARAMETER(float, ACESGamutCompression)
+	SHADER_PARAMETER(uint32, bScaleWhite)
+	SHADER_PARAMETER(uint32, bScaleWhiteSDR)
 	SHADER_PARAMETER_SRV(Texture2D<float>, ACESReachTable)
 	SHADER_PARAMETER_SRV(Texture2D<float>, ACESGamutTable)
 	SHADER_PARAMETER_SRV(Texture2D<float>, ACESGammaTable)
+	SHADER_PARAMETER_SRV(Texture2D<float>, ACESReachTableSDR)
+	SHADER_PARAMETER_SRV(Texture2D<float>, ACESGamutTableSDR)
+	SHADER_PARAMETER_SRV(Texture2D<float>, ACESGammaTableSDR)
 END_SHADER_PARAMETER_STRUCT()
 
 
@@ -118,6 +141,12 @@ BEGIN_SHADER_PARAMETER_STRUCT(FCombineLUTParameters, )
 	SHADER_PARAMETER_SCALAR_ARRAY(float, LUTWeights, [GMaxLUTBlendCount])
 	SHADER_PARAMETER_STRUCT_REF(FWorkingColorSpaceShaderParameters, WorkingColorSpace)
 	SHADER_PARAMETER_STRUCT_INCLUDE(FACESTonemapShaderParameters, ACESTonemapParameters)
+	SHADER_PARAMETER(FVector3f, LimitingRgbToXYZ0)
+	SHADER_PARAMETER(FVector3f, LimitingRgbToXYZ1)
+	SHADER_PARAMETER(FVector3f, LimitingRgbToXYZ2)
+	SHADER_PARAMETER(FVector3f, LimitingXYZToRgb0)
+	SHADER_PARAMETER(FVector3f, LimitingXYZToRgb1)
+	SHADER_PARAMETER(FVector3f, LimitingXYZToRgb2)
 	SHADER_PARAMETER(FVector4f, OverlayColor)
 	SHADER_PARAMETER(FVector3f, ColorScale)
 	SHADER_PARAMETER(FVector4f, ColorSaturation)
@@ -149,6 +178,8 @@ BEGIN_SHADER_PARAMETER_STRUCT(FCombineLUTParameters, )
 	SHADER_PARAMETER(float, BlueCorrection)
 	SHADER_PARAMETER(float, ExpandGamut)
 	SHADER_PARAMETER(float, ToneCurveAmount)
+	SHADER_PARAMETER(uint32, TonemappingMethod)
+	SHADER_PARAMETER(uint32, TonemappingSDRInvEOTF)
 	SHADER_PARAMETER(float, FilmSlope)
 	SHADER_PARAMETER(float, FilmToe)
 	SHADER_PARAMETER(float, FilmShoulder)
@@ -167,6 +198,33 @@ if(DestParameters != (ParamValue)) \
 	bOutHasChanged = true; \
 }
 
+uint32 DoesGamutWhitePointMatch(uint32 OutputGamut, const UE::Color::FColorSpace& ColorSpace)
+{
+	static_assert((uint32)EDisplayColorGamut::MAX == 5u, "Update logic if new values are added.");
+
+	bool bMatch = false;
+
+	// Seeing values from DXGI that look like (X / 1024), so tolerate 50% more than that.
+	constexpr double Tolerance = 0x3p-11;
+
+	switch (static_cast<EDisplayColorGamut>(OutputGamut))
+	{
+	case EDisplayColorGamut::sRGB_D65:
+	case EDisplayColorGamut::DCIP3_D65:
+	case EDisplayColorGamut::Rec2020_D65:
+		bMatch = ColorSpace.GetWhiteChromaticity().Equals(FVector2d(.3127, .329), Tolerance);
+		break;
+	case EDisplayColorGamut::ACES_D60:
+	case EDisplayColorGamut::ACEScg_D60:
+		bMatch = ColorSpace.GetWhiteChromaticity().Equals(FVector2d(.32168, .33767), Tolerance);
+		break;
+	default:
+		checkNoEntry();
+	}
+
+	return bMatch;
+}
+
 struct FCachedLUTSettings
 {
 	uint32 UniqueID = 0;
@@ -238,10 +296,15 @@ struct FCachedLUTSettings
 		UPDATE_CACHE_SETTINGS(Parameters.OutputDevice.OutputGamut, TonemapperOutputDeviceParameters.OutputGamut, bHasChanged);
 		UPDATE_CACHE_SETTINGS(Parameters.OutputDevice.OutputMaxLuminance, TonemapperOutputDeviceParameters.OutputMaxLuminance, bHasChanged);
 
-		// TODO: Plumb color space of the display.
-		static const UE::Color::FColorSpace LimitingColorSpaceHDR(UE::Color::EColorSpace::Rec2020);
-		static const UE::Color::FColorSpace LimitingColorSpaceSDR(UE::Color::EColorSpace::sRGB);
-		const UE::Color::FColorSpace& LimitingColorSpace = (TonemapperOutputDeviceParameters.OutputDevice == static_cast<uint32>(EDisplayOutputFormat::SDR_sRGB)) ? LimitingColorSpaceSDR : LimitingColorSpaceHDR;
+		const UE::Color::FColorSpace& LimitingColorSpace = ViewFamily.RenderTarget->GetLimitingColorSpace();
+		const FMatrix44d& LimitingRgbToXYZ = LimitingColorSpace.GetRgbToXYZ();
+		UPDATE_CACHE_SETTINGS(Parameters.LimitingRgbToXYZ0, FVector3f(LimitingRgbToXYZ.GetColumn(0)), bHasChanged);
+		UPDATE_CACHE_SETTINGS(Parameters.LimitingRgbToXYZ1, FVector3f(LimitingRgbToXYZ.GetColumn(1)), bHasChanged);
+		UPDATE_CACHE_SETTINGS(Parameters.LimitingRgbToXYZ2, FVector3f(LimitingRgbToXYZ.GetColumn(2)), bHasChanged);
+		const FMatrix44d& LimitingXYZToRgb = LimitingColorSpace.GetXYZToRgb();
+		UPDATE_CACHE_SETTINGS(Parameters.LimitingXYZToRgb0, FVector3f(LimitingXYZToRgb.GetColumn(0)), bHasChanged);
+		UPDATE_CACHE_SETTINGS(Parameters.LimitingXYZToRgb1, FVector3f(LimitingXYZToRgb.GetColumn(1)), bHasChanged);
+		UPDATE_CACHE_SETTINGS(Parameters.LimitingXYZToRgb2, FVector3f(LimitingXYZToRgb.GetColumn(2)), bHasChanged);
 
 		FACESTonemapParams TonemapperParams;
 		GetACESTonemapParameters(TonemapperParams);
@@ -253,6 +316,8 @@ struct FCachedLUTSettings
 		UPDATE_CACHE_SETTINGS(Parameters.ACESTonemapParameters.ACESCoefsHigh_4, TonemapperParams.ACESCoefsHigh_4, bHasChanged);
 		UPDATE_CACHE_SETTINGS(Parameters.ACESTonemapParameters.ACESSceneColorMultiplier, TonemapperParams.ACESSceneColorMultiplier, bHasChanged);
 		UPDATE_CACHE_SETTINGS(Parameters.ACESTonemapParameters.ACESGamutCompression, TonemapperParams.ACESGamutCompression, bHasChanged);
+		UPDATE_CACHE_SETTINGS(Parameters.ACESTonemapParameters.bScaleWhite, !DoesGamutWhitePointMatch(TonemapperOutputDeviceParameters.OutputGamut, LimitingColorSpace), bHasChanged);
+		UPDATE_CACHE_SETTINGS(Parameters.ACESTonemapParameters.bScaleWhiteSDR, !DoesGamutWhitePointMatch(TonemapperOutputDeviceParameters.OutputGamut, UE::Color::FColorSpace::GetWorking()), bHasChanged);
 		UE::Color::ACES::GetTransformResources(
 			GraphBuilder,
 			TonemapperOutputDeviceParameters.OutputMaxLuminance,
@@ -261,6 +326,14 @@ struct FCachedLUTSettings
 			Parameters.ACESTonemapParameters.ACESGamutTable,
 			Parameters.ACESTonemapParameters.ACESGammaTable
 		);
+		UE::Color::ACES::GetTransformResources(
+			GraphBuilder,
+			100.,
+			UE::Color::FColorSpace::GetWorking(),
+			Parameters.ACESTonemapParameters.ACESReachTableSDR,
+			Parameters.ACESTonemapParameters.ACESGamutTableSDR,
+			Parameters.ACESTonemapParameters.ACESGammaTableSDR
+		);
 
 		UPDATE_CACHE_SETTINGS(Parameters.ColorScale, FVector3f(View.ColorScale), bHasChanged);
 		UPDATE_CACHE_SETTINGS(Parameters.OverlayColor, FVector4f(View.OverlayColor), bHasChanged);
@@ -305,6 +378,20 @@ struct FCachedLUTSettings
 		UPDATE_CACHE_SETTINGS(Parameters.ExpandGamut, Settings.ExpandGamut, bHasChanged);
 		UPDATE_CACHE_SETTINGS(Parameters.ToneCurveAmount, Settings.ToneCurveAmount, bHasChanged);
 
+		const int32 TonemappingMethodOverride = CVarLUTTonemappingMethod.GetValueOnRenderThread();
+		const uint32 TonemappingMethod = (TonemappingMethodOverride >= 0) ?
+			static_cast<uint32>(TonemappingMethodOverride) :
+			(static_cast<uint32>(ViewFamily.RenderTarget->GetSceneHDREnabled() ?
+				Settings.TonemappingMethodHDR :
+				Settings.TonemappingMethodSDR));
+		UPDATE_CACHE_SETTINGS(Parameters.TonemappingMethod, TonemappingMethod, bHasChanged);
+
+		const int32 TonemappingSDRInvEOTFOverride = CVarLUTTonemappingSDRInvEOTF.GetValueOnRenderThread();
+		const uint32 TonemappingSDRInvEOTF = (TonemappingSDRInvEOTFOverride >= 0) ?
+			static_cast<uint32>(TonemappingSDRInvEOTFOverride) :
+			static_cast<uint32>(Settings.TonemappingMethodSDR == ETonemappingMethod::StandardACES);
+		UPDATE_CACHE_SETTINGS(Parameters.TonemappingSDRInvEOTF, TonemappingSDRInvEOTF, bHasChanged);
+
 		UPDATE_CACHE_SETTINGS(Parameters.FilmSlope, Settings.FilmSlope, bHasChanged);
 		UPDATE_CACHE_SETTINGS(Parameters.FilmToe, Settings.FilmToe, bHasChanged);
 		UPDATE_CACHE_SETTINGS(Parameters.FilmShoulder, Settings.FilmShoulder, bHasChanged);
diff --git a/Engine/Source/Runtime/Renderer/Private/PostProcess/PostProcessDeviceEncodingOnly.cpp b/Engine/Source/Runtime/Renderer/Private/PostProcess/PostProcessDeviceEncodingOnly.cpp
--- a/Engine/Source/Runtime/Renderer/Private/PostProcess/PostProcessDeviceEncodingOnly.cpp
+++ b/Engine/Source/Runtime/Renderer/Private/PostProcess/PostProcessDeviceEncodingOnly.cpp
@@ -82,8 +82,8 @@ FDeviceEncodingOnlyOutputDeviceParameters GetDeviceEncodingOnlyOutputDeviceParam
 	Parameters.InverseGamma = (FVector3f)InvDisplayGammaValue;
 	Parameters.OutputDevice = static_cast<uint32>(OutputDeviceValue);
 	Parameters.OutputGamut = static_cast<uint32>(Family.RenderTarget->GetDisplayColorGamut());
-	static const IConsoleVariable* const CVarHDROverrideOSMaxLuminance = IConsoleManager::Get().FindConsoleVariable(TEXT("r.HDR.OverrideOSMaxLuminance"));
-	Parameters.OutputMaxLuminance = CVarHDROverrideOSMaxLuminance->GetBool() ?
+	static const IConsoleVariable* const CVarHDRDisplayOverrideOSMaxLuminance = IConsoleManager::Get().FindConsoleVariable(TEXT("r.HDR.Display.OverrideOSMaxLuminance"));
+	Parameters.OutputMaxLuminance = CVarHDRDisplayOverrideOSMaxLuminance->GetBool() ?
 		(bIsHDR ? HDRGetDisplayMaximumLuminance() : 100.f) :
 		Family.RenderTarget->GetMaximumLuminanceInNits();
 	return Parameters;
@@ -98,6 +98,7 @@ BEGIN_SHADER_PARAMETER_STRUCT(FDeviceEncodingACESTonemapShaderParameters, )
 	SHADER_PARAMETER(float, ACESCoefsHigh_4)
 	SHADER_PARAMETER(float, ACESSceneColorMultiplier)
 	SHADER_PARAMETER(float, ACESGamutCompression)
+	SHADER_PARAMETER(uint32, bScaleWhite)
 	SHADER_PARAMETER_SRV(Texture2D<float>, ACESReachTable)
 	SHADER_PARAMETER_SRV(Texture2D<float>, ACESGamutTable)
 	SHADER_PARAMETER_SRV(Texture2D<float>, ACESGammaTable)
@@ -168,6 +169,8 @@ public:
 
 IMPLEMENT_GLOBAL_SHADER(FDeviceEncodingOnlyCS, "/Engine/Private/PostProcessDeviceEncodingOnly.usf", "MainCS", SF_Compute);
 
+extern uint32 DoesGamutWhitePointMatch(uint32 OutputGamut, const UE::Color::FColorSpace& ColorSpace);
+
 FScreenPassTexture AddDeviceEncodingOnlyPass(FRDGBuilder& GraphBuilder, const FViewInfo& View, const FDeviceEncodingOnlyInputs& Inputs)
 {
 	check(Inputs.SceneColor.IsValid());
@@ -231,12 +234,9 @@ FScreenPassTexture AddDeviceEncodingOnlyPass(FRDGBuilder& GraphBuilder, const FV
 	CommonParameters.Color = GetScreenPassTextureViewportParameters(SceneColorViewport);
 	CommonParameters.Output = GetScreenPassTextureViewportParameters(OutputViewport);
 
-	// TODO: Plumb color space of the display.
-	static const UE::Color::FColorSpace LimitingColorSpaceHDR(UE::Color::EColorSpace::Rec2020);
-	static const UE::Color::FColorSpace LimitingColorSpaceSDR(UE::Color::EColorSpace::sRGB);
-	const UE::Color::FColorSpace& LimitingColorSpace = (CommonParameters.OutputDevice.OutputDevice == static_cast<uint32>(EDisplayOutputFormat::SDR_sRGB)) ? LimitingColorSpaceSDR : LimitingColorSpaceHDR;
-
 	GetACESTonemapParameters(CommonParameters.ACESTonemapParameters);
+	const UE::Color::FColorSpace& LimitingColorSpace = ViewFamily.RenderTarget->GetLimitingColorSpace();
+	CommonParameters.ACESTonemapParameters.bScaleWhite = !DoesGamutWhitePointMatch(CommonParameters.OutputDevice.OutputGamut, LimitingColorSpace);
 	UE::Color::ACES::GetTransformResources(
 		GraphBuilder,
 		CommonParameters.OutputDevice.OutputMaxLuminance,
diff --git a/Engine/Source/Runtime/Renderer/Private/PostProcess/PostProcessTonemap.cpp b/Engine/Source/Runtime/Renderer/Private/PostProcess/PostProcessTonemap.cpp
--- a/Engine/Source/Runtime/Renderer/Private/PostProcess/PostProcessTonemap.cpp
+++ b/Engine/Source/Runtime/Renderer/Private/PostProcess/PostProcessTonemap.cpp
@@ -271,8 +271,8 @@ FTonemapperOutputDeviceParameters GetTonemapperOutputDeviceParameters(const FSce
 	Parameters.InverseGamma = InvDisplayGammaValue;
 	Parameters.OutputDevice = static_cast<uint32>(OutputDeviceValue);
 	Parameters.OutputGamut = static_cast<uint32>(Family.RenderTarget->GetDisplayColorGamut());
-	static const IConsoleVariable* const CVarHDROverrideOSMaxLuminance = IConsoleManager::Get().FindConsoleVariable(TEXT("r.HDR.OverrideOSMaxLuminance"));
-	Parameters.OutputMaxLuminance = CVarHDROverrideOSMaxLuminance->GetBool() ?
+	static const IConsoleVariable* const CVarHDRDisplayOverrideOSMaxLuminance = IConsoleManager::Get().FindConsoleVariable(TEXT("r.HDR.Display.OverrideOSMaxLuminance"));
+	Parameters.OutputMaxLuminance = CVarHDRDisplayOverrideOSMaxLuminance->GetBool() ?
 		(bIsHDR ? HDRGetDisplayMaximumLuminance() : 100.f) :
 		Family.RenderTarget->GetMaximumLuminanceInNits();
 	return Parameters;
