Fixed issue with RGBA channel mask post process shader that was causing fatal crashes when instanced stereo was enabled, caused by incorrect pixel shader parameters

diff --git a/Engine/Shaders/Private/PostProcessChannelMask.usf b/Engine/Shaders/Private/PostProcessChannelMask.usf
--- a/Engine/Shaders/Private/PostProcessChannelMask.usf
+++ b/Engine/Shaders/Private/PostProcessChannelMask.usf
@@ -15,7 +15,7 @@ float4 ChannelMask;
 uint bMaskChannels;
 uint bInvertAlphaChannelMask;
 uint bDisplayCheckerboard;
-FScreenTransform SvPositionToViewPos;
+FScreenTransform ScreenPosToViewPos;
 
 #define CheckerSize 16
 #define InvCheckerSize (1.0f / CheckerSize)
@@ -24,14 +24,14 @@ FScreenTransform SvPositionToViewPos;
 
 float3 GetCheckerboardColor(float2 PixelPos)
 {
-	float2 ViewPos = ApplyScreenTransform(PixelPos, SvPositionToViewPos);
+	float2 ViewPos = ApplyScreenTransform(PixelPos, ScreenPosToViewPos);
 	int2 CheckerCell = floor(ViewPos * float2(InvCheckerSize, InvCheckerSize) * View.ResolutionFractionAndInv.y);
 	return CheckerCell.x % 2 == CheckerCell.y % 2 ? CheckerWhite : CheckerGray;
 }
 
-void ChannelMask_MainPS(FScreenVertexOutput Input, out float4 OutColor : SV_Target0)
+void ChannelMask_MainPS(in noperspective float4 UVAndScreenPos : TEXCOORD0, out float4 OutColor : SV_Target0)
 {
-	float4 SceneColor = Texture2DSample(SceneColorTexture, SceneColorSampler, Input.UV);
+	float4 SceneColor = Texture2DSample(SceneColorTexture, SceneColorSampler, UVAndScreenPos.xy);
 
 	// Always output the original scene alpha, regardless of what the output rgb ends up being
 	OutColor.a = SceneColor.a;
@@ -44,7 +44,7 @@ void ChannelMask_MainPS(FScreenVertexOutput Input, out float4 OutColor : SV_Targ
 	SceneColor.a = bInvertAlphaChannelMask ? saturate(1.0 - SceneColor.a) : SceneColor.a;
 	
 	float3 Channel = bMaskChannels ? dot(SceneColor.rgba, ChannelMask) : SceneColor.rgb;
-	float3 CheckerboardColor = bDisplayCheckerboard ? GetCheckerboardColor(Input.Position.xy) : 0.0f;
+	float3 CheckerboardColor = bDisplayCheckerboard ? GetCheckerboardColor(UVAndScreenPos.zw) : 0.0f;
 
 	// Perform a naive premult alpha composite with the raw contents of the scene color's alpha.
 	// When not inverting the alpha, this is a technically incorrect blending, as Channel has been premultiplied with 1.0 - SceneColor.a
diff --git a/Engine/Source/Runtime/Renderer/Private/PostProcess/ChannelMask.cpp b/Engine/Source/Runtime/Renderer/Private/PostProcess/ChannelMask.cpp
--- a/Engine/Source/Runtime/Renderer/Private/PostProcess/ChannelMask.cpp
+++ b/Engine/Source/Runtime/Renderer/Private/PostProcess/ChannelMask.cpp
@@ -15,7 +15,7 @@ namespace ChannelMask
 		SHADER_PARAMETER(uint32, bMaskChannels)
 		SHADER_PARAMETER(uint32, bInvertAlphaChannelMask)
 		SHADER_PARAMETER(uint32, bDisplayCheckerboard)
-		SHADER_PARAMETER(FScreenTransform, SvPositionToViewPos)
+		SHADER_PARAMETER(FScreenTransform, ScreenPosToViewPos)
 		RENDER_TARGET_BINDING_SLOTS()
 	END_SHADER_PARAMETER_STRUCT()
 	
@@ -55,8 +55,8 @@ FScreenPassTexture ChannelMask::AddChannelMaskPass(FRDGBuilder& GraphBuilder, co
 	PassParameters->bMaskChannels = bMaskChannels;
 	PassParameters->bInvertAlphaChannelMask = View.Family->ChannelMaskParams.bInvertAlphaChannelMask;
 	PassParameters->bDisplayCheckerboard = View.Family->ChannelMaskParams.bDrawAlphaBlendedCheckerboard;
-	PassParameters->SvPositionToViewPos = (
-		FScreenTransform::ChangeTextureBasisFromTo(FScreenPassTextureViewport(Output), FScreenTransform::ETextureBasis::TexelPosition, FScreenTransform::ETextureBasis::ViewportUV) *
+	PassParameters->ScreenPosToViewPos = (
+		FScreenTransform::ChangeTextureBasisFromTo(FScreenPassTextureViewport(Output), FScreenTransform::ETextureBasis::ScreenPosition, FScreenTransform::ETextureBasis::ViewportUV) *
 		FScreenTransform::ChangeTextureBasisFromTo(View.ViewRect.Size(), View.ViewRect, FScreenTransform::ETextureBasis::ViewportUV, FScreenTransform::ETextureBasis::TexelPosition));
 	PassParameters->RenderTargets[0] = Output.GetRenderTargetBinding();
 	
