Substrate - fix for wrong two sided back diffuse color when using lumen reflection ray tracing lighting with Substrate GBufferFormat=0.

#rnx
#tests PC editor

diff --git a/Engine/Shaders/Private/RayTracing/RayTracingMaterialHitShaders.usf b/Engine/Shaders/Private/RayTracing/RayTracingMaterialHitShaders.usf
--- a/Engine/Shaders/Private/RayTracing/RayTracingMaterialHitShaders.usf
+++ b/Engine/Shaders/Private/RayTracing/RayTracingMaterialHitShaders.usf
@@ -499,6 +499,12 @@ void GetMaterialPayload(
 
 	Payload.Ior = IOR;
 	Payload.CustomData = CustomData;
+#if MATERIAL_SHADINGMODEL_TWOSIDED_FOLIAGE
+	if (Payload.ShadingModelID == SHADINGMODELID_TWOSIDED_FOLIAGE)
+	{
+		Payload.CustomData = half4(SubsurfaceData.rgb, Opacity);
+	}
+#endif
 
 #if !SUBSTRATE_ENABLED	// Not need to do that since the SubstrateExport already handle the CustomData all at once.
 #if MATERIAL_SHADINGMODEL_CLEAR_COAT
