Fixed voxel quantization causing voxel resolution to drop/increase, having side effect (running out of voxel pages).

diff --git a/Engine/Shaders/Private/HairStrands/HairStrandsVoxelPageCommon.ush b/Engine/Shaders/Private/HairStrands/HairStrandsVoxelPageCommon.ush
--- a/Engine/Shaders/Private/HairStrands/HairStrandsVoxelPageCommon.ush
+++ b/Engine/Shaders/Private/HairStrands/HairStrandsVoxelPageCommon.ush
@@ -86,15 +86,25 @@ struct FVirtualVoxelNodeDesc
 	bool    bIsValid;
 };
 
+uint PackUnorm11(float Value)
+{
+	return uint(saturate(Value) * 2047.0f);
+}
+
+float UnpackUnorm11(uint In)
+{
+	return float(In & 0x7FF) * (1.0f / 2047.0f);
+}
+
 uint PackVoxelWorldSize(float In)
 {
 	// Encode voxel world size [0..HAIR_MAX_VOXEL_SIZE] onto 10bits
 	const float NormalizedVoxelWorldSize = min(In, HAIR_MAX_VOXEL_SIZE) * (1.f/HAIR_MAX_VOXEL_SIZE);
-	return PackUnorm10(NormalizedVoxelWorldSize);
+	return PackUnorm11(NormalizedVoxelWorldSize);
 }
 float UnpackVoxelWorldSize(uint In)
 {
-	return UnpackUnorm10(In) * HAIR_MAX_VOXEL_SIZE;
+	return UnpackUnorm11(In) * HAIR_MAX_VOXEL_SIZE;
 }
 
 float QuantizeVoxelWorldSize(float In)
@@ -109,8 +119,8 @@ FVirtualVoxelNodeDesc UnpackVoxelNode(FPackedVirtualVoxelNodeDesc In, uint InPag
 	Out.TranslatedWorldMinAABB = In.TranslatedWorldMinAABB;
 	Out.PageIndexResolution = uint3(In.PackedPageIndexResolution & 0xFF, (In.PackedPageIndexResolution >> 8) & 0xFF, (In.PackedPageIndexResolution >> 16) & 0xFF);
 	Out.VirtualResolution = Out.PageIndexResolution * InPageResolution;
-	Out.PageIndexOffset = (In.PageIndexOffset_VoxelWorldSize & 0x3FFFFFu);
-	Out.VoxelWorldSize = UnpackVoxelWorldSize(In.PageIndexOffset_VoxelWorldSize >> 22u);
+	Out.PageIndexOffset = (In.PageIndexOffset_VoxelWorldSize & 0x1FFFFFu);
+	Out.VoxelWorldSize = UnpackVoxelWorldSize(In.PageIndexOffset_VoxelWorldSize >> 21u);
 	Out.bIsValid = all(Out.PageIndexResolution != 0);
 	return Out;
 }
@@ -125,8 +135,8 @@ FPackedVirtualVoxelNodeDesc PackVoxelNode(FVirtualVoxelNodeDesc In)
 		((In.PageIndexResolution.y & 0xFF)<<8) | 
 		((In.PageIndexResolution.z & 0xFF)<<16);
 	Out.PageIndexOffset_VoxelWorldSize =
-		(In.PageIndexOffset & 0x3FFFFFu) |
-		(PackVoxelWorldSize(In.VoxelWorldSize) << 22u);
+		(In.PageIndexOffset & 0x1FFFFFu) |
+		(PackVoxelWorldSize(In.VoxelWorldSize) << 21u);
 	return Out;
 }
 
