[DOF] Disable CONFIG_WAVE_BROADCAST_REDUCTION due to problems on RDNA 3 and 4.

diff --git a/Engine/Shaders/Private/DiaphragmDOF/DOFReduce.usf b/Engine/Shaders/Private/DiaphragmDOF/DOFReduce.usf
--- a/Engine/Shaders/Private/DiaphragmDOF/DOFReduce.usf
+++ b/Engine/Shaders/Private/DiaphragmDOF/DOFReduce.usf
@@ -42,8 +42,8 @@
 	#define CONFIG_GATHER_INPUT_LAYOUT (GATHER_INPUT_LAYOUT_RGB_COC)
 #endif
 
-#define CONFIG_WAVE_BROADCAST_REDUCTION (PLATFORM_SUPPORTS_WAVE_BROADCAST)
-
+//#define CONFIG_WAVE_BROADCAST_REDUCTION (PLATFORM_SUPPORTS_WAVE_BROADCAST)
+#define CONFIG_WAVE_BROADCAST_REDUCTION 0 //UE-308066
 
 // Configures the neighborhood analysis method to use for slight out of focus early out.
 #if COMPILER_SUPPORTS_WAVE_BIT_ORAND && (COMPILER_PSSL || XBOXONE_PROFILE)
