[DOF] Fix undefined behaviour in reduce pass around LDS access and group syncing, and remove path with PLATFORM_SUPPORTS_WAVE_BROADCAST that made faultly wavesize assumptions.

diff --git a/Engine/Shaders/Private/DiaphragmDOF/DOFReduce.usf b/Engine/Shaders/Private/DiaphragmDOF/DOFReduce.usf
--- a/Engine/Shaders/Private/DiaphragmDOF/DOFReduce.usf
+++ b/Engine/Shaders/Private/DiaphragmDOF/DOFReduce.usf
@@ -42,8 +42,6 @@
 	#define CONFIG_GATHER_INPUT_LAYOUT (GATHER_INPUT_LAYOUT_RGB_COC)
 #endif
 
-//#define CONFIG_WAVE_BROADCAST_REDUCTION (PLATFORM_SUPPORTS_WAVE_BROADCAST)
-#define CONFIG_WAVE_BROADCAST_REDUCTION 0 //UE-308066
 
 // Configures the neighborhood analysis method to use for slight out of focus early out.
 #if COMPILER_SUPPORTS_WAVE_BIT_ORAND && (COMPILER_PSSL || XBOXONE_PROFILE)
@@ -134,12 +132,23 @@ groupshared float4 GroupSharedArray[THREADGROUP_TOTALSIZE];
 
 //------------------------------------------------------- FUNCTIONS
 
-// Reduce using GroupSharedArray
-void ReduceOperator(FCocDownsampleParams DownsampleParams, uint OutId, uint ReduceBankSize)
+struct FReductionInput
 {
-	// Gather the 4 imput samples. Eaches samples are on banks of size <ReduceBankSize> for coherent LDS memory access.
 	float4 Colors[4];
 	float CocRadii[4];
+};
+
+struct FReductionOutput
+{
+	float4 Color;
+	float CocRadius;
+};
+
+FReductionInput ReadReductionInputFromLDS(uint OutId, uint ReduceBankSize)
+{
+	FReductionInput Result;
+
+	// Gather the 4 input samples. Eaches samples are on banks of size <ReduceBankSize> for coherent LDS memory access.
 	UNROLL
 	for (uint i = 0; i < 4; i++)
 	{
@@ -147,38 +156,49 @@ void ReduceOperator(FCocDownsampleParams DownsampleParams, uint OutId, uint Redu
 		
 		#if CONFIG_GATHER_INPUT_LAYOUT == GATHER_INPUT_LAYOUT_RGB_ALPHA_COC
 		{
-			Colors[i] = GroupSharedArray[InSharedId];
-			CocRadii[i] = GroupSharedArray2[InSharedId];
+			Result.Colors[i] = GroupSharedArray[InSharedId];
+			Result.CocRadii[i] = GroupSharedArray2[InSharedId];
 		}
 		#else
 		{
-			Colors[i] = GroupSharedArray[InSharedId];
-			CocRadii[i] = Colors[i].a;
+			Result.Colors[i] = GroupSharedArray[InSharedId];
+			Result.CocRadii[i] = Result.Colors[i].a;
 		}
 		#endif
 	}
 
-	// Downsample the 4 sample to one according to Coc.
-	float4 OutColor;
-	float OutCocRadius;
-	DownsampleSceneColorWithCoc(DownsampleParams, Colors, CocRadii, OutColor, OutCocRadius); 
-	
-	// Technically need a barrier, but in practice with a warp size >= 32, no need.
-	// GroupMemoryBarrierWithGroupSync();
-		
+	return Result;
+}
+
+void WriteReductionOutputToLDS(FReductionOutput Output, uint OutId)
+{
 	// Output to shared memory.
 	#if CONFIG_GATHER_INPUT_LAYOUT == GATHER_INPUT_LAYOUT_RGB_ALPHA_COC
 	{
-		GroupSharedArray[OutId] = OutColor;
-		GroupSharedArray2[OutId] = OutCocRadius;
+		GroupSharedArray[OutId] = Output.Color;
+		GroupSharedArray2[OutId] = Output.CocRadius;
 	}
 	#else
 	{
-		GroupSharedArray[OutId] = float4(OutColor.rgb, OutCocRadius);
+		GroupSharedArray[OutId] = float4(Output.Color.rgb, Output.CocRadius);
 	}
 	#endif
 }
 
+FReductionOutput ReduceOperator(FReductionInput Input, FCocDownsampleParams DownsampleParams, uint OutId)
+{
+	// Downsample the 4 sample to one according to Coc.
+	float4 OutColor;
+	float OutCocRadius;
+	DownsampleSceneColorWithCoc(DownsampleParams, Input.Colors, Input.CocRadii, OutColor, OutCocRadius); 
+
+	FReductionOutput Output;
+	Output.Color = OutColor;
+	Output.CocRadius = OutCocRadius;
+
+	return Output;
+}
+
 // Output buffer.
 void OutputMipLevel(const uint MipLevel, uint2 MipPixelPos, float4 Color, float CocRadius)
 {
@@ -652,91 +672,42 @@ void ReduceCS(
 	DownsampleParams.FrameExposureScale = FrameExposureScale;
 	DownsampleParams.bDoColorBasedWeighting = false;
 
-	#if CONFIG_WAVE_BROADCAST_REDUCTION
+	// Store initial values to LDS to prepare for first reduction round
 	{
-		UNROLL
-		for (int MipLevel = 1; MipLevel < DIM_REDUCE_MIP_COUNT; MipLevel++)
-		{
-			// Multiplier that needs to be applied on Coc computation to be resolution independent.
-			DownsampleParams.CocRadiusMultiplier = 0.5 / float(1u << (MipLevel - 1));
-
-			// Scale of the reduction within the 8x8 tile.
-			const uint ReductionScale = 1u << (MipLevel - 1);
-			
-			// Chooses output CoC across lanes.
-			float OutputCocRadius = DownsampleCoc_Wave8x8(CocRadius, ReductionScale);
-			
-			// Compute the sample's bilateral weight.
-			float SampleBilateralWeight = ComputeDownsamplingBilateralWeight(DownsampleParams, OutputCocRadius, CocRadius, GatheredColor.rgb);
-
-			// Weight sum the color across the lanes.
-			GatheredColor.r = Sum2x2WithinWave8x8(GatheredColor.r * SampleBilateralWeight, ReductionScale);
-			GatheredColor.g = Sum2x2WithinWave8x8(GatheredColor.g * SampleBilateralWeight, ReductionScale);
-			GatheredColor.b = Sum2x2WithinWave8x8(GatheredColor.b * SampleBilateralWeight, ReductionScale);
-			GatheredColor.a = Sum2x2WithinWave8x8(GatheredColor.a * SampleBilateralWeight, ReductionScale);
-			
-			// Compute normalising weight across lanes.
-			float InvTotalWeight = rcp(Sum2x2WithinWave8x8(SampleBilateralWeight, ReductionScale));
-			
-			// Normalize the color.
-			GatheredColor *= InvTotalWeight;
-			CocRadius = OutputCocRadius;
-
-			// Output mip map.
-			OutputPixelId = OutputPixelId >> 1;
-			BRANCH
-			if (((DispatchThreadId.x | DispatchThreadId.y) & ((1u << MipLevel) - 1)) == 0)
-			{
-				OutputMipLevel(MipLevel, OutputPixelId, GatheredColor, CocRadius);
-			}
-		}
+		FReductionOutput ReductionInit;
+		ReductionInit.Color = GatheredColor;
+		ReductionInit.CocRadius = CocRadius;
+		WriteReductionOutputToLDS(ReductionInit, GroupThreadIndex);
 	}
-	#else // !CONFIG_WAVE_BROADCAST_REDUCTION
+
+	// Output mip > 0 of the gathering pass.
+	UNROLL
+	for (int MipLevel = 1; MipLevel < DIM_REDUCE_MIP_COUNT; MipLevel++)
 	{
-		// Store color to LDS for reduction.
-		#if CONFIG_GATHER_INPUT_LAYOUT == GATHER_INPUT_LAYOUT_RGB_ALPHA_COC
-		{
-			GroupSharedArray[GroupThreadIndex] = GatheredColor;
-			GroupSharedArray2[GroupThreadIndex] = CocRadius;
-		}
-		#else
+		// Multiplier that needs to be applied on Coc computation to be resolution independent.
+		const uint TileSize = THREADGROUP_TILE_SIZE / (1u << MipLevel);
+		const uint ReduceBankSize = TileSize * TileSize;
+
+		DownsampleParams.CocRadiusMultiplier = 0.5 / float(1u << (MipLevel - 1));
+		
+		GroupMemoryBarrierWithGroupSync();
+
+		FReductionOutput ReductionOutput = (FReductionOutput)0;
+		if (GroupThreadIndex < ReduceBankSize)
 		{
-			GroupSharedArray[GroupThreadIndex] = float4(GatheredColor.rgb, CocRadius);
+			FReductionInput ReductionInput = ReadReductionInputFromLDS(GroupThreadIndex, ReduceBankSize);
+			ReductionOutput = ReduceOperator(ReductionInput, DownsampleParams, GroupThreadIndex);
 		}
-		#endif
+			
+		GroupMemoryBarrierWithGroupSync();
 
-		// Output mip > 0 of the gathering pass.
-		UNROLL
-		for (int MipLevel = 1; MipLevel < DIM_REDUCE_MIP_COUNT; MipLevel++)
+		if (GroupThreadIndex < ReduceBankSize)
 		{
-			// Multiplier that needs to be applied on Coc computation to be resolution independent.
-			const uint TileSize = THREADGROUP_TILE_SIZE / (1u << MipLevel);
-			const uint ReduceBankSize = TileSize * TileSize;
-		
-			// GroupSharedArray has been written before.
-			GroupMemoryBarrierWithGroupSync();
-		
-			if (GroupThreadIndex < ReduceBankSize)
-			{
-				DownsampleParams.CocRadiusMultiplier = 0.5 / float(1u << (MipLevel - 1));
+			WriteReductionOutputToLDS(ReductionOutput, GroupThreadIndex);
 
-				// Reduce.
-				ReduceOperator(DownsampleParams, GroupThreadIndex, ReduceBankSize);
-
-				// Read LDS (compiler are generally smart here to actualy keep value on VGPR.
-				#if CONFIG_GATHER_INPUT_LAYOUT == GATHER_INPUT_LAYOUT_RGB_ALPHA_COC
-					float4 Color = GroupSharedArray[GroupThreadIndex];
-					float OutCocRadius = GroupSharedArray2[GroupThreadIndex];
-				#else
-					float4 Color = GroupSharedArray[GroupThreadIndex];
-					float OutCocRadius = GroupSharedArray[GroupThreadIndex].a;
-				#endif
-			
-				// Output mip map.
-				OutputPixelId = OutputPixelId >> 1;
-				OutputMipLevel(MipLevel, OutputPixelId, Color, OutCocRadius);
-			}
+			// Output mip map.
+			OutputPixelId = OutputPixelId >> 1;
+			OutputMipLevel(MipLevel, OutputPixelId, ReductionOutput.Color, ReductionOutput.CocRadius);
 		}
 	}
-	#endif // !CONFIG_WAVE_BROADCAST_REDUCTION
 }
