Update SMAA to support mobile multi-view (MMV).
Also fix `toggledebugcamera` mode when using MMV on desktop via -es31 -vulkan.

diff --git a/Engine/Shaders/Private/SMAA/SMAA.ush b/Engine/Shaders/Private/SMAA/SMAA.ush
--- a/Engine/Shaders/Private/SMAA/SMAA.ush
+++ b/Engine/Shaders/Private/SMAA/SMAA.ush
@@ -588,6 +588,22 @@ SamplerState PointSampler;// { Filter = MIN_MAG_MIP_POINT; AddressU = Clamp; Add
 #define bool4 bvec4
 #endif
 
+#if MOBILE_MULTI_VIEW
+// MMV versions of the above defines which handle Texture2DArray color, edge, and blend textures and accept an array index
+#define SMAATexture2DMMV(tex) Texture2DArray tex
+#define SMAASamplePointMMV(tex, coord, index) tex.Sample(PointSampler, float3(coord, index))
+#define SMAASampleLevelZeroMMV(tex, coord, index) tex.SampleLevel(LinearSampler, float3(coord, index), 0)
+#define SMAASampleLevelZeroOffsetMMV(tex, coord, offset, index) tex.SampleLevel(LinearSampler, float3(coord, index), 0, offset)
+#define SMAASampleMMV(tex, coord, index) tex.Sample(LinearSampler, float3(coord, index))
+#else
+// When MMV is disabled, these match default behavior 
+#define SMAATexture2DMMV(tex) SMAATexture2D(tex)
+#define SMAASamplePointMMV(tex, coord, index) SMAASamplePoint(tex, coord)
+#define SMAASampleLevelZeroMMV(tex, coord, index) SMAASampleLevelZero(tex, coord)
+#define SMAASampleLevelZeroOffsetMMV(tex, coord, offset, index) SMAASampleLevelZeroOffset(tex, coord, offset)
+#define SMAASampleMMV(tex, coord, index) SMAASample(tex, coord)
+#endif
+
 #if !defined(SMAA_HLSL_3) && !defined(SMAA_HLSL_4) && !defined(SMAA_HLSL_4_1) && !defined(SMAA_GLSL_3) && !defined(SMAA_GLSL_4) && !defined(SMAA_CUSTOM_SL)
 #error you must define the shading language: SMAA_HLSL_*, SMAA_GLSL_* or SMAA_CUSTOM_SL
 #endif
@@ -690,13 +706,14 @@ void SMAANeighborhoodBlendingVS(float2 texcoord,
  */
 float2 SMAALumaEdgeDetectionPS(float2 texcoord,
                                float4 offset[3],
-                               SMAATexture2D(colorTex)
-                               #if SMAA_PREDICATION
-                               , SMAATexture2D(predicationTex)
+                               SMAATexture2DMMV(colorTex),
+                               #if SMAA_PREDICATION && !MOBILE_MULTI_VIEW
+                               SMAATexture2D(predicationTex)
                                #endif
+							   int ArrayIndex
                                ) {
     // Calculate the threshold:
-    #if SMAA_PREDICATION
+    #if SMAA_PREDICATION && !MOBILE_MULTI_VIEW
     float2 threshold = SMAACalculatePredicatedThreshold(texcoord, offset, SMAATexturePass2D(predicationTex));
     #else
     float2 threshold = float2(SMAA_THRESHOLD, SMAA_THRESHOLD);
@@ -704,10 +721,10 @@ float2 SMAALumaEdgeDetectionPS(float2 texcoord,
 
     // Calculate lumas:
     float3 weights = float3(0.2126, 0.7152, 0.0722);
-    float L = dot(SMAASamplePoint(colorTex, texcoord).rgb, weights);
+    float L = dot(SMAASamplePointMMV(colorTex, texcoord, ArrayIndex).rgb, weights);
 
-    float Lleft = dot(SMAASamplePoint(colorTex, offset[0].xy).rgb, weights);
-    float Ltop  = dot(SMAASamplePoint(colorTex, offset[0].zw).rgb, weights);
+    float Lleft = dot(SMAASamplePointMMV(colorTex, offset[0].xy, ArrayIndex).rgb, weights);
+    float Ltop  = dot(SMAASamplePointMMV(colorTex, offset[0].zw, ArrayIndex).rgb, weights);
 
     // We do the usual threshold:
     float4 delta;
@@ -719,16 +736,16 @@ float2 SMAALumaEdgeDetectionPS(float2 texcoord,
         discard;
 
     // Calculate right and bottom deltas:
-    float Lright = dot(SMAASamplePoint(colorTex, offset[1].xy).rgb, weights);
-    float Lbottom  = dot(SMAASamplePoint(colorTex, offset[1].zw).rgb, weights);
+    float Lright = dot(SMAASamplePointMMV(colorTex, offset[1].xy, ArrayIndex).rgb, weights);
+    float Lbottom  = dot(SMAASamplePointMMV(colorTex, offset[1].zw, ArrayIndex).rgb, weights);
     delta.zw = abs(L - float2(Lright, Lbottom));
 
     // Calculate the maximum delta in the direct neighborhood:
     float2 maxDelta = max(delta.xy, delta.zw);
 
     // Calculate left-left and top-top deltas:
-    float Lleftleft = dot(SMAASamplePoint(colorTex, offset[2].xy).rgb, weights);
-    float Ltoptop = dot(SMAASamplePoint(colorTex, offset[2].zw).rgb, weights);
+    float Lleftleft = dot(SMAASamplePointMMV(colorTex, offset[2].xy, ArrayIndex).rgb, weights);
+    float Ltoptop = dot(SMAASamplePointMMV(colorTex, offset[2].zw, ArrayIndex).rgb, weights);
     delta.zw = abs(float2(Lleft, Ltop) - float2(Lleftleft, Ltoptop));
 
     // Calculate the final maximum delta:
@@ -749,13 +766,14 @@ float2 SMAALumaEdgeDetectionPS(float2 texcoord,
  */
 float2 SMAAColorEdgeDetectionPS(float2 texcoord,
                                 float4 offset[3],
-                                SMAATexture2D(colorTex)
-                                #if SMAA_PREDICATION
-                                , SMAATexture2D(predicationTex)
+                                SMAATexture2DMMV(colorTex),
+                                #if SMAA_PREDICATION && !MOBILE_MULTI_VIEW
+                                SMAATexture2D(predicationTex),
                                 #endif
+								int ArrayIndex
                                 ) {
     // Calculate the threshold:
-    #if SMAA_PREDICATION
+    #if SMAA_PREDICATION && !MOBILE_MULTI_VIEW
     float2 threshold = SMAACalculatePredicatedThreshold(texcoord, offset, predicationTex);
     #else
     float2 threshold = float2(SMAA_THRESHOLD, SMAA_THRESHOLD);
@@ -763,13 +781,13 @@ float2 SMAAColorEdgeDetectionPS(float2 texcoord,
 
     // Calculate color deltas:
     float4 delta;
-    float3 C = SMAASamplePoint(colorTex, texcoord).rgb;
+    float3 C = SMAASamplePointMMV(colorTex, texcoord, ArrayIndex).rgb;
 
-    float3 Cleft = SMAASamplePoint(colorTex, offset[0].xy).rgb;
+    float3 Cleft = SMAASamplePointMMV(colorTex, offset[0].xy, ArrayIndex).rgb;
     float3 t = abs(C - Cleft);
     delta.x = max(max(t.r, t.g), t.b);
 
-    float3 Ctop  = SMAASamplePoint(colorTex, offset[0].zw).rgb;
+    float3 Ctop  = SMAASamplePointMMV(colorTex, offset[0].zw, ArrayIndex).rgb;
     t = abs(C - Ctop);
     delta.y = max(max(t.r, t.g), t.b);
 
@@ -781,11 +799,11 @@ float2 SMAAColorEdgeDetectionPS(float2 texcoord,
         discard;
 
     // Calculate right and bottom deltas:
-    float3 Cright = SMAASamplePoint(colorTex, offset[1].xy).rgb;
+    float3 Cright = SMAASamplePointMMV(colorTex, offset[1].xy, ArrayIndex).rgb;
     t = abs(C - Cright);
     delta.z = max(max(t.r, t.g), t.b);
 
-    float3 Cbottom  = SMAASamplePoint(colorTex, offset[1].zw).rgb;
+    float3 Cbottom  = SMAASamplePointMMV(colorTex, offset[1].zw, ArrayIndex).rgb;
     t = abs(C - Cbottom);
     delta.w = max(max(t.r, t.g), t.b);
 
@@ -793,11 +811,11 @@ float2 SMAAColorEdgeDetectionPS(float2 texcoord,
     float2 maxDelta = max(delta.xy, delta.zw);
 
     // Calculate left-left and top-top deltas:
-    float3 Cleftleft  = SMAASamplePoint(colorTex, offset[2].xy).rgb;
+    float3 Cleftleft  = SMAASamplePointMMV(colorTex, offset[2].xy, ArrayIndex).rgb;
     t = abs(C - Cleftleft);
     delta.z = max(max(t.r, t.g), t.b);
 
-    float3 Ctoptop = SMAASamplePoint(colorTex, offset[2].zw).rgb;
+    float3 Ctoptop = SMAASamplePointMMV(colorTex, offset[2].zw, ArrayIndex).rgb;
     t = abs(C - Ctoptop);
     delta.w = max(max(t.r, t.g), t.b);
 
@@ -861,19 +879,19 @@ float4 SMAADecodeDiagBilinearAccess(float4 e) {
 /**
  * These functions allows to perform diagonal pattern searches.
  */
-float2 SMAASearchDiag1(SMAATexture2D(edgesTex), float2 texcoord, float2 dir, out float2 e) {
+float2 SMAASearchDiag1(SMAATexture2DMMV(edgesTex), float2 texcoord, float2 dir, out float2 e, int ArrayIndex) {
     float4 coord = float4(texcoord, -1.0, 1.0);
     float3 t = float3(SMAA_RT_METRICS.xy, 1.0);
     while (coord.z < float(SMAA_MAX_SEARCH_STEPS_DIAG - 1) &&
            coord.w > 0.9) {
         coord.xyz = mad(t, float3(dir, 1.0), coord.xyz);
-        e = SMAASampleLevelZero(edgesTex, coord.xy).rg;
+        e = SMAASampleLevelZeroMMV(edgesTex, coord.xy, ArrayIndex).rg;
         coord.w = dot(e, float2(0.5, 0.5));
     }
     return coord.zw;
 }
 
-float2 SMAASearchDiag2(SMAATexture2D(edgesTex), float2 texcoord, float2 dir, out float2 e) {
+float2 SMAASearchDiag2(SMAATexture2DMMV(edgesTex), float2 texcoord, float2 dir, out float2 e, int ArrayIndex) {
     float4 coord = float4(texcoord, -1.0, 1.0);
     coord.x += 0.25 * SMAA_RT_METRICS.x; // See @SearchDiag2Optimization
     float3 t = float3(SMAA_RT_METRICS.xy, 1.0);
@@ -883,12 +901,12 @@ float2 SMAASearchDiag2(SMAATexture2D(edgesTex), float2 texcoord, float2 dir, out
 
         // @SearchDiag2Optimization
         // Fetch both edges at once using bilinear filtering:
-        e = SMAASampleLevelZero(edgesTex, coord.xy).rg;
+        e = SMAASampleLevelZeroMMV(edgesTex, coord.xy, ArrayIndex).rg;
         e = SMAADecodeDiagBilinearAccess(e);
 
         // Non-optimized version:
-        // e.g = SMAASampleLevelZero(edgesTex, coord.xy).g;
-        // e.r = SMAASampleLevelZeroOffset(edgesTex, coord.xy, int2(1, 0)).r;
+        // e.g = SMAASampleLevelZeroMMV(edgesTex, coord.xy, ArrayIndex).g;
+        // e.r = SMAASampleLevelZeroOffsetMMV(edgesTex, coord.xy, int2(1, 0), ArrayIndex).r;
 
         coord.w = dot(e, float2(0.5, 0.5));
     }
@@ -918,35 +936,35 @@ float2 SMAAAreaDiag(SMAATexture2D(areaTex), float2 dist, float2 e, float offset)
 /**
  * This searches for diagonal patterns and returns the corresponding weights.
  */
-float2 SMAACalculateDiagWeights(SMAATexture2D(edgesTex), SMAATexture2D(areaTex), float2 texcoord, float2 e, float4 subsampleIndices) {
+float2 SMAACalculateDiagWeights(SMAATexture2DMMV(edgesTex), SMAATexture2D(areaTex), float2 texcoord, float2 e, float4 subsampleIndices, int ArrayIndex) {
     float2 weights = float2(0.0, 0.0);
 
     // Search for the line ends:
     float4 d;
     float2 end;
     if (e.r > 0.0) {
-        d.xz = SMAASearchDiag1(SMAATexturePass2D(edgesTex), texcoord, float2(-1.0,  1.0), end);
+        d.xz = SMAASearchDiag1(SMAATexturePass2D(edgesTex), texcoord, float2(-1.0,  1.0), end, ArrayIndex);
         d.x += float(end.y > 0.9);
     } else
         d.xz = float2(0.0, 0.0);
-    d.yw = SMAASearchDiag1(SMAATexturePass2D(edgesTex), texcoord, float2(1.0, -1.0), end);
+    d.yw = SMAASearchDiag1(SMAATexturePass2D(edgesTex), texcoord, float2(1.0, -1.0), end, ArrayIndex);
 
     SMAA_BRANCH
     if (d.x + d.y > 2.0) { // d.x + d.y + 1 > 3
         // Fetch the crossing edges:
         float4 coords = mad(float4(-d.x + 0.25, d.x, d.y, -d.y - 0.25), SMAA_RT_METRICS.xyxy, texcoord.xyxy);
         float4 c;
-        c.xy = SMAASampleLevelZeroOffset(edgesTex, coords.xy, int2(-1,  0)).rg;
-        c.zw = SMAASampleLevelZeroOffset(edgesTex, coords.zw, int2( 1,  0)).rg;
+        c.xy = SMAASampleLevelZeroOffsetMMV(edgesTex, coords.xy, int2(-1,  0), ArrayIndex).rg;
+        c.zw = SMAASampleLevelZeroOffsetMMV(edgesTex, coords.zw, int2( 1,  0), ArrayIndex).rg;
         c.yxwz = SMAADecodeDiagBilinearAccess(c.xyzw);
 
         // Non-optimized version:
         // float4 coords = mad(float4(-d.x, d.x, d.y, -d.y), SMAA_RT_METRICS.xyxy, texcoord.xyxy);
         // float4 c;
-        // c.x = SMAASampleLevelZeroOffset(edgesTex, coords.xy, int2(-1,  0)).g;
-        // c.y = SMAASampleLevelZeroOffset(edgesTex, coords.xy, int2( 0,  0)).r;
-        // c.z = SMAASampleLevelZeroOffset(edgesTex, coords.zw, int2( 1,  0)).g;
-        // c.w = SMAASampleLevelZeroOffset(edgesTex, coords.zw, int2( 1, -1)).r;
+        // c.x = SMAASampleLevelZeroOffsetMMV(edgesTex, coords.xy, int2(-1,  0), ArrayIndex).g;
+        // c.y = SMAASampleLevelZeroOffsetMMV(edgesTex, coords.xy, int2( 0,  0), ArrayIndex).r;
+        // c.z = SMAASampleLevelZeroOffsetMMV(edgesTex, coords.zw, int2( 1,  0), ArrayIndex).g;
+        // c.w = SMAASampleLevelZeroOffsetMMV(edgesTex, coords.zw, int2( 1, -1), ArrayIndex).r;
 
         // Merge crossing edges at each side into a single value:
         float2 cc = mad(float2(2.0, 2.0), c.xz, c.yw);
@@ -959,9 +977,9 @@ float2 SMAACalculateDiagWeights(SMAATexture2D(edgesTex), SMAATexture2D(areaTex),
     }
 
     // Search for the line ends:
-    d.xz = SMAASearchDiag2(SMAATexturePass2D(edgesTex), texcoord, float2(-1.0, -1.0), end);
-    if (SMAASampleLevelZeroOffset(edgesTex, texcoord, int2(1, 0)).r > 0.0) {
-        d.yw = SMAASearchDiag2(SMAATexturePass2D(edgesTex), texcoord, float2(1.0, 1.0), end);
+    d.xz = SMAASearchDiag2(SMAATexturePass2D(edgesTex), texcoord, float2(-1.0, -1.0), end, ArrayIndex);
+    if (SMAASampleLevelZeroOffsetMMV(edgesTex, texcoord, int2(1, 0), ArrayIndex).r > 0.0) {
+        d.yw = SMAASearchDiag2(SMAATexturePass2D(edgesTex), texcoord, float2(1.0, 1.0), end, ArrayIndex);
         d.y += float(end.y > 0.9);
     } else
         d.yw = float2(0.0, 0.0);
@@ -971,9 +989,9 @@ float2 SMAACalculateDiagWeights(SMAATexture2D(edgesTex), SMAATexture2D(areaTex),
         // Fetch the crossing edges:
         float4 coords = mad(float4(-d.x, -d.x, d.y, d.y), SMAA_RT_METRICS.xyxy, texcoord.xyxy);
         float4 c;
-        c.x  = SMAASampleLevelZeroOffset(edgesTex, coords.xy, int2(-1,  0)).g;
-        c.y  = SMAASampleLevelZeroOffset(edgesTex, coords.xy, int2( 0, -1)).r;
-        c.zw = SMAASampleLevelZeroOffset(edgesTex, coords.zw, int2( 1,  0)).gr;
+        c.x  = SMAASampleLevelZeroOffsetMMV(edgesTex, coords.xy, int2(-1,  0), ArrayIndex).g;
+        c.y  = SMAASampleLevelZeroOffsetMMV(edgesTex, coords.xy, int2( 0, -1), ArrayIndex).r;
+        c.zw = SMAASampleLevelZeroOffsetMMV(edgesTex, coords.zw, int2( 1,  0), ArrayIndex).gr;
         float2 cc = mad(float2(2.0, 2.0), c.xz, c.yw);
 
         // Remove the crossing edge if we didn't found the end of the line:
@@ -1018,7 +1036,7 @@ float SMAASearchLength(SMAATexture2D(searchTex), float2 e, float offset) {
 /**
  * Horizontal/vertical search functions for the 2nd pass.
  */
-float SMAASearchXLeft(SMAATexture2D(edgesTex), SMAATexture2D(searchTex), float2 texcoord, float end) {
+float SMAASearchXLeft(SMAATexture2DMMV(edgesTex), SMAATexture2D(searchTex), float2 texcoord, float end, int ArrayIndex) {
     /**
      * @PSEUDO_GATHER4
      * This texcoord has been offset by (-0.25, -0.125) in the vertex shader to
@@ -1030,7 +1048,7 @@ float SMAASearchXLeft(SMAATexture2D(edgesTex), SMAATexture2D(searchTex), float2
     while (texcoord.x > end && 
            e.g > 0.8281 && // Is there some edge not activated?
            e.r == 0.0) { // Or is there a crossing edge that breaks the line?
-        e = SMAASampleLevelZero(edgesTex, texcoord).rg;
+        e = SMAASampleLevelZeroMMV(edgesTex, texcoord, ArrayIndex).rg;
         texcoord = mad(-float2(2.0, 0.0), SMAA_RT_METRICS.xy, texcoord);
     }
 
@@ -1050,36 +1068,36 @@ float SMAASearchXLeft(SMAATexture2D(edgesTex), SMAATexture2D(searchTex), float2
     // return mad(SMAA_RT_METRICS.x, offset, texcoord.x);
 }
 
-float SMAASearchXRight(SMAATexture2D(edgesTex), SMAATexture2D(searchTex), float2 texcoord, float end) {
+float SMAASearchXRight(SMAATexture2DMMV(edgesTex), SMAATexture2D(searchTex), float2 texcoord, float end, int ArrayIndex) {
     float2 e = float2(0.0, 1.0);
     while (texcoord.x < end && 
            e.g > 0.8281 && // Is there some edge not activated?
            e.r == 0.0) { // Or is there a crossing edge that breaks the line?
-        e = SMAASampleLevelZero(edgesTex, texcoord).rg;
+        e = SMAASampleLevelZeroMMV(edgesTex, texcoord, ArrayIndex).rg;
         texcoord = mad(float2(2.0, 0.0), SMAA_RT_METRICS.xy, texcoord);
     }
     float offset = mad(-(255.0 / 127.0), SMAASearchLength(SMAATexturePass2D(searchTex), e, 0.5), 3.25);
     return mad(-SMAA_RT_METRICS.x, offset, texcoord.x);
 }
 
-float SMAASearchYUp(SMAATexture2D(edgesTex), SMAATexture2D(searchTex), float2 texcoord, float end) {
+float SMAASearchYUp(SMAATexture2DMMV(edgesTex), SMAATexture2D(searchTex), float2 texcoord, float end, int ArrayIndex) {
     float2 e = float2(1.0, 0.0);
     while (texcoord.y > end && 
            e.r > 0.8281 && // Is there some edge not activated?
            e.g == 0.0) { // Or is there a crossing edge that breaks the line?
-        e = SMAASampleLevelZero(edgesTex, texcoord).rg;
+        e = SMAASampleLevelZeroMMV(edgesTex, texcoord, ArrayIndex).rg;
         texcoord = mad(-float2(0.0, 2.0), SMAA_RT_METRICS.xy, texcoord);
     }
     float offset = mad(-(255.0 / 127.0), SMAASearchLength(SMAATexturePass2D(searchTex), e.gr, 0.0), 3.25);
     return mad(SMAA_RT_METRICS.y, offset, texcoord.y);
 }
 
-float SMAASearchYDown(SMAATexture2D(edgesTex), SMAATexture2D(searchTex), float2 texcoord, float end) {
+float SMAASearchYDown(SMAATexture2DMMV(edgesTex), SMAATexture2D(searchTex), float2 texcoord, float end, int ArrayIndex) {
     float2 e = float2(1.0, 0.0);
     while (texcoord.y < end && 
            e.r > 0.8281 && // Is there some edge not activated?
            e.g == 0.0) { // Or is there a crossing edge that breaks the line?
-        e = SMAASampleLevelZero(edgesTex, texcoord).rg;
+        e = SMAASampleLevelZeroMMV(edgesTex, texcoord, ArrayIndex).rg;
         texcoord = mad(float2(0.0, 2.0), SMAA_RT_METRICS.xy, texcoord);
     }
     float offset = mad(-(255.0 / 127.0), SMAASearchLength(SMAATexturePass2D(searchTex), e.gr, 0.5), 3.25);
@@ -1107,7 +1125,7 @@ float2 SMAAArea(SMAATexture2D(areaTex), float2 dist, float e1, float e2, float o
 //-----------------------------------------------------------------------------
 // Corner Detection Functions
 
-void SMAADetectHorizontalCornerPattern(SMAATexture2D(edgesTex), inout float2 weights, float4 texcoord, float2 d) {
+void SMAADetectHorizontalCornerPattern(SMAATexture2DMMV(edgesTex), inout float2 weights, float4 texcoord, float2 d, int ArrayIndex) {
     #if !defined(SMAA_DISABLE_CORNER_DETECTION)
     float2 leftRight = step(d.xy, d.yx);
     float2 rounding = (1.0 - SMAA_CORNER_ROUNDING_NORM) * leftRight;
@@ -1115,16 +1133,16 @@ void SMAADetectHorizontalCornerPattern(SMAATexture2D(edgesTex), inout float2 wei
     rounding /= leftRight.x + leftRight.y; // Reduce blending for pixels in the center of a line.
 
     float2 factor = float2(1.0, 1.0);
-    factor.x -= rounding.x * SMAASampleLevelZeroOffset(edgesTex, texcoord.xy, int2(0,  1)).r;
-    factor.x -= rounding.y * SMAASampleLevelZeroOffset(edgesTex, texcoord.zw, int2(1,  1)).r;
-    factor.y -= rounding.x * SMAASampleLevelZeroOffset(edgesTex, texcoord.xy, int2(0, -2)).r;
-    factor.y -= rounding.y * SMAASampleLevelZeroOffset(edgesTex, texcoord.zw, int2(1, -2)).r;
+    factor.x -= rounding.x * SMAASampleLevelZeroOffsetMMV(edgesTex, texcoord.xy, int2(0,  1), ArrayIndex).r;
+    factor.x -= rounding.y * SMAASampleLevelZeroOffsetMMV(edgesTex, texcoord.zw, int2(1,  1), ArrayIndex).r;
+    factor.y -= rounding.x * SMAASampleLevelZeroOffsetMMV(edgesTex, texcoord.xy, int2(0, -2), ArrayIndex).r;
+    factor.y -= rounding.y * SMAASampleLevelZeroOffsetMMV(edgesTex, texcoord.zw, int2(1, -2), ArrayIndex).r;
 
     weights *= saturate(factor);
     #endif
 }
 
-void SMAADetectVerticalCornerPattern(SMAATexture2D(edgesTex), inout float2 weights, float4 texcoord, float2 d) {
+void SMAADetectVerticalCornerPattern(SMAATexture2DMMV(edgesTex), inout float2 weights, float4 texcoord, float2 d, int ArrayIndex) {
     #if !defined(SMAA_DISABLE_CORNER_DETECTION)
     float2 leftRight = step(d.xy, d.yx);
     float2 rounding = (1.0 - SMAA_CORNER_ROUNDING_NORM) * leftRight;
@@ -1132,10 +1150,10 @@ void SMAADetectVerticalCornerPattern(SMAATexture2D(edgesTex), inout float2 weigh
     rounding /= leftRight.x + leftRight.y;
 
     float2 factor = float2(1.0, 1.0);
-    factor.x -= rounding.x * SMAASampleLevelZeroOffset(edgesTex, texcoord.xy, int2( 1, 0)).g;
-    factor.x -= rounding.y * SMAASampleLevelZeroOffset(edgesTex, texcoord.zw, int2( 1, 1)).g;
-    factor.y -= rounding.x * SMAASampleLevelZeroOffset(edgesTex, texcoord.xy, int2(-2, 0)).g;
-    factor.y -= rounding.y * SMAASampleLevelZeroOffset(edgesTex, texcoord.zw, int2(-2, 1)).g;
+    factor.x -= rounding.x * SMAASampleLevelZeroOffsetMMV(edgesTex, texcoord.xy, int2( 1, 0), ArrayIndex).g;
+    factor.x -= rounding.y * SMAASampleLevelZeroOffsetMMV(edgesTex, texcoord.zw, int2( 1, 1), ArrayIndex).g;
+    factor.y -= rounding.x * SMAASampleLevelZeroOffsetMMV(edgesTex, texcoord.xy, int2(-2, 0), ArrayIndex).g;
+    factor.y -= rounding.y * SMAASampleLevelZeroOffsetMMV(edgesTex, texcoord.zw, int2(-2, 1), ArrayIndex).g;
 
     weights *= saturate(factor);
     #endif
@@ -1147,20 +1165,21 @@ void SMAADetectVerticalCornerPattern(SMAATexture2D(edgesTex), inout float2 weigh
 float4 SMAABlendingWeightCalculationPS(float2 texcoord,
                                        float2 pixcoord,
                                        float4 offset[3],
-                                       SMAATexture2D(edgesTex),
+                                       SMAATexture2DMMV(edgesTex),
                                        SMAATexture2D(areaTex),
                                        SMAATexture2D(searchTex),
-                                       float4 subsampleIndices) { // Just pass zero for SMAA 1x, see @SUBSAMPLE_INDICES.
+                                       float4 subsampleIndices, // Just pass zero for SMAA 1x, see @SUBSAMPLE_INDICES.
+									   int ArrayIndex) { 
     float4 weights = float4(0.0, 0.0, 0.0, 0.0);
 
-    float2 e = SMAASample(edgesTex, texcoord).rg;
+    float2 e = SMAASampleMMV(edgesTex, texcoord, ArrayIndex).rg;
 
     SMAA_BRANCH
     if (e.g > 0.0) { // Edge at north
         #if !defined(SMAA_DISABLE_DIAG_DETECTION)
         // Diagonals have both north and west edges, so searching for them in
         // one of the boundaries is enough.
-        weights.rg = SMAACalculateDiagWeights(SMAATexturePass2D(edgesTex), SMAATexturePass2D(areaTex), texcoord, e, subsampleIndices);
+        weights.rg = SMAACalculateDiagWeights(SMAATexturePass2D(edgesTex), SMAATexturePass2D(areaTex), texcoord, e, subsampleIndices, ArrayIndex);
 
         // We give priority to diagonals, so if we find a diagonal we skip 
         // horizontal/vertical processing.
@@ -1172,17 +1191,17 @@ float4 SMAABlendingWeightCalculationPS(float2 texcoord,
 
         // Find the distance to the left:
         float3 coords;
-        coords.x = SMAASearchXLeft(SMAATexturePass2D(edgesTex), SMAATexturePass2D(searchTex), offset[0].xy, offset[2].x);
+        coords.x = SMAASearchXLeft(SMAATexturePass2D(edgesTex), SMAATexturePass2D(searchTex), offset[0].xy, offset[2].x, ArrayIndex);
         coords.y = offset[1].y; // offset[1].y = texcoord.y - 0.25 * SMAA_RT_METRICS.y (@CROSSING_OFFSET)
         d.x = coords.x;
 
         // Now fetch the left crossing edges, two at a time using bilinear
         // filtering. Sampling at -0.25 (see @CROSSING_OFFSET) enables to
         // discern what value each edge has:
-        float e1 = SMAASampleLevelZero(edgesTex, coords.xy).r;
+        float e1 = SMAASampleLevelZeroMMV(edgesTex, coords.xy, ArrayIndex).r;
 
         // Find the distance to the right:
-        coords.z = SMAASearchXRight(SMAATexturePass2D(edgesTex), SMAATexturePass2D(searchTex), offset[0].zw, offset[2].y);
+        coords.z = SMAASearchXRight(SMAATexturePass2D(edgesTex), SMAATexturePass2D(searchTex), offset[0].zw, offset[2].y, ArrayIndex);
         d.y = coords.z;
 
         // We want the distances to be in pixel units (doing this here allow to
@@ -1194,7 +1213,7 @@ float4 SMAABlendingWeightCalculationPS(float2 texcoord,
         float2 sqrt_d = sqrt(d);
 
         // Fetch the right crossing edges:
-        float e2 = SMAASampleLevelZeroOffset(edgesTex, coords.zy, int2(1, 0)).r;
+        float e2 = SMAASampleLevelZeroOffsetMMV(edgesTex, coords.zy, int2(1, 0), ArrayIndex).r;
 
         // Ok, we know how this pattern looks like, now it is time for getting
         // the actual area:
@@ -1202,7 +1221,7 @@ float4 SMAABlendingWeightCalculationPS(float2 texcoord,
 
         // Fix corners:
         coords.y = texcoord.y;
-        SMAADetectHorizontalCornerPattern(SMAATexturePass2D(edgesTex), weights.rg, coords.xyzy, d);
+        SMAADetectHorizontalCornerPattern(SMAATexturePass2D(edgesTex), weights.rg, coords.xyzy, d, ArrayIndex);
 
         #if !defined(SMAA_DISABLE_DIAG_DETECTION)
         } else
@@ -1216,15 +1235,15 @@ float4 SMAABlendingWeightCalculationPS(float2 texcoord,
 
         // Find the distance to the top:
         float3 coords;
-        coords.y = SMAASearchYUp(SMAATexturePass2D(edgesTex), SMAATexturePass2D(searchTex), offset[1].xy, offset[2].z);
+        coords.y = SMAASearchYUp(SMAATexturePass2D(edgesTex), SMAATexturePass2D(searchTex), offset[1].xy, offset[2].z, ArrayIndex);
         coords.x = offset[0].x; // offset[1].x = texcoord.x - 0.25 * SMAA_RT_METRICS.x;
         d.x = coords.y;
 
         // Fetch the top crossing edges:
-        float e1 = SMAASampleLevelZero(edgesTex, coords.xy).g;
+        float e1 = SMAASampleLevelZeroMMV(edgesTex, coords.xy, ArrayIndex).g;
 
         // Find the distance to the bottom:
-        coords.z = SMAASearchYDown(SMAATexturePass2D(edgesTex), SMAATexturePass2D(searchTex), offset[1].zw, offset[2].w);
+        coords.z = SMAASearchYDown(SMAATexturePass2D(edgesTex), SMAATexturePass2D(searchTex), offset[1].zw, offset[2].w, ArrayIndex);
         d.y = coords.z;
 
         // We want the distances to be in pixel units:
@@ -1235,14 +1254,14 @@ float4 SMAABlendingWeightCalculationPS(float2 texcoord,
         float2 sqrt_d = sqrt(d);
 
         // Fetch the bottom crossing edges:
-        float e2 = SMAASampleLevelZeroOffset(edgesTex, coords.xz, int2(0, 1)).g;
+        float e2 = SMAASampleLevelZeroOffsetMMV(edgesTex, coords.xz, int2(0, 1), ArrayIndex).g;
 
         // Get the area for this direction:
         weights.ba = SMAAArea(SMAATexturePass2D(areaTex), sqrt_d, e1, e2, subsampleIndices.x);
 
         // Fix corners:
         coords.x = texcoord.x;
-        SMAADetectVerticalCornerPattern(SMAATexturePass2D(edgesTex), weights.ba, coords.xyxz, d);
+        SMAADetectVerticalCornerPattern(SMAATexturePass2D(edgesTex), weights.ba, coords.xyxz, d, ArrayIndex);
     }
 
     return weights;
@@ -1253,24 +1272,25 @@ float4 SMAABlendingWeightCalculationPS(float2 texcoord,
 
 float4 SMAANeighborhoodBlendingPS(float2 texcoord,
                                   float4 offset,
-                                  SMAATexture2D(colorTex),
-                                  SMAATexture2D(blendTex)
+                                  SMAATexture2DMMV(colorTex),
+                                  SMAATexture2DMMV(blendTex)
                                   #if SMAA_REPROJECTION
                                   , SMAATexture2D(velocityTex)
                                   #endif
+								  , int ArrayIndex
                                   ) {
     // Fetch the blending weights for current pixel:
     float4 a;
-    a.x = SMAASample(blendTex, offset.xy).a; // Right
-    a.y = SMAASample(blendTex, offset.zw).g; // Top
-    a.wz = SMAASample(blendTex, texcoord).xz; // Bottom / Left
+    a.x = SMAASampleMMV(blendTex, offset.xy, ArrayIndex).a; // Right
+    a.y = SMAASampleMMV(blendTex, offset.zw, ArrayIndex).g; // Top
+    a.wz = SMAASampleMMV(blendTex, texcoord, ArrayIndex).xz; // Bottom / Left
 
     // Is there any blending weight with a value greater than 0.0?
     SMAA_BRANCH
     if (dot(a, float4(1.0, 1.0, 1.0, 1.0)) < 1e-5) {
-        float4 color = SMAASampleLevelZero(colorTex, texcoord);
+        float4 color = SMAASampleLevelZeroMMV(colorTex, texcoord, ArrayIndex);
 
-        #if SMAA_REPROJECTION
+        #if SMAA_REPROJECTION && !MOBILE_MULTI_VIEW
         float2 velocity = SMAA_DECODE_VELOCITY(SMAASampleLevelZero(velocityTex, texcoord));
 
         // Pack velocity into the alpha channel:
@@ -1293,10 +1313,10 @@ float4 SMAANeighborhoodBlendingPS(float2 texcoord,
 
         // We exploit bilinear filtering to mix current pixel with the chosen
         // neighbor:
-        float4 color = blendingWeight.x * SMAASampleLevelZero(colorTex, blendingCoord.xy);
-        color += blendingWeight.y * SMAASampleLevelZero(colorTex, blendingCoord.zw);
+        float4 color = blendingWeight.x * SMAASampleLevelZeroMMV(colorTex, blendingCoord.xy, ArrayIndex);
+        color += blendingWeight.y * SMAASampleLevelZeroMMV(colorTex, blendingCoord.zw, ArrayIndex);
 
-        #if SMAA_REPROJECTION
+        #if SMAA_REPROJECTION && !MOBILE_MULTI_VIEW
         // Antialias velocity for proper reprojection in a later stage:
         float2 velocity = blendingWeight.x * SMAA_DECODE_VELOCITY(SMAASampleLevelZero(velocityTex, blendingCoord.xy));
         velocity += blendingWeight.y * SMAA_DECODE_VELOCITY(SMAASampleLevelZero(velocityTex, blendingCoord.zw));
diff --git a/Engine/Shaders/Private/SMAA/SMAABlendingWeightCalculationShader.usf b/Engine/Shaders/Private/SMAA/SMAABlendingWeightCalculationShader.usf
--- a/Engine/Shaders/Private/SMAA/SMAABlendingWeightCalculationShader.usf
+++ b/Engine/Shaders/Private/SMAA/SMAABlendingWeightCalculationShader.usf
@@ -21,9 +21,13 @@ Texture2D SearchTex;
 	void MainVS(
 		float2 InPosition : ATTRIBUTE0,
 		float2 InUV       : ATTRIBUTE1,
+		in FStereoVSInput StereoInput,
+		out FStereoVSOutput StereoOutput,
 		out FSMAABlendingWeightCalculationVertexOutput Output
 	)
 	{
+		StereoSetupVS(StereoInput, StereoOutput);
+
 		DrawRectangle( float4( InPosition, 0, 1 ), InUV, Output.Position, Output.UV);
 
 		SMAABlendingWeightCalculationVS(Output.UV, Output.Pixcoord, Output.Offset);
@@ -39,9 +43,12 @@ Texture2D SearchTex;
 
 	void MainPS(
 		FSMAABlendingWeightCalculationVertexOutput Input,
+		in FStereoPSInput StereoInput,
 		out float4 OutColor : SV_Target0)
 	{
-		OutColor = SMAABlendingWeightCalculationPS(Input.UV, Input.Pixcoord, Input.Offset, Input_Texture, AreaTex, SearchTex, 0);
+		int ArrayIndex = GetEyeIndex(StereoInput); // If using mobile multi-view (MMV), input edge texture is a texture array
+
+		OutColor = SMAABlendingWeightCalculationPS(Input.UV, Input.Pixcoord, Input.Offset, Input_Texture, AreaTex, SearchTex, 0, ArrayIndex);
 	}
 
 #endif // MainPS
diff --git a/Engine/Shaders/Private/SMAA/SMAACommon.ush b/Engine/Shaders/Private/SMAA/SMAACommon.ush
--- a/Engine/Shaders/Private/SMAA/SMAACommon.ush
+++ b/Engine/Shaders/Private/SMAA/SMAACommon.ush
@@ -5,7 +5,11 @@
 
 SCREEN_PASS_TEXTURE_VIEWPORT(Input)
 
+#if MOBILE_MULTI_VIEW
+Texture2DArray Input_Texture;
+#else
 Texture2D Input_Texture;
+#endif
 SamplerState Input_Sampler;
 
 float4 RTMetrics;
diff --git a/Engine/Shaders/Private/SMAA/SMAADebugVisualization.usf b/Engine/Shaders/Private/SMAA/SMAADebugVisualization.usf
--- a/Engine/Shaders/Private/SMAA/SMAADebugVisualization.usf
+++ b/Engine/Shaders/Private/SMAA/SMAADebugVisualization.usf
@@ -9,8 +9,13 @@ struct FSMAADebugVisualizationVertexOutput
 	float4 Position : SV_POSITION;
 };
 
+#if MOBILE_MULTI_VIEW
+Texture2DArray EdgeTex;
+Texture2DArray BlendTex;
+#else
 Texture2D EdgeTex;
 Texture2D BlendTex;
+#endif
 
 int DebugMode;
 
@@ -22,9 +27,13 @@ int DebugMode;
 	void MainVS(
 		float2 InPosition : ATTRIBUTE0,
 		float2 InUV       : ATTRIBUTE1,
+		in FStereoVSInput StereoInput,
+		out FStereoVSOutput StereoOutput,
 		out FSMAADebugVisualizationVertexOutput Output
 	)
 	{
+		StereoSetupVS(StereoInput, StereoOutput);
+
 		DrawRectangle( float4( InPosition, 0, 1 ), InUV, Output.Position, Output.UV);
 
 		SMAANeighborhoodBlendingVS(Output.UV, Output.Offset);
@@ -40,17 +49,26 @@ int DebugMode;
 
 	void MainPS(
 		FSMAADebugVisualizationVertexOutput Input,
+		in FStereoPSInput StereoInput,
 		out float4 OutColor : SV_Target0)
 	{
+		int ArrayIndex = GetEyeIndex(StereoInput);
+
 		float4 color = float4(0,0,0,0);
 
+	#if MOBILE_MULTI_VIEW
+		float3 IUV = float3(Input.UV, ArrayIndex);
+	#else
+		float2 IUV = Input.UV;
+	#endif
+
 		if (DebugMode == 1)
 		{
-			color = EdgeTex.Sample(LinearSampler, Input.UV);
+			color = EdgeTex.Sample(LinearSampler, IUV);
 		}
 		else if (DebugMode == 2)
 		{
-			color = BlendTex.Sample(LinearSampler, Input.UV);
+			color = BlendTex.Sample(LinearSampler, IUV);
 		}
 		else if (DebugMode == 3)
 		{
@@ -58,14 +76,19 @@ int DebugMode;
 	
 			// Fetch the blending weights for current pixel:
 			float4 a;
+		#if MOBILE_MULTI_VIEW
+			a.x = BlendTex.Sample(LinearSampler, float3(Input.Offset.xy, ArrayIndex)).a; // Right
+			a.y = BlendTex.Sample(LinearSampler, float3(Input.Offset.zw, ArrayIndex)).g; // Top
+		#else
 			a.x = BlendTex.Sample(LinearSampler, Input.Offset.xy).a; // Right
 			a.y = BlendTex.Sample(LinearSampler, Input.Offset.zw).g; // Top
-			a.wz = BlendTex.Sample(LinearSampler, Input.UV).xz; // Bottom / Left
+		#endif
+			a.wz = BlendTex.Sample(LinearSampler, IUV).xz; // Bottom / Left
 
 			// Is there any blending weight with a value greater than 0.0?
 			if (dot(a, float4(1.0, 1.0, 1.0, 1.0)) < 1e-5)
 			{
-				color = Input_Texture.SampleLevel(LinearSampler, Input.UV, 0);
+				color = Input_Texture.SampleLevel(LinearSampler, IUV, 0);
 			}
 			else
 			{
diff --git a/Engine/Shaders/Private/SMAA/SMAAEdgeDetectionShader.usf b/Engine/Shaders/Private/SMAA/SMAAEdgeDetectionShader.usf
--- a/Engine/Shaders/Private/SMAA/SMAAEdgeDetectionShader.usf
+++ b/Engine/Shaders/Private/SMAA/SMAAEdgeDetectionShader.usf
@@ -19,9 +19,13 @@ int EdgeMode;
 	void MainVS(
 		float2 InPosition : ATTRIBUTE0,
 		float2 InUV       : ATTRIBUTE1,
+		in FStereoVSInput StereoInput,
+		out FStereoVSOutput StereoOutput,
 		out FSMAAEdgeDetectionScreenVertexOutput Output
 	)
 	{
+		StereoSetupVS(StereoInput, StereoOutput);
+
 		DrawRectangle( float4( InPosition, 0, 1 ), InUV, Output.Position, Output.UV);
 
 		SMAAEdgeDetectionVS(Output.UV, Output.Offset);
@@ -37,16 +41,19 @@ int EdgeMode;
 
 	void MainPS(
 		FSMAAEdgeDetectionScreenVertexOutput Input,
+		in FStereoPSInput StereoInput,
 		out float2 OutColor : SV_Target0)
 	{
+		int ArrayIndex = GetEyeIndex(StereoInput); // If using mobile multi-view (MMV), input color is a texture array
+
 		BRANCH
 		if (EdgeMode == 0) // Color
 		{
-			OutColor.rg = SMAAColorEdgeDetectionPS(Input.UV, Input.Offset, Input_Texture);
+			OutColor.rg = SMAAColorEdgeDetectionPS(Input.UV, Input.Offset, Input_Texture, ArrayIndex);
 		}
 		else // if (EdgeMode == 1) // Luma (uses Linear Color as input texture)
 		{
-			OutColor.rg = SMAALumaEdgeDetectionPS(Input.UV, Input.Offset, Input_Texture);
+			OutColor.rg = SMAALumaEdgeDetectionPS(Input.UV, Input.Offset, Input_Texture, ArrayIndex);
 		}
 	}
 
diff --git a/Engine/Shaders/Private/SMAA/SMAANeighborhoodBlendingShader.usf b/Engine/Shaders/Private/SMAA/SMAANeighborhoodBlendingShader.usf
--- a/Engine/Shaders/Private/SMAA/SMAANeighborhoodBlendingShader.usf
+++ b/Engine/Shaders/Private/SMAA/SMAANeighborhoodBlendingShader.usf
@@ -9,7 +9,11 @@ struct FSMAANeighborhoodBlendingScreenVertexOutput
 	float4 Position : SV_POSITION;
 };
 
+#if MOBILE_MULTI_VIEW
+Texture2DArray BlendTex;
+#else
 Texture2D BlendTex;
+#endif
 
 /*-----------------------------------------------------------------------------
 	vertex shader
@@ -19,9 +23,13 @@ Texture2D BlendTex;
 	void MainVS(
 		float2 InPosition : ATTRIBUTE0,
 		float2 InUV       : ATTRIBUTE1,
+		in FStereoVSInput StereoInput,
+		out FStereoVSOutput StereoOutput,
 		out FSMAANeighborhoodBlendingScreenVertexOutput Output
 	)
 	{
+		StereoSetupVS(StereoInput, StereoOutput);
+
 		DrawRectangle( float4( InPosition, 0, 1 ), InUV, Output.Position, Output.UV);
 
 		SMAANeighborhoodBlendingVS(Output.UV, Output.Offset);
@@ -37,9 +45,12 @@ Texture2D BlendTex;
 
 	void MainPS(
 		FSMAANeighborhoodBlendingScreenVertexOutput Input,
+		in FStereoPSInput StereoInput,
 		out float4 OutColor : SV_Target0)
 	{
-		OutColor = SMAANeighborhoodBlendingPS(Input.UV, Input.Offset, Input_Texture, BlendTex);
+		int ArrayIndex = GetEyeIndex(StereoInput); // If using mobile multi-view (MMV), input blend texture is a texture array
+
+		OutColor = SMAANeighborhoodBlendingPS(Input.UV, Input.Offset, Input_Texture, BlendTex, ArrayIndex);
 
 	#if DIM_ALPHA_CHANNEL == 0
 		OutColor.a = 1.0;
diff --git a/Engine/Source/Runtime/Renderer/Private/PostProcess/PostProcessTonemap.cpp b/Engine/Source/Runtime/Renderer/Private/PostProcess/PostProcessTonemap.cpp
--- a/Engine/Source/Runtime/Renderer/Private/PostProcess/PostProcessTonemap.cpp
+++ b/Engine/Source/Runtime/Renderer/Private/PostProcess/PostProcessTonemap.cpp
@@ -595,11 +595,13 @@ FScreenPassTexture AddTonemapPass(FRDGBuilder& GraphBuilder, const FViewInfo& Vi
 
 	if (!Output.IsValid())
 	{
-		FRDGTextureDesc OutputDesc = FRDGTextureDesc::Create2D(
-			SceneColorViewport.Extent,
-			Inputs.SceneColor.TextureSRV->Desc.Texture->Desc.Format,
-			FClearValueBinding(FLinearColor(0, 0, 0, 0)),
-			GFastVRamConfig.Tonemap | TexCreate_ShaderResource | TexCreate_RenderTargetable | (View.bUseComputePasses ? TexCreate_UAV : TexCreate_None));;
+		const bool bUseTextureArrays = View.bIsMobileMultiViewEnabled && !View.bUseComputePasses; // MMV not currently supported for compute passes
+		const EPixelFormat OutputFormat = Inputs.SceneColor.TextureSRV->Desc.Texture->Desc.Format;
+		const ETextureCreateFlags OutputFlags = GFastVRamConfig.Tonemap | TexCreate_ShaderResource | TexCreate_RenderTargetable | (View.bUseComputePasses ? TexCreate_UAV : TexCreate_None);
+
+		FRDGTextureDesc OutputDesc = bUseTextureArrays ?
+			FRDGTextureDesc::Create2DArray(SceneColorViewport.Extent, OutputFormat, FClearValueBinding(FLinearColor(0, 0, 0, 0)), OutputFlags, 2) :
+			FRDGTextureDesc::Create2D(SceneColorViewport.Extent, OutputFormat, FClearValueBinding(FLinearColor(0, 0, 0, 0)), OutputFlags);
 
 		const FTonemapperOutputDeviceParameters OutputDeviceParameters = GetTonemapperOutputDeviceParameters(*View.Family);
 		const EDisplayOutputFormat OutputDevice = static_cast<EDisplayOutputFormat>(OutputDeviceParameters.OutputDevice);
diff --git a/Engine/Source/Runtime/Renderer/Private/PostProcess/PostProcessing.cpp b/Engine/Source/Runtime/Renderer/Private/PostProcess/PostProcessing.cpp
--- a/Engine/Source/Runtime/Renderer/Private/PostProcess/PostProcessing.cpp
+++ b/Engine/Source/Runtime/Renderer/Private/PostProcess/PostProcessing.cpp
@@ -2641,7 +2641,8 @@ void AddMobilePostProcessingPasses(FRDGBuilder& GraphBuilder, FScene* Scene, con
 	PassSequence.SetEnabled(EPass::EditorPrimitive, false);
 #endif
 
-	PassSequence.SetEnabled(EPass::DebugPrimitive, IsDebugPrimitivePassEnabled(View));
+	// Debug primitive compositing does not currently support mobile multi-view
+	PassSequence.SetEnabled(EPass::DebugPrimitive, IsDebugPrimitivePassEnabled(View) && !View.bIsMobileMultiViewEnabled);
 	PassSequence.SetEnabled(EPass::PrimaryUpscale, bShouldPrimaryUpscale);
 	// Mobile is different from the deferred that the secondary upscale logic only runs when the interface is available
 	// While the deferred itself has a fallback when secondary upscaler interface is not present.
@@ -3315,7 +3316,9 @@ void AddMobilePostProcessingPasses(FRDGBuilder& GraphBuilder, FScene* Scene, con
 			PassSequence.AcceptOverrideIfLastPass(EPass::SMAA, PassInputs.OverrideOutput);
 			PassInputs.SceneColor = SceneColor;
 			PassInputs.Quality = GetSMAAQuality();
-			PassInputs.SceneColorBeforeTonemap = FScreenPassTexture::CopyFromSlice(GraphBuilder, SceneColorBeforeTonemapSlice);
+
+			const uint32 MultiViewCount = View.bIsMobileMultiViewEnabled ? 2 : 0; // SMAA can support MMV
+			PassInputs.SceneColorBeforeTonemap = FScreenPassTexture::CopyFromSlice(GraphBuilder, SceneColorBeforeTonemapSlice, FScreenPassTexture(), MultiViewCount);
 
 			SceneColor = AddSMAAPasses(GraphBuilder, View, PassInputs);
 		}
diff --git a/Engine/Source/Runtime/Renderer/Private/PostProcess/SubpixelMorphologicalAA.cpp b/Engine/Source/Runtime/Renderer/Private/PostProcess/SubpixelMorphologicalAA.cpp
--- a/Engine/Source/Runtime/Renderer/Private/PostProcess/SubpixelMorphologicalAA.cpp
+++ b/Engine/Source/Runtime/Renderer/Private/PostProcess/SubpixelMorphologicalAA.cpp
@@ -198,6 +198,7 @@ void AddSMAAEdgeDetectionPass(FRDGBuilder& GraphBuilder, const FSceneView& InSce
 		ERenderTargetLoadAction::EClear,
 		ERenderTargetLoadAction::EClear,
 		FExclusiveDepthStencil::DepthWrite_StencilWrite);
+	PassParameters->RenderTargets.MultiViewCount = View.bIsMobileMultiViewEnabled ? 2 : 0;
 	PassParameters->Input = GetScreenPassTextureInput(SceneColorTexture, PointClampSampler);
 	PassParameters->RTMetrics = RTMetrics;
 	PassParameters->PointSampler = PointClampSampler;
@@ -246,6 +247,7 @@ void AddSMAABlendingWeightCalculationPass(FRDGBuilder& GraphBuilder, const FScen
 		ERenderTargetLoadAction::EClear,
 		ERenderTargetLoadAction::ELoad,
 		FExclusiveDepthStencil::DepthWrite_StencilRead);
+	PassParameters->RenderTargets.MultiViewCount = View.bIsMobileMultiViewEnabled ? 2 : 0;
 	PassParameters->Input = GetScreenPassTextureInput(FScreenPassTexture(EdgeTexture), BilinearClampSampler);
 	PassParameters->RTMetrics = RTMetrics;
 	PassParameters->PointSampler = TStaticSamplerState<SF_Point, AM_Clamp, AM_Clamp, AM_Clamp>::GetRHI();
@@ -297,6 +299,7 @@ FScreenPassTexture AddSMAANeighborhoodBlendingPass(FRDGBuilder& GraphBuilder, co
 
 	FSMAANeighborhoodBlendingPS::FParameters* PassParameters = GraphBuilder.AllocParameters<FSMAANeighborhoodBlendingPS::FParameters>();
 	PassParameters->RenderTargets[0] = Output.GetRenderTargetBinding();
+	PassParameters->RenderTargets.MultiViewCount = View.bIsMobileMultiViewEnabled ? 2 : 0;
 	PassParameters->Input = GetScreenPassTextureInput(InputSceneTexture, BilinearClampSampler);
 	PassParameters->RTMetrics = RTMetrics;
 	PassParameters->BlendTex = InputBlendTex;
@@ -334,6 +337,7 @@ FScreenPassTexture AddSMAADebugVisualizationPass(FRDGBuilder& GraphBuilder, cons
 
 	FSMAADebugVisualizationPS::FParameters* PassParameters = GraphBuilder.AllocParameters<FSMAADebugVisualizationPS::FParameters>();
 	PassParameters->RenderTargets[0] = Output.GetRenderTargetBinding();
+	PassParameters->RenderTargets.MultiViewCount = View.bIsMobileMultiViewEnabled ? 2 : 0;
 	PassParameters->Input = GetScreenPassTextureInput(InputSceneTexture, BilinearClampSampler);
 	PassParameters->RTMetrics = RTMetrics;
 	PassParameters->EdgeTex = InputEdgeTex;
@@ -384,31 +388,19 @@ FScreenPassTexture AddSMAAPasses(FRDGBuilder& GraphBuilder, const FSceneView& Vi
 
 	FVector4f RTMetrics = FVector4f(1.0f / TexSize.X, 1.0f / TexSize.Y, TexSize.X, TexSize.Y);
 
-	FRDGTextureDesc EdgeTexDesc = FRDGTextureDesc::Create2D(
-		TexSize,
-		PF_R8G8,
-		FClearValueBinding(FLinearColor::Transparent),
-		TexCreate_RenderTargetable | TexCreate_ShaderResource
-	);
-
+	FRDGTextureDesc EdgeTexDesc = View.bIsMobileMultiViewEnabled ?
+	FRDGTextureDesc::Create2DArray(TexSize, PF_R8G8, FClearValueBinding(FLinearColor::Transparent), TexCreate_RenderTargetable | TexCreate_ShaderResource, 2) :
+	FRDGTextureDesc::Create2D(TexSize, PF_R8G8, FClearValueBinding(FLinearColor::Transparent), TexCreate_RenderTargetable | TexCreate_ShaderResource);
 	FRDGTextureRef EdgeTexture = GraphBuilder.CreateTexture(EdgeTexDesc, TEXT("SMAAEdgeTexture"));
 
-	FRDGTextureDesc BlendTexDesc = FRDGTextureDesc::Create2D(
-		TexSize,
-		PF_R8G8B8A8,
-		FClearValueBinding(FLinearColor::Transparent),
-		TexCreate_RenderTargetable | TexCreate_ShaderResource
-	);
-
+	FRDGTextureDesc BlendTexDesc = View.bIsMobileMultiViewEnabled ?
+	FRDGTextureDesc::Create2DArray(TexSize, PF_R8G8B8A8, FClearValueBinding(FLinearColor::Transparent), TexCreate_RenderTargetable | TexCreate_ShaderResource, 2) :
+	FRDGTextureDesc::Create2D(TexSize, PF_R8G8B8A8, FClearValueBinding(FLinearColor::Transparent), TexCreate_RenderTargetable | TexCreate_ShaderResource);
 	FRDGTextureRef BlendTexture = GraphBuilder.CreateTexture(BlendTexDesc, TEXT("SMAABlendTexture"));
 
-	FRDGTextureDesc DepthStencilDesc = FRDGTextureDesc::Create2D(
-		TexSize,
-		PF_DepthStencil,
-		FClearValueBinding(0.0f, 0),
-		TexCreate_DepthStencilTargetable | TexCreate_ShaderResource
-	);
-
+	FRDGTextureDesc DepthStencilDesc = View.bIsMobileMultiViewEnabled ? 
+	FRDGTextureDesc::Create2DArray(TexSize, PF_DepthStencil, FClearValueBinding(0.0f, 0), TexCreate_DepthStencilTargetable | TexCreate_ShaderResource, 2) :
+	FRDGTextureDesc::Create2D(TexSize, PF_DepthStencil, FClearValueBinding(0.0f, 0), TexCreate_DepthStencilTargetable | TexCreate_ShaderResource);
 	FRDGTextureRef DepthStencilTexture = GraphBuilder.CreateTexture(DepthStencilDesc, TEXT("SMAAStencilTexture"));
 
 	AddSMAAEdgeDetectionPass(GraphBuilder, View, Inputs.SceneColorBeforeTonemap, EdgeTexture, DepthStencilTexture, Inputs.Quality, RTMetrics);
diff --git a/Engine/Source/Runtime/Renderer/Private/ScreenPass.cpp b/Engine/Source/Runtime/Renderer/Private/ScreenPass.cpp
--- a/Engine/Source/Runtime/Renderer/Private/ScreenPass.cpp
+++ b/Engine/Source/Runtime/Renderer/Private/ScreenPass.cpp
@@ -59,7 +59,7 @@ FRDGTextureRef TryCreateViewFamilyDepthTexture(FRDGBuilder& GraphBuilder, const
 }
 
 // static
-FScreenPassTexture FScreenPassTexture::CopyFromSlice(FRDGBuilder& GraphBuilder, const FScreenPassTextureSlice& ScreenTextureSlice, FScreenPassTexture OverrideOutput)
+FScreenPassTexture FScreenPassTexture::CopyFromSlice(FRDGBuilder& GraphBuilder, const FScreenPassTextureSlice& ScreenTextureSlice, FScreenPassTexture OverrideOutput, uint32 MultiViewCount /*= 0*/)
 {
 	FRDGTextureSRV* InputTextureSRV = ScreenTextureSlice.TextureSRV;
 
@@ -83,9 +83,20 @@ FScreenPassTexture FScreenPassTexture::CopyFromSlice(FRDGBuilder& GraphBuilder,
 	if (!OutputTexture)
 	{
 		FRDGTextureDesc Desc = InputTexture->Desc;
-		Desc.Dimension = ETextureDimension::Texture2D;
-		Desc.ArraySize = 1;
 
+		if (MultiViewCount > 0)
+		{
+			check(Desc.Dimension == ETextureDimension::Texture2DArray);
+			check(Desc.ArraySize >= MultiViewCount);
+
+			Desc.ArraySize = MultiViewCount;
+		}
+		else
+		{
+			Desc.Dimension = ETextureDimension::Texture2D;
+			Desc.ArraySize = 1;
+		}
+		
 		// If a pass uses blending to write to this post process texture, it needs to support being a render target, so make sure this flag is included.
 		// Most post processing uses SceneColor or its FRDGTextureDesc, and SceneColor already has the RenderTargetable flag set, but TSR (the input to
 		// the "Before Bloom" stage) writes to texture slices using compute, and its output doesn't have this flag.
@@ -102,6 +113,11 @@ FScreenPassTexture FScreenPassTexture::CopyFromSlice(FRDGBuilder& GraphBuilder,
 	CopyInfo.SourcePosition = CopyInfo.DestPosition = FIntVector(ScreenTextureSlice.ViewRect.Min.X, ScreenTextureSlice.ViewRect.Min.Y, 0);
 	CopyInfo.Size = FIntVector(ViewSize.X, ViewSize.Y, 1);
 
+	if (MultiViewCount > 0)
+	{
+		CopyInfo.NumSlices = MultiViewCount;
+	}
+	
 	AddCopyTexturePass(
 		GraphBuilder,
 		InputTexture,
diff --git a/Engine/Source/Runtime/Renderer/Public/ScreenPass.h b/Engine/Source/Runtime/Renderer/Public/ScreenPass.h
--- a/Engine/Source/Runtime/Renderer/Public/ScreenPass.h
+++ b/Engine/Source/Runtime/Renderer/Public/ScreenPass.h
@@ -47,7 +47,8 @@ struct FScreenPassTexture
 	FScreenPassTexture(FRDGTextureRef InTexture, FIntRect InViewRect);
 
 	// Returns a 2D texture from the input slice. A copy is used if the input is a slice of a 2D texture array, or the override output is active.
-	RENDERER_API static FScreenPassTexture CopyFromSlice(FRDGBuilder& GraphBuilder, const FScreenPassTextureSlice& ScreenTextureSlice, FScreenPassTexture OverrideOutput = FScreenPassTexture());
+	// If a MultiViewCount is set, that number of slices will be copied. This mode of use is incompatible with TSR, which outputs a Texture2DArray to represent SceneColor history rather than multi-view slices.
+	RENDERER_API static FScreenPassTexture CopyFromSlice(FRDGBuilder& GraphBuilder, const FScreenPassTextureSlice& ScreenTextureSlice, FScreenPassTexture OverrideOutput = FScreenPassTexture(), uint32 MultiViewCount = 0);
 
 	bool IsValid() const;
 
