[DOF] Fix incorrect scatter occlusion test in hybrid scatter pixel shader. The scatter occlusion output from the gather stage is stored in encoded coc radius units, while the scatter occlusion shader works with output coc radius, which lead to false occlusion of scattered bokeh.

diff --git a/Engine/Shaders/Private/DiaphragmDOF/DOFHybridScatterPixelShader.usf b/Engine/Shaders/Private/DiaphragmDOF/DOFHybridScatterPixelShader.usf
--- a/Engine/Shaders/Private/DiaphragmDOF/DOFHybridScatterPixelShader.usf
+++ b/Engine/Shaders/Private/DiaphragmDOF/DOFHybridScatterPixelShader.usf
@@ -164,8 +164,8 @@ void ScatterMainPS(
 	#if DIM_SCATTER_OCCLUSION
 		float4 ScatterOcclusionSample = ScatterOcclusion.SampleLevel(GlobalBilinearWrappedSampler, SvPosition.xy * ScatterOcclusionSize.zw, 0);
 
-		half CocRadiusAvg = half(ScatterOcclusionSample.x);
-		half CocRadiusVariance = half(ScatterOcclusionSample.y);
+		half EncodedCocRadiusAvg = half(ScatterOcclusionSample.x);
+		half EncodedCocRadiusVariance = half(ScatterOcclusionSample.y);
 	#endif
 	
 	half2 PixelOffset0 = half2(SvPosition.xy - BokehCenter0);
@@ -228,8 +228,9 @@ void ScatterMainPS(
 			half2 Visibility = 1;
 			#if DIM_SCATTER_OCCLUSION
 			{
-				half2 CocRadiusComparison = max(PixelAbsCocRadius * half(1 - FAST_GATHERING_COC_ERROR) - CocRadiusAvg, (half2)0);
-				Visibility = CocRadiusVariance / (CocRadiusVariance + CocRadiusComparison * CocRadiusComparison);
+				half2 EncodedPixelAbsCocRadius = PixelAbsCocRadius / ScatteringScaling;
+				half2 CocRadiusComparison = max(EncodedPixelAbsCocRadius * half(1 - FAST_GATHERING_COC_ERROR) - EncodedCocRadiusAvg, (half2)0);
+				Visibility = EncodedCocRadiusVariance / (EncodedCocRadiusVariance + CocRadiusComparison * CocRadiusComparison);
 			}
 			#endif
 
