[HWRT] Improvements to ray tracing debug visualization.
- flip normal for backface hits when using inline.
- support hit kind visualization mode using inline ray tracing.

diff --git a/Engine/Shaders/Private/RayTracing/RayTracingDebug.usf b/Engine/Shaders/Private/RayTracing/RayTracingDebug.usf
--- a/Engine/Shaders/Private/RayTracing/RayTracingDebug.usf
+++ b/Engine/Shaders/Private/RayTracing/RayTracingDebug.usf
@@ -315,6 +315,11 @@ float4 DebugRayTracingInstance(FRayDesc Ray, out float OutDepth)
 	Payload.InstanceIndex = TraceRayResult.InstanceIndex;
 	Payload.GeometryIndex = TraceRayResult.GeometryIndex;
 	Payload.ScenePrimitiveIndex = InstanceSceneData.PrimitiveId;
+
+	if (TraceRayResult.bIsFrontFace)
+	{
+		Payload.SetFrontFace();
+	}
 #else
 	TraceRay(
 		TLAS,
@@ -327,6 +332,8 @@ float4 DebugRayTracingInstance(FRayDesc Ray, out float OutDepth)
 		Payload);
 #endif
 
+	Payload.WorldNormal = Payload.IsFrontFace() ? Payload.WorldNormal : -Payload.WorldNormal;
+
 	if (Payload.IsHit())
 	{
 		const float3 InstanceColor = IntToColor(Payload.InstanceHash);
@@ -338,7 +345,6 @@ float4 DebugRayTracingInstance(FRayDesc Ray, out float OutDepth)
 
 		switch (VisualizationMode)
 		{
-		default:
 		case RAY_TRACING_DEBUG_VIZ_INSTANCES:
 			Result.rgb = InstanceColor;
 			break;
@@ -366,6 +372,12 @@ float4 DebugRayTracingInstance(FRayDesc Ray, out float OutDepth)
 		case RAY_TRACING_DEBUG_VIZ_PROXY_TYPE:
 			Result.rgb = IntToColor(InstanceDebugData.ProxyHash);
 			break;
+		case RAY_TRACING_DEBUG_VIZ_HITKIND:
+			Result.rgb = Payload.IsFrontFace() ? float4(0.0, 1.0, 0.0, 0.0f) : float4(1.0, 0.0, 0.0, 0.0f);
+			break;
+		default:
+			Result.rgb = 0.5f;
+			break;
 		}
 
 		if (VisualizationMode == RAY_TRACING_DEBUG_VIZ_PICKER)
diff --git a/Engine/Shaders/Private/RayTracing/RayTracingDebugCHS.usf b/Engine/Shaders/Private/RayTracing/RayTracingDebugCHS.usf
--- a/Engine/Shaders/Private/RayTracing/RayTracingDebugCHS.usf
+++ b/Engine/Shaders/Private/RayTracing/RayTracingDebugCHS.usf
@@ -25,6 +25,11 @@ RAY_TRACING_ENTRY_CLOSEST_HIT(RayTracingDebugMainCHS,
 	Payload.InstanceIndex = InstanceIndex();
 	Payload.GeometryIndex = GeometryIndex();
 	Payload.ScenePrimitiveIndex = InstanceSceneData.PrimitiveId;
+	
+	if (HitKind() == HIT_KIND_TRIANGLE_FRONT_FACE)
+	{
+		Payload.SetFrontFace();
+	}
 }
 
 RAY_TRACING_ENTRY_ANY_HIT(RayTracingDebugAHS,
diff --git a/Engine/Shaders/Private/RayTracing/RayTracingDebugUtils.ush b/Engine/Shaders/Private/RayTracing/RayTracingDebugUtils.ush
--- a/Engine/Shaders/Private/RayTracing/RayTracingDebugUtils.ush
+++ b/Engine/Shaders/Private/RayTracing/RayTracingDebugUtils.ush
@@ -16,6 +16,16 @@ struct FRayTracingDebugPayload : FMinimalPayload
 	
 	uint TriangleHitCountPerRay;
 	uint ScenePrimitiveIndex;
+	uint Flags;
+	
+	void SetFrontFace()
+	{
+		Flags |= RAY_TRACING_PAYLOAD_OUTPUT_FLAG_FRONT_FACE;
+	}
+	bool IsFrontFace()
+	{
+		return (Flags & RAY_TRACING_PAYLOAD_OUTPUT_FLAG_FRONT_FACE) != 0;
+	}
 };
 #if IS_PAYLOAD_ENABLED(RT_PAYLOAD_TYPE_RAYTRACING_DEBUG)
 CHECK_RT_PAYLOAD_SIZE(FRayTracingDebugPayload)
diff --git a/Engine/Shaders/Shared/RaytracingDebugDefinitions.h b/Engine/Shaders/Shared/RaytracingDebugDefinitions.h
--- a/Engine/Shaders/Shared/RaytracingDebugDefinitions.h
+++ b/Engine/Shaders/Shared/RaytracingDebugDefinitions.h
@@ -105,6 +105,7 @@ INLINE_ATTR bool RayTracingDebugModeSupportsInline(UINT_TYPE DebugVisualizationM
 		&& (DebugVisualizationMode != RAY_TRACING_DEBUG_VIZ_HITCOUNT_PER_INSTANCE);
 	const bool bIsTraversalMode = IsRayTracingDebugTraversalMode(DebugVisualizationMode);
 	const bool bIsPickingMode = DebugVisualizationMode == RAY_TRACING_DEBUG_VIZ_PICKER;
+	const bool bIsHitKind = DebugVisualizationMode == RAY_TRACING_DEBUG_VIZ_HITKIND;
 
-	return bIsBarycentricsMode || bIsWorldNormalMode || bIsPickingMode || bUsesDebugCHSAndNotHitCount || bIsTraversalMode;
+	return bIsBarycentricsMode || bIsWorldNormalMode || bIsPickingMode || bUsesDebugCHSAndNotHitCount || bIsTraversalMode || bIsHitKind;
 }
