[HWRT] Fix truncation warning in RayTracingDebug.usf

diff --git a/Engine/Shaders/Private/RayTracing/RayTracingDebug.usf b/Engine/Shaders/Private/RayTracing/RayTracingDebug.usf
--- a/Engine/Shaders/Private/RayTracing/RayTracingDebug.usf
+++ b/Engine/Shaders/Private/RayTracing/RayTracingDebug.usf
@@ -373,7 +373,7 @@ float4 DebugRayTracingInstance(FRayDesc Ray, out float OutDepth)
 			Result.rgb = IntToColor(InstanceDebugData.ProxyHash);
 			break;
 		case RAY_TRACING_DEBUG_VIZ_HITKIND:
-			Result.rgb = Payload.IsFrontFace() ? float4(0.0, 1.0, 0.0, 0.0f) : float4(1.0, 0.0, 0.0, 0.0f);
+			Result.rgb = Payload.IsFrontFace() ? float3(0.0, 1.0, 0.0) : float3(1.0, 0.0, 0.0);
 			break;
 		default:
 			Result.rgb = 0.5f;
