Fix incorrect anisotropy output when using Substrate Format=0 with Substrate node like vertical layering.

The vertical layering for aniso to be disabled, but the shader was still marking the pixel has having aniso. This caused the outpt pixel to have incorrect shading but also pixels to have the wrong classification.

#Synced-CL 48133794

diff --git a/Engine/Shaders/Private/BasePassPixelShader.usf b/Engine/Shaders/Private/BasePassPixelShader.usf
--- a/Engine/Shaders/Private/BasePassPixelShader.usf
+++ b/Engine/Shaders/Private/BasePassPixelShader.usf
@@ -875,6 +875,8 @@ uint GetDiffuseIndirectSampleOcclusion(FGBufferData GBuffer, float3 V, float3 Wo
 #endif // SUBSTRATE_ENABLED
 
 
+#define SELECTIVE_OUTPUT_MASK_OFFSET 4
+
 #if USES_GBUFFER
 
 // The selective output mask can only depend on defines, since the shadow will not export the data.
@@ -895,6 +897,14 @@ uint GetSelectiveOutputMask()
 #endif
 	return Mask;
 }
+
+void RemoveAnisotropyFromSelectiveOutputMask(inout uint InMask)
+{
+#if MATERIAL_USES_ANISOTROPY && !IS_NANITE_PASS
+	InMask &= ~(HAS_ANISOTROPY_MASK >> SELECTIVE_OUTPUT_MASK_OFFSET);
+#endif
+}
+
 #endif // USES_GBUFFER
 
 // is called in MainPS() from PixelShaderOutputCommon.usf
@@ -1196,12 +1206,12 @@ void FPixelShaderInOut_MainPS(
 #if USES_GBUFFER
 	// This requires cleanup. Shader code that uses GBuffer.SelectiveOutputMask expects the outputmask to be in
 	// bits [4:7], but it gets packed as bits [0:3] in the flexible gbuffer since we might move it around.
-	GBuffer.SelectiveOutputMask = GetSelectiveOutputMask() >> 4;
+	GBuffer.SelectiveOutputMask = GetSelectiveOutputMask() >> SELECTIVE_OUTPUT_MASK_OFFSET;
 #if !ALLOW_STATIC_LIGHTING
 	// ZERO_PRECSHADOW_MASK is not used when !ALLOW_STATIC_LIGHTING, so we can alias the bit with bIsFirstPerson.
 	if (GetPrimitiveData(MaterialParameters).Flags & PRIMITIVE_SCENE_DATA_FLAG_IS_FIRST_PERSON)
 	{
-		GBuffer.SelectiveOutputMask |= IS_FIRST_PERSON_MASK >> 4;
+		GBuffer.SelectiveOutputMask |= IS_FIRST_PERSON_MASK >> SELECTIVE_OUTPUT_MASK_OFFSET;
 	}
 #endif
 	GBuffer.Velocity = 0;
@@ -1928,6 +1938,13 @@ void FPixelShaderInOut_MainPS(
 			Dither,
 			Export.ShadingModelID
 			);
+
+		// If a material is marked as aniso on the CPU, but actually does not have anisotropy for the current pixel, remove it.
+		if (Export.Anisotropy == 0)
+		{
+			RemoveAnisotropyFromSelectiveOutputMask(GBuffer.SelectiveOutputMask);
+		}
+
 		GBuffer.SpecularColor = ComputeF0(GBuffer.Specular, GBuffer.BaseColor, GBuffer.Metallic);
 		GBuffer.DiffuseColor = GBuffer.BaseColor - GBuffer.BaseColor * GBuffer.Metallic;
 
