Substrate - fixed missing energy when SheenQuality=0.

#rnx
#tests PC editor dx12 material preview
#Synced-CL 48277472

diff --git a/Engine/Shaders/Private/Substrate/SubstrateEvaluation.ush b/Engine/Shaders/Private/Substrate/SubstrateEvaluation.ush
--- a/Engine/Shaders/Private/Substrate/SubstrateEvaluation.ush
+++ b/Engine/Shaders/Private/Substrate/SubstrateEvaluation.ush
@@ -1868,7 +1868,7 @@ FSubstrateEnvLightResult SubstrateEvaluateForEnvLight(FSubstrateBSDFContext BSDF
 				#endif // SUBSTRATE_SHEEN_QUALITY
 
 				// The specular and diffuse components below the fuzz are only attenuated linearly according to the amount of fuzz.
-				const float Cloth_DirectionalAlbedo_SpecularTransmission = 1 - ClothEnergyTerms.E * FuzzAmount; // equivalent to: lerp(1.0, ComputeEnergyPreservation(ClothEnergyTerms), FuzzAmount);
+				const float Cloth_DirectionalAlbedo_SpecularTransmission = lerp(1.0, ComputeEnergyPreservation(ClothEnergyTerms), FuzzAmount);
 
 				SubstrateEnvLightResult.DiffuseColor		*= Cloth_DirectionalAlbedo_SpecularTransmission;
 				SubstrateEnvLightResult.DiffuseWeight		*= Cloth_DirectionalAlbedo_SpecularTransmission;
