Change hair strands to use Substrate lighting when Substrate is enabled.

This has the side effect of fixing light's specular scale working correct on hair shading (regular surface and strands).

#Synced-CL 48279520

diff --git a/Engine/Shaders/Private/DeferredLightPixelShaders.usf b/Engine/Shaders/Private/DeferredLightPixelShaders.usf
--- a/Engine/Shaders/Private/DeferredLightPixelShaders.usf
+++ b/Engine/Shaders/Private/DeferredLightPixelShaders.usf
@@ -28,7 +28,7 @@
   #endif
 #endif
 
-#if SUBSTRATE_GBUFFER_FORMAT==0
+#if SUBSTRATE_GBUFFER_FORMAT==0 || USE_HAIR_LIGHTING
 #define SUBSTRATE_LOAD_FROM_MATERIALCONTAINER 0	// Load from a specified BSDF, not from MaterialBuffer.
 #define SUBSTRATE_INLINE_SHADING 1				// Inline shading for normal setup if needed.
 #define SUBSTRATE_SSS_MATERIAL_OVERRIDE 1		// Need for SSS dither.
@@ -220,14 +220,19 @@ float4 SubstrateReadPrecomputedShadowFactors(FSubstratePixelHeader SubstratePixe
 
 // Affect the light color by any other shadow/transmittance before the lighting is evaluated
 // in order to get a correct alpha channel containing important information for SSS.
-void UpdateLightDataColor(inout float3 LightDataColor, FInputParams InputParams, FDerivedParams DerivedParams)
+void UpdateLightDataColor(inout float3 LightDataColor, float2 InScreenPos, float2 InScreenUV, float InDeviceZ, float3 InTranslatedWorldPosition)
 {
-	const float Attenuation = ComputeLightProfileMultiplier(DerivedParams.TranslatedWorldPosition, GetDeferredLightTranslatedWorldPosition(), -DeferredLightUniforms.Direction, DeferredLightUniforms.Tangent, DeferredLightUniforms.IESAtlasIndex);
-	const float3 AttenuationRGB = Attenuation * GetAtmosphereAndCloudAttenuation(GetStableTranslatedLargeWorldPosition(InputParams.PixelPos, LookupDeviceZ(InputParams.ScreenUV)), InputParams.ScreenUV);
+	const float Attenuation = ComputeLightProfileMultiplier(InTranslatedWorldPosition, GetDeferredLightTranslatedWorldPosition(), -DeferredLightUniforms.Direction, DeferredLightUniforms.Tangent, DeferredLightUniforms.IESAtlasIndex);
+	const float3 AttenuationRGB = Attenuation * GetAtmosphereAndCloudAttenuation(GetStableTranslatedLargeWorldPosition(InScreenPos, InDeviceZ), InScreenUV);
 
 	LightDataColor *= AttenuationRGB;
 }
 
+void UpdateLightDataColor(inout float3 LightDataColor, FInputParams InputParams, float3 InTranslatedWorldPosition)
+{
+	UpdateLightDataColor(LightDataColor, InputParams.PixelPos, InputParams.ScreenUV, LookupDeviceZ(InputParams.ScreenUV), InTranslatedWorldPosition);
+}
+
 #if USE_HAIR_LIGHTING == 0
 
 void DeferredLightPixelMain(
@@ -292,7 +297,7 @@ void DeferredLightPixelMain(
 		const FDerivedParams DerivedParams = GetDerivedParams(InputParams, SceneDepth);
 
 		FDeferredLightData LightData = InitDeferredLightFromUniforms(CURRENT_LIGHT_TYPE);
-		UpdateLightDataColor(LightData.Color, InputParams, DerivedParams);
+		UpdateLightDataColor(LightData.Color, InputParams, DerivedParams.TranslatedWorldPosition);
 
 		float3 V =-DerivedParams.CameraVector;
 		float3 L = LightData.Direction;	// Already normalized
@@ -377,7 +382,7 @@ void DeferredLightPixelMain(
 		const FDerivedParams DerivedParams = GetDerivedParams(InputParams, SceneDepth);
 
 		FDeferredLightData LightData = InitDeferredLightFromUniforms(CURRENT_LIGHT_TYPE);
-		UpdateLightDataColor(LightData.Color, InputParams, DerivedParams);
+		UpdateLightDataColor(LightData.Color, InputParams, DerivedParams.TranslatedWorldPosition);
 
 	 #if USE_HAIR_COMPLEX_TRANSMITTANCE
 		if (ScreenSpaceData.GBuffer.ShadingModelID == SHADINGMODELID_HAIR && ShouldUseHairComplexTransmittance(ScreenSpaceData.GBuffer))
@@ -489,6 +494,7 @@ void DeferredLightPixelMain(
 	const float3 CameraVector = GetCameraVectorFromTranslatedWorldPosition(TranslatedWorldPosition);
 
 	FDeferredLightData LightData = InitDeferredLightFromUniforms(CURRENT_LIGHT_TYPE);
+	UpdateLightDataColor(LightData.Color, ScreenPosition, UV, Sample.Depth, TranslatedWorldPosition);
 
 	// todo add light culling/early out
 
@@ -547,14 +553,56 @@ void DeferredLightPixelMain(
 	}
 	LightAttenuation = min(LightAttenuation, LightData.HairTransmittance.OpaqueVisibility.xxxx);
 
+#if SUBSTRATE_ENABLED==1
+
+	float3 V =-CameraVector;
+	float3 L = LightData.Direction;	// Already normalized
+	float3 ToLight = L;
+	float LightMask = 1;
+	if (LightData.bRadialLight)
+	{
+		LightMask = GetLocalLightAttenuation(TranslatedWorldPosition, LightData, ToLight, L);
+	}
+
+	if (LightMask > 0)
+	{
+		FSubstrateShadowTermInputParameters SubstrateShadowTermInputParameters	= GetInitialisedSubstrateShadowTermInputParameters();
+		SubstrateShadowTermInputParameters.bEvaluateShadowTerm					= true;
+		SubstrateShadowTermInputParameters.SceneDepth							= SceneDepth;
+		SubstrateShadowTermInputParameters.PrecomputedShadowFactors 			= 1.f;
+		SubstrateShadowTermInputParameters.TranslatedWorldPosition				= TranslatedWorldPosition;
+		SubstrateShadowTermInputParameters.LightAttenuation 					= LightAttenuation;
+		SubstrateShadowTermInputParameters.Dither								= InterleavedGradientNoise(PixelCoord, View.StateFrameIndexMod8);
+		SubstrateShadowTermInputParameters.bUseShadowTermHairTransmittance 		= true;
+			
+		FSubstrateGBufferBSDF GBufferBSDF = SubstrateReadGBufferBSDF(HairScreenSpaceData);
+
+		FSubstrateDeferredLighting SubstrateLighting = SubstrateDeferredLighting(
+			LightData,
+			V,
+			L,
+			ToLight,
+			LightMask,
+			SubstrateShadowTermInputParameters,
+			GBufferBSDF.BSDF,
+			GBufferBSDF.BSDFTangentBasis,
+			GBufferBSDF.BSDFAO
+		);
+
+		const float LocalCoverage = From8bitCoverage(Sample.Coverage8bit);
+
+		OutColor.rgb = SubstrateLighting.SceneColor.xyz * LocalCoverage * GetExposure();
+		OutColor.a = LocalCoverage;
+	}
+#else // SUBSTRATE_ENABLED
+
 	float SurfaceShadow = 1.0f;
 	const float4 Radiance = GetDynamicLighting(TranslatedWorldPosition, CameraVector, HairScreenSpaceData.GBuffer, HairScreenSpaceData.AmbientOcclusion, LightData, LightAttenuation, Dither, PixelCoord, SurfaceShadow);
-	const float  Attenuation = ComputeLightProfileMultiplier(TranslatedWorldPosition, GetDeferredLightTranslatedWorldPosition(), -DeferredLightUniforms.Direction, DeferredLightUniforms.Tangent, DeferredLightUniforms.IESAtlasIndex);
-	const float3 AtmosphericTransmittance = GetAtmosphereAndCloudAttenuation(GetStableTranslatedLargeWorldPosition(ScreenPosition, Sample.Depth), UV);
 
 	const float LocalCoverage = From8bitCoverage(Sample.Coverage8bit);
 
-	OutColor.rgb = AtmosphericTransmittance * Radiance.xyz * Attenuation * LocalCoverage * GetExposure();
+	OutColor.rgb = Radiance.xyz * LocalCoverage * GetExposure();
 	OutColor.a = LocalCoverage;
+#endif // SUBSTRATE_ENABLED
 }
 #endif
diff --git a/Engine/Shaders/Private/Substrate/Substrate.ush b/Engine/Shaders/Private/Substrate/Substrate.ush
--- a/Engine/Shaders/Private/Substrate/Substrate.ush
+++ b/Engine/Shaders/Private/Substrate/Substrate.ush
@@ -2262,6 +2262,7 @@ struct FSubstrateIntegrationSettings
 	bool bForceFullyRoughAccurate;
 	bool bRoughDiffuseEnabled;
 	bool bRoughnessTracking;
+	bool bUseShadowTermHairTransmittance;
 
 	// Embbeded debug settings
 	int  SliceStoringDebugSubstrateTreeData;
@@ -2275,6 +2276,7 @@ FSubstrateIntegrationSettings InitSubstrateIntegrationSettings(bool bForceFullyR
 	Out.bForceFullyRoughAccurate = false;
 	Out.bRoughDiffuseEnabled = bRoughDiffuseEnabled;
 	Out.bRoughnessTracking = bRoughnessTracking;
+	Out.bUseShadowTermHairTransmittance = false;
 
 	Out.PeelLayersAboveDepth = PeelLayersAboveDepth;
 	Out.SliceStoringDebugSubstrateTreeData = -1;
diff --git a/Engine/Shaders/Private/Substrate/SubstrateDeferredLighting.ush b/Engine/Shaders/Private/Substrate/SubstrateDeferredLighting.ush
--- a/Engine/Shaders/Private/Substrate/SubstrateDeferredLighting.ush
+++ b/Engine/Shaders/Private/Substrate/SubstrateDeferredLighting.ush
@@ -44,6 +44,7 @@ struct FSubstrateShadowTermInputParameters
 	float3 TranslatedWorldPosition;
 	float4 LightAttenuation;
 	float Dither;
+	bool bUseShadowTermHairTransmittance;
 };
 
 FSubstrateShadowTermInputParameters GetInitialisedSubstrateShadowTermInputParameters()
@@ -75,9 +76,11 @@ FSubstrateDeferredLighting SubstrateDeferredLighting(
 {
 	FSubstrateDeferredLighting SubstrateLighting = GetInitialisedSubstrateDeferredLighting();
 #if SUBSTRATE_MATERIALCONTAINER_IS_VIEWRESOURCE
-	const FSubstrateIntegrationSettings Settings = InitSubstrateIntegrationSettings(false /*bForceFullyRough*/, Substrate.bRoughDiffuse, Substrate.PeelLayersAboveDepth, Substrate.bRoughnessTracking);
+	FSubstrateIntegrationSettings Settings = InitSubstrateIntegrationSettings(false /*bForceFullyRough*/, Substrate.bRoughDiffuse, Substrate.PeelLayersAboveDepth, Substrate.bRoughnessTracking);
+	Settings.bUseShadowTermHairTransmittance = SubstrateShadowTermInputParameters.bUseShadowTermHairTransmittance;
 #else
-	const FSubstrateIntegrationSettings Settings = InitSubstrateIntegrationSettings();
+	FSubstrateIntegrationSettings Settings = InitSubstrateIntegrationSettings();
+	Settings.bUseShadowTermHairTransmittance = SubstrateShadowTermInputParameters.bUseShadowTermHairTransmittance;
 #endif
 
 	FRectTexture RectTexture = ConvertToRectTexture(LightData);
diff --git a/Engine/Shaders/Private/Substrate/SubstrateEvaluation.ush b/Engine/Shaders/Private/Substrate/SubstrateEvaluation.ush
--- a/Engine/Shaders/Private/Substrate/SubstrateEvaluation.ush
+++ b/Engine/Shaders/Private/Substrate/SubstrateEvaluation.ush
@@ -1140,7 +1140,12 @@ FSubstrateEvaluateResult SubstrateEvaluateBSDFCommon(FSubstrateBSDFContext BSDFC
 			GBuffer.WorldNormal		= BSDFContext.N;
 
 			FHairTransmittanceData HairTransmittance = InitHairTransmittanceData();
-			if (HAIR_COMPLEXTRANSMITTANCE(BSDFContext.BSDF))
+			if (Settings.bUseShadowTermHairTransmittance)
+			{
+				HairTransmittance = ShadowTerms.HairTransmittance;
+				HairTransmittance.OpaqueVisibility = ShadowTerms.SurfaceShadow;
+			}
+			else if (HAIR_COMPLEXTRANSMITTANCE(BSDFContext.BSDF))
 			{
 				HairTransmittance = EvaluateDualScattering(GBuffer.BaseColor, BSDFContext.N, GBuffer.Roughness, BSDFContext.V, BSDFContext.L);
 				HairTransmittance.OpaqueVisibility = ShadowTerms.SurfaceShadow;
diff --git a/Engine/Source/Runtime/Renderer/Private/LightRendering.cpp b/Engine/Source/Runtime/Renderer/Private/LightRendering.cpp
--- a/Engine/Source/Runtime/Renderer/Private/LightRendering.cpp
+++ b/Engine/Source/Runtime/Renderer/Private/LightRendering.cpp
@@ -1023,7 +1023,7 @@ class FDeferredLightPS : public FGlobalShader
 			return false;
 		}
 
-		if (Substrate::IsSubstrateEnabled() && PermutationVector.Get<FHairLighting>() == 1 && PermutationVector.Get<Substrate::FSubstrateTileType>() != ESubstrateTileType::ESimple)
+		if (Substrate::IsSubstrateEnabled() && PermutationVector.Get<FHairLighting>() == 1 && PermutationVector.Get<Substrate::FSubstrateTileType>() != ESubstrateTileType::EComplex)
 		{
 			return false;
 		}
@@ -3122,7 +3122,7 @@ void FDeferredShadingSceneRenderer::RenderLightForHair(
 	PermutationVector.Set< FDeferredLightPS::FHairLighting>(1);
 	PermutationVector.Set< FDeferredLightPS::FHairComplexTransmittance>(true);
 	PermutationVector.Set< FDeferredLightPS::FVirtualShadowMapMask >(bUseVirtualShadowMapMask);
-	PermutationVector.Set< Substrate::FSubstrateTileType>(Substrate::IsSubstrateEnabled() ? ESubstrateTileType::ESimple : 0);
+	PermutationVector.Set< Substrate::FSubstrateTileType>(Substrate::IsSubstrateEnabled() ? ESubstrateTileType::EComplex : 0);
 	PermutationVector.Set< FDeferredLightPS::FLightFunctionAtlasDim >(
 		LightFunctionAtlas::IsEnabled(View, ELightFunctionAtlasSystem::DeferredLighting) && LightSceneInfo->Proxy->HasValidLightFunctionAtlasSlot() &&
 		LightSceneInfo->Proxy->GetLightFunctionMaterial() != nullptr && !View.Family->EngineShowFlags.VisualizeLightCulling && bCanLightUsesAtlasForUnbatchedLight);
@@ -3363,7 +3363,7 @@ static void InternalRenderSimpleLightsForHair(
 	PermutationVector.Set< FDeferredLightPS::FHairComplexTransmittance>(true /*bNeedComplexTransmittanceSupport*/);
 	PermutationVector.Set< FDeferredLightPS::FAtmosphereTransmittance >(false);
 	PermutationVector.Set< FDeferredLightPS::FCloudTransmittance >(false);
-	PermutationVector.Set< Substrate::FSubstrateTileType>(Substrate::IsSubstrateEnabled() ? ESubstrateTileType::ESimple : 0);
+	PermutationVector.Set< Substrate::FSubstrateTileType>(Substrate::IsSubstrateEnabled() ? ESubstrateTileType::EComplex : 0);
 
 	PermutationVector = FDeferredLightPS::RemapPermutation(PermutationVector);
 	TShaderMapRef<FDeferredLightPS> PixelShader(View.ShaderMap, PermutationVector);
