Fix OIT shader compilation error with post process domain material when per pixel OIT is enabled.

#Synced-CL 48347083

diff --git a/Engine/Shaders/Private/BasePassPixelShader.usf b/Engine/Shaders/Private/BasePassPixelShader.usf
--- a/Engine/Shaders/Private/BasePassPixelShader.usf
+++ b/Engine/Shaders/Private/BasePassPixelShader.usf
@@ -2478,7 +2478,7 @@ void FPixelShaderInOut_MainPS(
 			Out.MRT[0] = half4(0, 0, 0, 0);
 			Out.MRT[1] = half4(1, 1, 1, 1);
 		}
-		#elif SUBSTRATE_TRANSLUCENT_MATERIAL
+		#elif SUBSTRATE_TRANSLUCENT_MATERIAL && !MATERIAL_DOMAIN_POSTPROCESS
 		{
 			#if SUBSTRATE_BLENDING_COLOREDTRANSMITTANCEONLY
 			AddOITSample(uint2(In.SvPosition.xy), float3(0, 0, 0), FoggedColor, MaterialParameters.ScreenPosition.w);
