PT - adding more bits to PT shading model since it was virtually stealing bits from SHADINGMODELID when those are already precious. And one is needed for toon.

#tests Pc editor dx 12 with PT
#rnx
#Synced-CL 48347003

diff --git a/Engine/Shaders/Private/PathTracing/PathTracingCommon.ush b/Engine/Shaders/Private/PathTracing/PathTracingCommon.ush
--- a/Engine/Shaders/Private/PathTracing/PathTracingCommon.ush
+++ b/Engine/Shaders/Private/PathTracing/PathTracingCommon.ush
@@ -42,8 +42,10 @@
 #define SHADINGMODELID_MEDIUM  			(SHADINGMODELID_NUM + 1)
 #define SHADINGMODELID_SOLID_GLASS 		(SHADINGMODELID_NUM + 2)
 
-HLSL_STATIC_ASSERT(SHADINGMODELID_MEDIUM      <= SHADINGMODELID_MASK, "");
-HLSL_STATIC_ASSERT(SHADINGMODELID_SOLID_GLASS <= SHADINGMODELID_MASK, "");
+#define PATH_TRACING_SHADINGMODEL_ID_MASK 0x3F
+
+HLSL_STATIC_ASSERT(SHADINGMODELID_MEDIUM      <= PATH_TRACING_SHADINGMODEL_ID_MASK, "");
+HLSL_STATIC_ASSERT(SHADINGMODELID_SOLID_GLASS <= PATH_TRACING_SHADINGMODEL_ID_MASK, "");
 
 #define PATH_TRACING_ABSORPTION_SCALE  0.01        // inverse distance at which we hit the transmission color (Beer's law)
 
@@ -249,8 +251,8 @@ struct FPackedPathTracingPayload : FMinimalPayload
 	// 64 bytes total
 #endif
 
-	float UnpackBSDFOpacity() { return f16tof32(PackedData[0]); }
-	uint GetShadingModelID() { return (PackedData[0] >> 16) & 0xF; }
+	float UnpackBSDFOpacity() { return Unpack10F(PackedData[0] & 0x3FF); }
+	uint GetShadingModelID() { return (PackedData[0] >> 10) & PATH_TRACING_SHADINGMODEL_ID_MASK; }
 	uint GetFlags() { return (PackedData[0] >> 20) & 0x1FF; }
 	void SetFlag(uint Flag) { PackedData[0] |= Flag << 20; }
 	void RemoveFlag(uint Flag) { PackedData[0] &= ~(Flag << 20); }
@@ -1005,11 +1007,11 @@ FPackedPathTracingPayload PackPathTracingPayload(FPathTracingPayload Input)
 	FPackedPathTracingPayload Output = (FPackedPathTracingPayload)0;
 	Output.HitT = Input.HitT;
 
-	Output.PackedData[0] = f32tof16(Input.BSDFOpacity);
-	Output.PackedData[0] |= (Input.ShadingModelID & 0xF) << 16;               // 4 bits
-	Output.PackedData[0] |= (Input.Flags & 0x1FF) << 20;                      // 9 bits
-	Output.PackedData[0] |= (Input.PrimitiveLightingChannelMask & 0x7) << 29; // 3 bits
-	                                                                   // total 32 bits
+	Output.PackedData[0]  = Pack10F(Input.BSDFOpacity);                                       // 10 bits
+	Output.PackedData[0] |= (Input.ShadingModelID & PATH_TRACING_SHADINGMODEL_ID_MASK) << 10; // 10 bits - That is a lot, some bits from there could be reused for other things.
+	Output.PackedData[0] |= (Input.Flags & 0x1FF) << 20;                                      // 9  bits
+	Output.PackedData[0] |= (Input.PrimitiveLightingChannelMask & 0x7) << 29;                 // 3  bits
+	                                                                                          // 32 bits total
 #if IS_PAYLOAD_ENABLED(RT_PAYLOAD_TYPE_GPULIGHTMASS)
 	Output.PackedData[1] = PayloadEncodeUnitVector(Input.WorldGeoNormal);
 	Output.PackedData[2] = PayloadEncodeUnitVector(Input.WorldNormal);
