Substrate - built in toon shading BSDF start.
- Works with ProjectGBufferFormat = 0/1
- Basic decal support to be expended upon
- Need to implement support for RT and PT
- Toon is part of complex path for now. Complex path now has a mode to clarify things: Slab, Special, Toon.
- TOONTODO for the rest

#rnx
#tests PC editor
#Synced-CL 48347003

diff --git a/Engine/Shaders/Private/BasePassCommon.ush b/Engine/Shaders/Private/BasePassCommon.ush
--- a/Engine/Shaders/Private/BasePassCommon.ush
+++ b/Engine/Shaders/Private/BasePassCommon.ush
@@ -41,7 +41,7 @@
 #define USES_GBUFFER						(FEATURE_LEVEL >= FEATURE_LEVEL_SM4 && (MATERIALBLENDING_SOLID || MATERIALBLENDING_MASKED) && !FORWARD_SHADING)
 
 // Only some shader models actually need custom data.
-#define WRITES_CUSTOMDATA_TO_GBUFFER		(USES_GBUFFER && (MATERIAL_SHADINGMODEL_SUBSURFACE || MATERIAL_SHADINGMODEL_PREINTEGRATED_SKIN || MATERIAL_SHADINGMODEL_SUBSURFACE_PROFILE || MATERIAL_SHADINGMODEL_CLEAR_COAT || MATERIAL_SHADINGMODEL_TWOSIDED_FOLIAGE || MATERIAL_SHADINGMODEL_HAIR || MATERIAL_SHADINGMODEL_CLOTH || MATERIAL_SHADINGMODEL_EYE))
+#define WRITES_CUSTOMDATA_TO_GBUFFER		(USES_GBUFFER && (MATERIAL_SHADINGMODEL_SUBSURFACE || MATERIAL_SHADINGMODEL_PREINTEGRATED_SKIN || MATERIAL_SHADINGMODEL_SUBSURFACE_PROFILE || MATERIAL_SHADINGMODEL_CLEAR_COAT || MATERIAL_SHADINGMODEL_TWOSIDED_FOLIAGE || MATERIAL_SHADINGMODEL_HAIR || MATERIAL_SHADINGMODEL_CLOTH || MATERIAL_SHADINGMODEL_EYE || MATERIAL_SHADINGMODEL_TOON))
 
 // Based on GetPrecomputedShadowMasks()
 // Note: WRITES_PRECSHADOWFACTOR_TO_GBUFFER is currently disabled because we use the precomputed shadow factor GBuffer outside of STATICLIGHTING_TEXTUREMASK to store UseSingleSampleShadowFromStationaryLights
diff --git a/Engine/Shaders/Private/DBufferDecalShared.ush b/Engine/Shaders/Private/DBufferDecalShared.ush
--- a/Engine/Shaders/Private/DBufferDecalShared.ush
+++ b/Engine/Shaders/Private/DBufferDecalShared.ush
@@ -350,6 +350,21 @@ void ApplyDBufferData(in FSubstrateDBuffer SubstrateDBufferData, inout FSubstrat
 			#endif
 				break;
 			}
+
+			case SUBSTRATE_BSDF_TYPE_TOON:
+			{
+			#if SUBSTRATE_DBUFFER_RESPONSE_BASECOLOR
+				TOON_DIFFUSEALBEDO(BSDF) = TOON_DIFFUSEALBEDO(BSDF) * SubstrateDBufferData.OneMinusCoverage_BaseColor + SubstrateDBufferData.BaseColor;
+			#endif
+
+			#if SUBSTRATE_DBUFFER_RESPONSE_ROUGHNESS
+				const float DBufferAnisotropy = 0.0f;
+				TOON_ROUGHNESS(BSDF)		= TOON_ROUGHNESS(BSDF) * SubstrateDBufferData.OneMinusCoverage_Roughness		+ SubstrateDBufferData.Roughness;
+				TOON_SPECULARCOLOR(BSDF)	= TOON_SPECULARCOLOR(BSDF)  * SubstrateDBufferData.OneMinusCoverage_Roughness	+ SubstrateDBufferData.Specular;
+				TOON_ANISOTROPY(BSDF)		= TOON_ANISOTROPY(BSDF)  * SubstrateDBufferData.OneMinusCoverage_Roughness		+ DBufferAnisotropy;
+			#endif
+				break;
+			}
 		} // switch
 
 	#undef BSDF
diff --git a/Engine/Shaders/Private/DeferredShadingCommon.ush b/Engine/Shaders/Private/DeferredShadingCommon.ush
--- a/Engine/Shaders/Private/DeferredShadingCommon.ush
+++ b/Engine/Shaders/Private/DeferredShadingCommon.ush
@@ -327,7 +327,8 @@ bool HasCustomGBufferData(int ShadingModelID)
 		|| ShadingModelID == SHADINGMODELID_TWOSIDED_FOLIAGE
 		|| ShadingModelID == SHADINGMODELID_HAIR
 		|| ShadingModelID == SHADINGMODELID_CLOTH
-		|| ShadingModelID == SHADINGMODELID_EYE;
+		|| ShadingModelID == SHADINGMODELID_EYE
+		|| ShadingModelID == SHADINGMODELID_SUBSTRATE_TOON;
 }
 
 bool HasAnisotropy(int SelectiveOutputMask)
diff --git a/Engine/Shaders/Private/PathTracing/PathTracingMaterialHitShader.usf b/Engine/Shaders/Private/PathTracing/PathTracingMaterialHitShader.usf
--- a/Engine/Shaders/Private/PathTracing/PathTracingMaterialHitShader.usf
+++ b/Engine/Shaders/Private/PathTracing/PathTracingMaterialHitShader.usf
@@ -777,6 +777,35 @@ RAY_TRACING_ENTRY_CLOSEST_HIT(PathTracingMaterialCHS,
 					Payload.GlintValue = 1.f;
 					break;
 				}
+				case SUBSTRATE_BSDF_TYPE_TOON:
+				{
+					Payload.BSDFOpacity = BaseBSDFOpacity;
+					Payload.WorldNormal = SubstratePixelHeader.SharedLocalBases.Normals[BSDF_GETSHAREDLOCALBASISID(BSDF)];
+					if (SubstrateGetSharedLocalBasisType(SubstratePixelHeader.SharedLocalBases.Types, BSDFIdx) == SUBSTRATE_BASIS_TYPE_TANGENT)
+					{
+						Payload.WorldTangent = SubstratePixelHeader.SharedLocalBases.Tangents[BSDF_GETSHAREDLOCALBASISID(BSDF)];
+					}
+					Payload.WeightV               = BSDF.LuminanceWeightV;
+					Payload.TransmittanceN        = BSDF.TransmittanceAboveAlongN;
+					Payload.CoverageAboveAlongN   = BSDF.CoverageAboveAlongN;
+					Payload.ShadingModelID    	  = SHADINGMODELID_SUBSTRATE;
+					Payload.Ior                   = 0.0f;
+					Payload.DiffuseColor          = float3(1,0,1); // TOON_TODO
+					Payload.SpecularColor         = 0.0f;
+					Payload.SpecularEdgeColor     = 0.0f;
+					Payload.RoughnessData         = 0.0f;
+					Payload.BottomNormalOct16bits = 0;
+
+					// Handle special-case specular features
+					Payload.Anisotropy			= 0.0f;
+					Payload.FuzzAmount			= 0.0f;
+					Payload.FuzzColor			= 0.0f;
+					Payload.FuzzRoughness		= 0.0f;
+					Payload.SpecularProfileId	= 0.0f;
+					Payload.GlintValue			= 1.0f;
+
+					break;
+				}
 			}  // switch on slab tyype
 		} // for loop over BSDFs
 
@@ -1558,6 +1587,11 @@ RAY_TRACING_ENTRY_ANY_HIT(PathTracingMaterialAHS,
 						RefractMix = 1;
 						break;
 					}
+					case SUBSTRATE_BSDF_TYPE_TOON
+					{
+						// No refraction
+						break;
+					}
 				}
 			}
 #undef BSDF
diff --git a/Engine/Shaders/Private/ShadingCommon.ush b/Engine/Shaders/Private/ShadingCommon.ush
--- a/Engine/Shaders/Private/ShadingCommon.ush
+++ b/Engine/Shaders/Private/ShadingCommon.ush
@@ -30,7 +30,8 @@
 #define SHADINGMODELID_SINGLELAYERWATER		10
 #define SHADINGMODELID_THIN_TRANSLUCENT		11
 #define SHADINGMODELID_SUBSTRATE			12		// Temporary while we convert everything to Substrate
-#define SHADINGMODELID_NUM					13
+#define SHADINGMODELID_SUBSTRATE_TOON		13
+#define SHADINGMODELID_NUM					14
 #define SHADINGMODELID_MASK					0xF		// 4 bits reserved for ShadingModelID			
 
 // The flags are defined so that 0 value has no effect!
diff --git a/Engine/Shaders/Private/ShadingModelsMaterial.ush b/Engine/Shaders/Private/ShadingModelsMaterial.ush
--- a/Engine/Shaders/Private/ShadingModelsMaterial.ush
+++ b/Engine/Shaders/Private/ShadingModelsMaterial.ush
@@ -207,4 +207,10 @@ void SetGBufferForShadingModel(
 	#endif
 	}
 #endif
+#if MATERIAL_SHADINGMODEL_TOON
+	else if (ShadingModel == SHADINGMODELID_SUBSTRATE_TOON)
+	{
+		// GBuffer.CustomData is used
+	}
+#endif
 }
\ No newline at end of file
diff --git a/Engine/Shaders/Private/Substrate/Substrate.ush b/Engine/Shaders/Private/Substrate/Substrate.ush
--- a/Engine/Shaders/Private/Substrate/Substrate.ush
+++ b/Engine/Shaders/Private/Substrate/Substrate.ush
@@ -418,19 +418,26 @@ struct FSubstrateBSDF
 #define HEADER_BIT_COUNT_BSDF_COUNT					4
 #define HEADER_BIT_COUNT_SHAREDLOCALBASES_COUNT		3
 #define HEADER_BIT_COUNT_SHAREDLOCALBASES_TYPE		4
-#define HEADER_BIT_COUNT_COMPLEXPATH_SPECIAL		1
-#define HEADER_BIT_COUNT___UNUSED___				4
+#define HEADER_BIT_COUNT_COMPLEXPATH_MODE			2
+#define HEADER_BIT_COUNT___UNUSED___				3
 // End of header bits.
 
+// When following the complex path, we can have complex slabs or toon shading BSDF.
+// This is needed from the head to know how to unpack the incoming data.
+#define COMPLEXPATH_MODE_SLAB						0
+#define COMPLEXPATH_MODE_SLABSPECIAL				1
+#define COMPLEXPATH_MODE_TOON						2
+#define COMPLEXPATH_MODE_UNUSED						3		// Reserved for custom BSDF?
+
 #define HEADER_BIT_OFFSET_STATE						0
 #define HEADER_BIT_OFFSET_IRRADIANCE_AO				(HEADER_BIT_OFFSET_STATE					+ HEADER_BIT_COUNT_STATE)
 #define HEADER_BIT_OFFSET_BSDF_COUNT				(HEADER_BIT_OFFSET_IRRADIANCE_AO			+ HEADER_BIT_COUNT_IRRADIANCE_AO)
 #define HEADER_BIT_OFFSET_SHAREDLOCALBASES_COUNT	(HEADER_BIT_OFFSET_BSDF_COUNT				+ HEADER_BIT_COUNT_BSDF_COUNT)
 #define HEADER_BIT_OFFSET_SHAREDLOCALBASES_TYPE		(HEADER_BIT_OFFSET_SHAREDLOCALBASES_COUNT	+ HEADER_BIT_COUNT_SHAREDLOCALBASES_COUNT)
-#define HEADER_BIT_OFFSET_COMPLEXPATH_SPECIAL		(HEADER_BIT_OFFSET_SHAREDLOCALBASES_TYPE	+ HEADER_BIT_COUNT_SHAREDLOCALBASES_TYPE)
+#define HEADER_BIT_OFFSET_COMPLEXPATH_MODE			(HEADER_BIT_OFFSET_SHAREDLOCALBASES_TYPE	+ HEADER_BIT_COUNT_SHAREDLOCALBASES_TYPE)
 
 // Sanity check
-#if (HEADER_BIT_OFFSET_COMPLEXPATH_SPECIAL + HEADER_BIT_COUNT_COMPLEXPATH_SPECIAL) > 32
+#if (HEADER_BIT_OFFSET_COMPLEXPATH_SPECIAL + HEADER_BIT_COUNT_COMPLEXPATH_MODE) > 32
 #error Substrate Header is larger than 32 bits
 #endif
 
@@ -477,9 +484,9 @@ struct FSubstrateBSDF
 #define HEADER_GETSHAREDLOCALBASESTYPE(X)			READ_BITS( X, HEADER_BIT_COUNT_SHAREDLOCALBASES_TYPE, HEADER_BIT_OFFSET_SHAREDLOCALBASES_TYPE)
 #define HEADER_SETSHAREDLOCALBASESTYPE(X, NewValue)	WRITE_BITS(X, HEADER_BIT_COUNT_SHAREDLOCALBASES_TYPE, HEADER_BIT_OFFSET_SHAREDLOCALBASES_TYPE, NewValue)
 
-#define HEADER_GETCOMPLEXSPECIALPATH(X)				READ_BITS( X, HEADER_BIT_COUNT_COMPLEXPATH_SPECIAL, HEADER_BIT_OFFSET_COMPLEXPATH_SPECIAL)
-#define HEADER_SETCOMPLEXSPECIALPATH(X, NewValue)	WRITE_BITS(X, HEADER_BIT_COUNT_COMPLEXPATH_SPECIAL, HEADER_BIT_OFFSET_COMPLEXPATH_SPECIAL, NewValue)
-#define HEADER_COMPLEXSPECIALPATH_MASK				TO_OFFSET_BITMASK(HEADER_BIT_COUNT_COMPLEXPATH_SPECIAL, HEADER_BIT_OFFSET_COMPLEXPATH_SPECIAL)
+#define HEADER_GETCOMPLEXPATHMODE(X)				READ_BITS( X, HEADER_BIT_COUNT_COMPLEXPATH_MODE, HEADER_BIT_OFFSET_COMPLEXPATH_MODE)
+#define HEADER_SETCOMPLEXPATHMODE(X, NewValue)		WRITE_BITS(X, HEADER_BIT_COUNT_COMPLEXPATH_MODE, HEADER_BIT_OFFSET_COMPLEXPATH_MODE, NewValue)
+#define HEADER_COMPLEXPATHMODE_MASK					TO_OFFSET_BITMASK(HEADER_BIT_COUNT_COMPLEXPATH_MODE, HEADER_BIT_OFFSET_COMPLEXPATH_MODE)
 
 ///////////////////////////////////////////////////////////////////////////////
 // BSDF State
@@ -704,6 +711,13 @@ struct FSubstrateBSDF
 #define SLW_COLORSCALEBEHINDWATER(X)	X.InlineVGPRs[2].xyz
 #endif
 
+///////////////////////////////////////////////////////////////////////////////
+// Toon BSDF members
+#define TOON_DIFFUSEALBEDO(X)			X.VGPRs[0].xyz
+#define TOON_ROUGHNESS(X)				X.VGPRs[0].w
+#define TOON_SPECULARCOLOR(X)			X.VGPRs[1].xyz
+#define TOON_ANISOTROPY(X)				X.VGPRs[1].w
+
 ///////////////////////////////////////////////////////////////////////////////
 // Forward declarations
 
@@ -1953,6 +1967,50 @@ FSubstrateData GetSubstrateSingleLayerWaterBSDF(
 	return SubstrateData;
 }
 
+FSubstrateData GetSubstrateToonBSDF(
+	float3 DiffuseColor, float3 SpecularColor, float Roughness, float Anisotropy, 
+	float3 Emissive,
+	uint SharedLocalBasisIndex, inout uint SharedLocalBasisTypes)
+{
+	FSubstrateData SubstrateData = GetInitialisedSubstrateData();
+
+#if MATERIAL_USES_ANISOTROPY && SUBSTRATE_FASTPATH==0
+	const uint bHasAnisotropy = Anisotropy != 0 ? 1 : 0;
+	if (bHasAnisotropy)
+	{
+		SubstrateRequestSharedLocalBasisTangent(SharedLocalBasisTypes, SharedLocalBasisIndex);
+	}
+#else
+	const uint bHasAnisotropy = 0;
+	Anisotropy = 0.0f;
+#endif
+
+#if (SUBSTRATE_FASTPATH==0 || MATERIALBLENDING_ANY_TRANSLUCENT) // !Fast path or rendering translucent materials
+	BSDF_SETHASANISOTROPY			(SubstrateData.InlinedBSDF, bHasAnisotropy);
+	BSDF_SETISTOPLAYER				(SubstrateData.InlinedBSDF, 1);
+#endif
+
+	BSDF_SETTYPE					(SubstrateData.InlinedBSDF, SUBSTRATE_BSDF_TYPE_TOON);
+	BSDF_SETSHAREDLOCALBASISID		(SubstrateData.InlinedBSDF, SharedLocalBasisIndex);
+	BSDF_SETEMISSIVE				(SubstrateData.InlinedBSDF, Emissive);
+	TOON_DIFFUSEALBEDO				(SubstrateData.InlinedBSDF) = DiffuseColor;
+	TOON_SPECULARCOLOR				(SubstrateData.InlinedBSDF) = SpecularColor;
+	TOON_ROUGHNESS					(SubstrateData.InlinedBSDF) = Roughness;
+#if MATERIAL_USES_ANISOTROPY
+	TOON_ANISOTROPY					(SubstrateData.InlinedBSDF) = Anisotropy;
+#endif
+
+#if SUBSTRATE_INLINE_SHADING
+	SubstrateData.InlinedBSDF.Coverage	= 1.0f;
+#endif
+
+#if USE_DEVELOPMENT_SHADERS
+	SubstrateData.PreviewColor			= DiffuseColor;
+#endif
+
+	return SubstrateData;
+}
+
 ///////////////////////////////////////////////////////////////////////////////
 // Raytracing material buffer
 
@@ -2398,7 +2456,7 @@ struct FSubstratePixelHeader
 	bool IsSingleMaterial()										{ return (this.State & HEADER_MASK_MATERIALMODE) == HEADER_MATERIALMODE_SLAB_SINGLE; }
 	bool IsComplexMaterial()									{ return (this.State & HEADER_MASK_MATERIALMODE) == HEADER_MATERIALMODE_SLAB_COMPLEX; }
 	#endif
-	bool IsComplexSpecialMaterial()								{ return (this.State & HEADER_COMPLEXSPECIALPATH_MASK) != 0; }
+	bool IsComplexSpecialMaterial()								{ return HEADER_GETCOMPLEXPATHMODE(this.State) == COMPLEXPATH_MODE_SLABSPECIAL; }
 	bool IsSingleLayerWater()									{ return (this.State & HEADER_MASK_MATERIALMODE) == HEADER_MATERIALMODE_SLWATER; }
 	bool IsHair()												{ return (this.State & HEADER_MASK_MATERIALMODE) == HEADER_MATERIALMODE_HAIR; }
 	bool IsEye()												{ return (this.State & HEADER_MASK_MATERIALMODE) == HEADER_MATERIALMODE_EYE; }
@@ -2408,7 +2466,7 @@ struct FSubstratePixelHeader
 		{
 		case HEADER_MATERIALMODE_SLAB_SIMPLE:
 		case HEADER_MATERIALMODE_SLAB_SINGLE:
-		case HEADER_MATERIALMODE_SLAB_COMPLEX:	return SUBSTRATE_BSDF_TYPE_SLAB;
+		case HEADER_MATERIALMODE_SLAB_COMPLEX:	return HEADER_GETCOMPLEXPATHMODE(this.State) == COMPLEXPATH_MODE_TOON ? SUBSTRATE_BSDF_TYPE_TOON : SUBSTRATE_BSDF_TYPE_SLAB; // Get the general type for unpacking. So far, SUBSTRATE_BSDF_TYPE_SLAB that can also be SUBSTRATE_BSDF_TYPE_TOON
 		case HEADER_MATERIALMODE_SLWATER:		return SUBSTRATE_BSDF_TYPE_SINGLELAYERWATER;
 		case HEADER_MATERIALMODE_HAIR:			return SUBSTRATE_BSDF_TYPE_HAIR;
 		case HEADER_MATERIALMODE_EYE:			return SUBSTRATE_BSDF_TYPE_EYE;
@@ -2443,7 +2501,8 @@ struct FSubstratePixelHeader
 	bool IsSubstrateMaterial()									{ return this.GetMaterialMode() != HEADER_MATERIALMODE_NONE; }
 	void SetMaterialMode(uint MaterialMode)						{ this.State = (this.State & ~HEADER_MASK_MATERIALMODE) | (MaterialMode & HEADER_MASK_MATERIALMODE); }
 
-	void SetIsComplexSpecialMaterial(bool bIn)					{ this.State |= (bIn ? HEADER_COMPLEXSPECIALPATH_MASK : 0u); }
+	uint GetComplexPathMode()									{ return HEADER_GETCOMPLEXPATHMODE(this.State); }
+	void SetComplexPathMode(uint ComplexPathMode)				{ HEADER_SETCOMPLEXPATHMODE(this.State, ComplexPathMode); }
 
 #if SUBSTRATE_DEFERRED_SHADING
 	uint GetLocalBasesCount()									{ return GetMaterialMode() == HEADER_MATERIALMODE_SLAB_COMPLEX ? HEADER_GETSHAREDLOCALBASESCOUNT(PackedHeader) : 1; }
@@ -3625,6 +3684,7 @@ float3 SubstrateGetBSDFBaseColor(in FSubstrateBSDF BSDF)
 	case SUBSTRATE_BSDF_TYPE_HAIR: 				return HAIR_BASECOLOR(BSDF);
 	case SUBSTRATE_BSDF_TYPE_EYE: 				return EYE_DIFFUSEALBEDO(BSDF);
 	case SUBSTRATE_BSDF_TYPE_SINGLELAYERWATER:	return SLW_BASECOLOR(BSDF);
+	case SUBSTRATE_BSDF_TYPE_TOON:				return TOON_DIFFUSEALBEDO(BSDF);
 	}
 	return 0.f;
 }
@@ -3637,6 +3697,7 @@ float3 SubstrateGetBSDFDiffuseColor(in FSubstrateBSDF BSDF)
 	case SUBSTRATE_BSDF_TYPE_HAIR: 				return HAIR_BASECOLOR(BSDF);
 	case SUBSTRATE_BSDF_TYPE_EYE: 				return EYE_DIFFUSEALBEDO(BSDF);
 	case SUBSTRATE_BSDF_TYPE_SINGLELAYERWATER:	return SLW_BASECOLOR(BSDF);
+	case SUBSTRATE_BSDF_TYPE_TOON:				return TOON_DIFFUSEALBEDO(BSDF);
 	}
 	return 0.f;
 }
@@ -3648,7 +3709,8 @@ void SubstrateSetBSDFDiffuseColor(inout FSubstrateBSDF BSDF, float3 InColor)
 	case SUBSTRATE_BSDF_TYPE_SLAB: 				SLAB_DIFFUSEALBEDO(BSDF) = InColor; break;
 	case SUBSTRATE_BSDF_TYPE_HAIR: 				HAIR_BASECOLOR(BSDF) = InColor; break;
 	case SUBSTRATE_BSDF_TYPE_EYE: 				EYE_DIFFUSEALBEDO(BSDF) = InColor; break;
-	case SUBSTRATE_BSDF_TYPE_SINGLELAYERWATER: SLW_BASECOLOR(BSDF) = InColor; break;
+	case SUBSTRATE_BSDF_TYPE_SINGLELAYERWATER:	SLW_BASECOLOR(BSDF) = InColor; break;
+	case SUBSTRATE_BSDF_TYPE_TOON:				TOON_DIFFUSEALBEDO(BSDF) = InColor; break;
 	}
 }
 
@@ -3660,6 +3722,7 @@ float3 SubstrateGetBSDFSpecularF0(in FSubstrateBSDF BSDF)
 	case SUBSTRATE_BSDF_TYPE_HAIR: 				return DielectricSpecularToF0(HAIR_SPECULAR(BSDF));
 	case SUBSTRATE_BSDF_TYPE_EYE: 				return EYE_F0(BSDF);
 	case SUBSTRATE_BSDF_TYPE_SINGLELAYERWATER:	return DielectricSpecularToF0(SLW_SPECULAR(BSDF));
+	case SUBSTRATE_BSDF_TYPE_TOON:				return TOON_SPECULARCOLOR(BSDF);
 	}
 	return 0.f;
 }
@@ -3684,6 +3747,7 @@ void SubstrateSetBSDFSpecularF0(in FSubstrateBSDF BSDF, float3 F0)
 	case SUBSTRATE_BSDF_TYPE_HAIR: 				HAIR_SPECULAR(BSDF) = F0RGBToDielectricSpecular(F0); break;
 	case SUBSTRATE_BSDF_TYPE_EYE: 				/* EYE_F0 us hardcoded */ break;
 	case SUBSTRATE_BSDF_TYPE_SINGLELAYERWATER:	SLW_SPECULAR(BSDF) = F0RGBToDielectricSpecular(F0); break;
+	case SUBSTRATE_BSDF_TYPE_TOON:				TOON_SPECULARCOLOR(BSDF) = F0; break;
 	}
 }
 
@@ -3776,9 +3840,10 @@ float SubstrateGetBSDFRoughness(in FSubstrateBSDF BSDF)
 			#endif
 			return Out;
 		}
-	case SUBSTRATE_BSDF_TYPE_HAIR: 			return HAIR_ROUGHNESS(BSDF);
-	case SUBSTRATE_BSDF_TYPE_EYE: 				return EYE_ROUGHNESS(BSDF);
-	case SUBSTRATE_BSDF_TYPE_SINGLELAYERWATER: return SLW_ROUGHNESS(BSDF);
+	case SUBSTRATE_BSDF_TYPE_HAIR:				return HAIR_ROUGHNESS(BSDF);
+	case SUBSTRATE_BSDF_TYPE_EYE:				return EYE_ROUGHNESS(BSDF);
+	case SUBSTRATE_BSDF_TYPE_SINGLELAYERWATER:	return SLW_ROUGHNESS(BSDF);
+	case SUBSTRATE_BSDF_TYPE_TOON:				return TOON_ROUGHNESS(BSDF);
 	}
 	return 0.f;
 }
@@ -3787,10 +3852,11 @@ void FSubstrateBSDF::SubstrateSetBSDFRoughness(in float Roughness)
 {
 	switch (BSDF_GETTYPE(this))
 	{
-	case SUBSTRATE_BSDF_TYPE_SLAB: 			SLAB_ROUGHNESS(this) = Roughness; break;
-	case SUBSTRATE_BSDF_TYPE_HAIR: 			HAIR_ROUGHNESS(this) = Roughness; break;
-	case SUBSTRATE_BSDF_TYPE_EYE: 				EYE_ROUGHNESS(this)  = Roughness; break;
-	case SUBSTRATE_BSDF_TYPE_SINGLELAYERWATER: SLW_ROUGHNESS(this)  = Roughness; break;
+	case SUBSTRATE_BSDF_TYPE_SLAB:				SLAB_ROUGHNESS(this) = Roughness; break;
+	case SUBSTRATE_BSDF_TYPE_HAIR:				HAIR_ROUGHNESS(this) = Roughness; break;
+	case SUBSTRATE_BSDF_TYPE_EYE:				EYE_ROUGHNESS(this)  = Roughness; break;
+	case SUBSTRATE_BSDF_TYPE_SINGLELAYERWATER:	SLW_ROUGHNESS(this)  = Roughness; break;
+	case SUBSTRATE_BSDF_TYPE_TOON:				TOON_ROUGHNESS(this)  = Roughness; break;
 	}
 }
 
@@ -3798,10 +3864,11 @@ float SubstrateGetBSDFAnisotropy(in FSubstrateBSDF BSDF)
 {
 	switch (BSDF_GETTYPE(BSDF))
 	{
-	case SUBSTRATE_BSDF_TYPE_SLAB: 			return BSDF_GETHASANISOTROPY(BSDF) ? SLAB_ANISOTROPY(BSDF) : 0.f;
-	case SUBSTRATE_BSDF_TYPE_HAIR: 			return 0;
-	case SUBSTRATE_BSDF_TYPE_EYE: 				return 0;
-	case SUBSTRATE_BSDF_TYPE_SINGLELAYERWATER: return 0;
+	case SUBSTRATE_BSDF_TYPE_SLAB:				return BSDF_GETHASANISOTROPY(BSDF) ? SLAB_ANISOTROPY(BSDF) : 0.f;
+	case SUBSTRATE_BSDF_TYPE_HAIR:				return 0;
+	case SUBSTRATE_BSDF_TYPE_EYE:				return 0;
+	case SUBSTRATE_BSDF_TYPE_SINGLELAYERWATER:	return 0;
+	case SUBSTRATE_BSDF_TYPE_TOON:				return BSDF_GETHASANISOTROPY(BSDF) ? TOON_ANISOTROPY(BSDF) : 0.0f;
 	}
 	return 0.f;
 }
@@ -4238,7 +4305,7 @@ uint PackSubstrateHeader(uint InBSDFCount, FSubstratePixelHeader InHeader)
 #endif
 
 #if SUBSTRATE_COMPLEXPATH
-	Out |= InHeader.State & HEADER_COMPLEXSPECIALPATH_MASK;
+	Out |= InHeader.State & HEADER_COMPLEXPATHMODE_MASK;
 #endif
 
 	HEADER_SETCOMMONSTATES(Out, InHeader.State);
@@ -4276,7 +4343,7 @@ FSubstratePixelHeader UnpackSubstrateHeaderIn(uint PackedHeader, inout FSubstrat
 		uint SharedLocalBasesCount = HEADER_GETSHAREDLOCALBASESCOUNT(PackedHeader);
 		uint SharedLocalBasesTypes = HEADER_GETSHAREDLOCALBASESTYPE(PackedHeader);
 
-		Out.SetIsComplexSpecialMaterial((PackedHeader & HEADER_COMPLEXSPECIALPATH_MASK) > 0);
+		Out.SetComplexPathMode(HEADER_GETCOMPLEXPATHMODE(PackedHeader));
 
 	#if SUBSTRATE_DEFERRED_SHADING
 		Out.PackedHeader = PackedHeader;
@@ -4733,11 +4800,24 @@ FSubstrateBSDF UnpackSubstrateBSDFIn(FSubstrateMaterialContainer SubstrateBuffer
 			}
 
 			#if SUBSTRATE_DEFERRED_SHADING
-			OutBSDF.UnpackSingleOrComplexPathSlabBSDF(DataX, DataY);
+			if (SubstrateHeader.GetComplexPathMode() == COMPLEXPATH_MODE_TOON)
+			{
+				float4 DiffuseRoughness		= UnpackColorGamma2ToLinearAlphaLinear(DataX);
+				float4 SpecularAnisotropy	= UnpackColorGamma2ToLinearAlphaLinear(DataY);
+	
+				TOON_DIFFUSEALBEDO(OutBSDF)	= DiffuseRoughness.xyz;
+				TOON_ROUGHNESS(OutBSDF)		= DiffuseRoughness.w;
+				TOON_SPECULARCOLOR(OutBSDF)	= SpecularAnisotropy.xyz;
+				TOON_ANISOTROPY(OutBSDF)	= SpecularAnisotropy.w;
+				return OutBSDF;
+			}
+			else
+			{
+				OutBSDF.UnpackSingleOrComplexPathSlabBSDF(DataX, DataY);
+			}
 			#endif
 		}
 
-
 		SLAB_F90(OutBSDF) = 1.0f;
 		const bool bHasF90 = BSDF_GETHASF90(OutBSDF);
 		if (bHasF90 || BSDF_GETHASHAZINESS(OutBSDF))
diff --git a/Engine/Shaders/Private/Substrate/SubstrateEvaluation.ush b/Engine/Shaders/Private/Substrate/SubstrateEvaluation.ush
--- a/Engine/Shaders/Private/Substrate/SubstrateEvaluation.ush
+++ b/Engine/Shaders/Private/Substrate/SubstrateEvaluation.ush
@@ -1304,6 +1304,27 @@ FSubstrateEvaluateResult SubstrateEvaluateBSDFCommon(FSubstrateBSDFContext BSDFC
 		}
 		break;
 
+		case SUBSTRATE_BSDF_TYPE_TOON:
+		{
+			const float3 DiffuseColor		= TOON_DIFFUSEALBEDO(BSDFContext.BSDF);
+			const float3 SpecularColor		= TOON_SPECULARCOLOR(BSDFContext.BSDF);
+
+			Sample.DiffuseColor				= DiffuseColor;
+
+			// TOON_TODO Hardcoded Toon Lambert diffuse before we implement the profile.
+			Sample.DiffusePathValue			= Diffuse_Lambert(Sample.DiffuseColor) * floor(saturate(AreaLight.NoL) * 10)*0.1;
+			Sample.IntegratedDiffuseValue	= ShadowTerms.SurfaceShadow * AreaLight.Falloff * AreaLight.FalloffColor * Sample.DiffusePathValue;
+
+			Sample.DiffusePDF				= BSDFContext.SatNoL / PI;
+			Sample.bPostProcessSubsurface	= false;
+
+			Sample.EmissivePathValue		= BSDF_GETEMISSIVE(BSDFContext.BSDF);
+
+			Sample.ThroughputV				= OpaqueBSDFThroughput;
+			Sample.TransmittanceAlongN		= OpaqueBSDFThroughput;
+		}
+		break;
+
 		//case SUBSTRATE_BSDF_TYPE_VOLUMETRICFOGCLOUD:
 		//case SUBSTRATE_BSDF_TYPE_UNLIT:
 		//case SUBSTRATE_BSDF_TYPE_SINGLELAYERWATER:
@@ -1540,6 +1561,20 @@ FSubstrateImportanceSampleResult SubstrateImportanceSampleBSDF(FSubstrateBSDFCon
 			break;
 		}
 
+		case SUBSTRATE_BSDF_TYPE_TOON:
+		{
+			const float3 DiffuseColor		= TOON_DIFFUSEALBEDO(BSDFContext.BSDF);
+			const float3 SpecularColor		= TOON_SPECULARCOLOR(BSDFContext.BSDF);
+
+			float4 ImportanceSample	= CosineSampleHemisphere(E);
+			Out.L					= mul(ImportanceSample.xyz, BSDFContext.TangentBasis);
+			Out.PDF					= ImportanceSample.w;
+			Out.Term				= SHADING_TERM_DIFFUSE;
+			Out.Weight				= 1.0f;
+			break;
+		}
+		break;
+
 		//case SUBSTRATE_BSDF_TYPE_VOLUMETRICFOGCLOUD:
 		//case SUBSTRATE_BSDF_TYPE_UNLIT:
 		//Nothing to do in this case because these BSDF are evaluated in other specialised passes.
@@ -1959,6 +1994,19 @@ FSubstrateEnvLightResult SubstrateEvaluateForEnvLight(FSubstrateBSDFContext BSDF
 			break;
 		}
 
+		case SUBSTRATE_BSDF_TYPE_TOON:
+		{
+			const float3 DiffuseColor						= TOON_DIFFUSEALBEDO(BSDFContext.BSDF);
+			const float3 SpecularColor						= TOON_SPECULARCOLOR(BSDFContext.BSDF);
+
+			SubstrateEnvLightResult.DiffuseNormal			= BSDFContext.N;
+			SubstrateEnvLightResult.bPostProcessSubsurface	= false;
+
+			SubstrateEnvLightResult.DiffuseWeight			= DiffuseColor;
+			SubstrateEnvLightResult.DiffuseColor			= DiffuseColor;
+		}
+		break;
+
 		//case SUBSTRATE_BSDF_TYPE_VOLUMETRICFOGCLOUD:
 		//case SUBSTRATE_BSDF_TYPE_UNLIT:
 		//case SUBSTRATE_BSDF_TYPE_SINGLELAYERWATER:
diff --git a/Engine/Shaders/Private/Substrate/SubstrateExport.ush b/Engine/Shaders/Private/Substrate/SubstrateExport.ush
--- a/Engine/Shaders/Private/Substrate/SubstrateExport.ush
+++ b/Engine/Shaders/Private/Substrate/SubstrateExport.ush
@@ -239,7 +239,7 @@ void PackSubstrateOut(
 						const bool bSSSMask = BSDF.bIsBottom;	// SSS and Thin can only be on the bottom layer, so ignoring bIsSimpleVolume here.
 						SubstratePixelHeader.SetHasSubsurface(bSSSMask && bHasSSS && !bIsSimpleVolume);
 
-						SubstratePixelHeader.SetIsComplexSpecialMaterial(bForceComplexSpecialPath);
+						SubstratePixelHeader.SetComplexPathMode(bForceComplexSpecialPath ? COMPLEXPATH_MODE_SLABSPECIAL : COMPLEXPATH_MODE_SLAB);
 						break;
 					}
 
@@ -294,6 +294,25 @@ void PackSubstrateOut(
 						SubstratePixelHeader.SetMaterialMode(HEADER_MATERIALMODE_SLWATER);
 						break;
 					}
+
+					case SUBSTRATE_BSDF_TYPE_TOON:
+					{
+						bIsOnlySlab	= false;
+						bIsSimpleMaterial = false;
+						bIsSingleMaterial = false;
+						OutEmissiveLuminance += BSDF_GETEMISSIVE(BSDF) * BSDF.LuminanceWeightV;
+
+						TopLayerTotalWeight				+= BSDF.TopLayerDataWeight;
+						TopLayerData.Roughness			+= BSDF.TopLayerDataWeight * TOON_ROUGHNESS(BSDF);
+						TopLayerData.UnlitViewBaseColor += BSDF.TopLayerDataWeight * SubstrateGetBSDFBaseColor(BSDF);
+
+						TopLayerData.Material = TOP_LAYER_MATERIAL_VALID;
+
+						SubstratePixelHeader.SetMaterialMode(HEADER_MATERIALMODE_SLAB_COMPLEX); // Use the complex path for Toon for now.
+						SubstratePixelHeader.SetComplexPathMode(COMPLEXPATH_MODE_TOON);
+						break;
+					}
+
 					}
 
 					float3x3 TangentBasis = SubstrateGetBSDFSharedBasis_InlineShading(SubstratePixelHeader, BSDF);
@@ -856,6 +875,13 @@ void PackSubstrateOut(
 							}
 						}
 						break;
+						case SUBSTRATE_BSDF_TYPE_TOON:
+						{
+							SUBSTRATE_STORE_UINT1(PackColorLinearToGamma2AlphaLinear(float4(TOON_DIFFUSEALBEDO(BSDF), TOON_ROUGHNESS(BSDF))));
+							SUBSTRATE_STORE_UINT1(PackColorLinearToGamma2AlphaLinear(float4(TOON_SPECULARCOLOR(BSDF), TOON_ANISOTROPY(BSDF))));
+							// 8 bytes
+						}
+						break;
 						//case SUBSTRATE_BSDF_TYPE_SINGLELAYERWATER: 
 						//{
 						//	// SLW always use the path defined by bHasFastWaterEncoding.
@@ -1212,6 +1238,20 @@ FExportResult SubstrateMaterialExportOut(
 				#endif
 				}
 			}
+			else
+#endif
+#if MATERIAL_SHADINGMODEL_TOON || !MATERIAL_IS_SUBSTRATEHIDDENCONVERSION
+				if (BSDFType == SUBSTRATE_BSDF_TYPE_TOON)
+			{
+				Out.ShadingModelID = SHADINGMODELID_SUBSTRATE_TOON;
+
+				// DiffuseColor in BaseColor
+				// Roughness in Roughness
+				Out.Metallic = 0.0;
+				Out.Specular = 1.0;
+				Out.CustomData = float4(TOON_SPECULARCOLOR(CurrentBSDF), TOON_ANISOTROPY(CurrentBSDF));
+			}
+			else
 #endif
 			// In case of an else, we need to provide a code block.
 			{}
diff --git a/Engine/Shaders/Private/Substrate/SubstrateRead.ush b/Engine/Shaders/Private/Substrate/SubstrateRead.ush
--- a/Engine/Shaders/Private/Substrate/SubstrateRead.ush
+++ b/Engine/Shaders/Private/Substrate/SubstrateRead.ush
@@ -113,6 +113,16 @@ FSubstrateGBufferBSDF SubstrateReadGBufferBSDF(in FScreenSpaceData ScreenSpaceDa
 		#endif
 			BSDF = GetSubstrateEyeBSDF(DiffuseAlbedo, F0, GBuffer.Roughness, IrisMask, IrisDistance, IrisNormal, IrisPlaneNormal, SSSProfileID, Emissive, SharedLocalBasisIndex).InlinedBSDF;
 		}
+		else if (GBuffer.ShadingModelID == SHADINGMODELID_SUBSTRATE_TOON)
+		{
+			float3 SpecularColor= GBuffer.CustomData.rgb;
+			float Anisotropy	= GBuffer.CustomData.a;
+
+			BSDF = GetSubstrateToonBSDF(
+				GBuffer.BaseColor, SpecularColor, GBuffer.Roughness, Anisotropy,
+				Emissive,
+				SharedLocalBasisIndex, SharedLocalBasisTypes).InlinedBSDF;
+		}
 		else
 	#endif
 		{
diff --git a/Engine/Shaders/Private/Substrate/SubstrateVisualizeCommon.ush b/Engine/Shaders/Private/Substrate/SubstrateVisualizeCommon.ush
--- a/Engine/Shaders/Private/Substrate/SubstrateVisualizeCommon.ush
+++ b/Engine/Shaders/Private/Substrate/SubstrateVisualizeCommon.ush
@@ -172,6 +172,11 @@ void SubstratePrintBSDF(
 		}
 		break;
 
+		case SUBSTRATE_BSDF_TYPE_TOON:
+		{
+		}
+		break;
+
 	default:
 		{
 			Print(Context, TEXT("Error - Uknown BSDF type"), FontRed);
diff --git a/Engine/Shaders/Shared/SubstrateDefinitions.h b/Engine/Shaders/Shared/SubstrateDefinitions.h
--- a/Engine/Shaders/Shared/SubstrateDefinitions.h
+++ b/Engine/Shaders/Shared/SubstrateDefinitions.h
@@ -36,6 +36,8 @@
 #define SUBSTRATE_BSDF_TYPE_HAIR						3
 #define SUBSTRATE_BSDF_TYPE_SINGLELAYERWATER			4
 #define SUBSTRATE_BSDF_TYPE_EYE							5
+#define SUBSTRATE_BSDF_TYPE_TOON						6
+//define SUBSTRATE_BSDF_TYPE_CUSTOM						7 Reserved for Custom shading model
 // When more than 7 BSDF must exists, please update STATE_BIT_COUNT_BSDF and FSubstrateClassification.ShadingModels packing in Substrate.ush
 
 // The size of Substrate material classification tiles on screen
diff --git a/Engine/Shaders/Shared/SubstrateVisualizeDefinitions.h b/Engine/Shaders/Shared/SubstrateVisualizeDefinitions.h
--- a/Engine/Shaders/Shared/SubstrateVisualizeDefinitions.h
+++ b/Engine/Shaders/Shared/SubstrateVisualizeDefinitions.h
@@ -678,6 +678,19 @@ FSubstratePixelDebugData ConvertToSerializableSubstratePixelDebugData(uint2 InCo
 		}
 		break;
 
+		case SUBSTRATE_BSDF_TYPE_TOON:
+		{
+			Data.Closures[i].DiffuseR	= TOON_DIFFUSEALBEDO(BSDF).r;
+			Data.Closures[i].DiffuseG	= TOON_DIFFUSEALBEDO(BSDF).g;
+			Data.Closures[i].DiffuseB	= TOON_DIFFUSEALBEDO(BSDF).b;
+			Data.Closures[i].F0R		= TOON_SPECULARCOLOR(BSDF).r;
+			Data.Closures[i].F0G		= TOON_SPECULARCOLOR(BSDF).g;
+			Data.Closures[i].F0B		= TOON_SPECULARCOLOR(BSDF).b;
+			Data.Closures[i].Roughness	= TOON_ROUGHNESS(BSDF);
+			Data.Closures[i].Anisotropy	= TOON_ANISOTROPY(BSDF);
+		}
+		break;
+
 		default:
 		{
 			// Error
diff --git a/Engine/Source/Editor/MaterialEditor/Private/MaterialNodes/SGraphSubstrateMaterial.cpp b/Engine/Source/Editor/MaterialEditor/Private/MaterialNodes/SGraphSubstrateMaterial.cpp
--- a/Engine/Source/Editor/MaterialEditor/Private/MaterialNodes/SGraphSubstrateMaterial.cpp
+++ b/Engine/Source/Editor/MaterialEditor/Private/MaterialNodes/SGraphSubstrateMaterial.cpp
@@ -187,7 +187,7 @@ static const TSharedRef<SWidget> InternalProcessOperator(
 
 			FString BSDFDesc = OutputType == ESubstrateWidgetOutputType::Node ? 
 											 FString(TEXT("BSDF")) : 
-											 FString::Printf(TEXT("BSDF (%s%s%s%s%s%s%s%s%s%s)")
+											 FString::Printf(TEXT("BSDF (%s%s%s%s%s%s%s%s%s%s%s)")
 											, Op.Has(ESubstrateBsdfFeature::EdgeColor) ? TEXT("F90 ") : TEXT("")
 											, Op.Has(ESubstrateBsdfFeature::SSS) ? TEXT("SSS ") : TEXT("")
 											, Op.Has(ESubstrateBsdfFeature::MFPPluggedIn) ? TEXT("MFP ") : TEXT("")
@@ -198,6 +198,7 @@ static const TSharedRef<SWidget> InternalProcessOperator(
 											, Op.Has(ESubstrateBsdfFeature::SpecularProfile) ? TEXT("Spc ") : TEXT("")
 											, Op.Has(ESubstrateBsdfFeature::Eye) ? *GetEyeDesc() : TEXT("")
 											, Op.Has(ESubstrateBsdfFeature::Hair) ? TEXT("Hai ") : TEXT("")
+											, Op.Has(ESubstrateBsdfFeature::Toon) ? TEXT("Toon ") : TEXT("")
 											);
 
 			static FString ToolTip;
diff --git a/Engine/Source/Editor/MaterialEditor/Private/SMaterialEditorSubstrateWidget.cpp b/Engine/Source/Editor/MaterialEditor/Private/SMaterialEditorSubstrateWidget.cpp
--- a/Engine/Source/Editor/MaterialEditor/Private/SMaterialEditorSubstrateWidget.cpp
+++ b/Engine/Source/Editor/MaterialEditor/Private/SMaterialEditorSubstrateWidget.cpp
@@ -502,7 +502,7 @@ static FString SubstrateMaterialTypeToString(uint8 InMaterialType)
 	{
 		case SUBSTRATE_MATERIAL_TYPE_SIMPLE:  			return FString::Printf(TEXT("SIMPLE  (Diffuse, albedo, roughness)\r\n"));
 		case SUBSTRATE_MATERIAL_TYPE_SINGLE:  			return FString::Printf(TEXT("SINGLE  (BSDF all features except anisotropy)\r\n"));
-		case SUBSTRATE_MATERIAL_TYPE_COMPLEX:  			return FString::Printf(TEXT("COMPLEX (Anisotropy, specular profile, multi-slabs)\r\n"));
+		case SUBSTRATE_MATERIAL_TYPE_COMPLEX:  			return FString::Printf(TEXT("COMPLEX (Anisotropy, multi-slabs, Toon)\r\n"));
 		case SUBSTRATE_MATERIAL_TYPE_COMPLEX_SPECIAL:  	return FString::Printf(TEXT("COMPLEX SPECIAL (Glints)\r\n"));
 		default: 										return FString::Printf(TEXT("UNKOWN => ERROR!\r\n"));
 	}
@@ -690,6 +690,7 @@ void SMaterialEditorSubstrateWidget::Tick(const FGeometry& AllottedGeometry, con
 				case SUBSTRATE_MATERIAL_TYPE_POSTPROCESS:			MaterialType = TEXT("PostProcess"); break;
 				case SUBSTRATE_MATERIAL_TYPE_UI:					MaterialType = TEXT("UI"); break;
 				case SUBSTRATE_MATERIAL_TYPE_DECAL:					MaterialType = TEXT("Decal"); break;
+				case SUBSTRATE_MATERIAL_TYPE_TOON:					MaterialType = TEXT("Toon"); break;
 				}
 
 				FString MaterialTypeDetail;
diff --git a/Engine/Source/Runtime/Engine/Classes/Engine/EngineTypes.h b/Engine/Source/Runtime/Engine/Classes/Engine/EngineTypes.h
--- a/Engine/Source/Runtime/Engine/Classes/Engine/EngineTypes.h
+++ b/Engine/Source/Runtime/Engine/Classes/Engine/EngineTypes.h
@@ -827,6 +827,7 @@ enum ESubstrateShadingModel : int
 	SSM_PostProcess,
 	SSM_Decal,
 	SSM_UI,
+	SSM_Toon,
 	/** Number of unique shading models. */
 	SSM_NUM,
 };
diff --git a/Engine/Source/Runtime/Engine/Private/Materials/HLSLMaterialTranslator.cpp b/Engine/Source/Runtime/Engine/Private/Materials/HLSLMaterialTranslator.cpp
--- a/Engine/Source/Runtime/Engine/Private/Materials/HLSLMaterialTranslator.cpp
+++ b/Engine/Source/Runtime/Engine/Private/Materials/HLSLMaterialTranslator.cpp
@@ -3005,6 +3005,11 @@ void FHLSLMaterialTranslator::GetMaterialEnvironment(EShaderPlatform InPlatform,
 			OutEnvironment.SetDefine(TEXT("MATERIAL_SHADINGMODEL_EYE"), TEXT("1"));
 		}
 
+		if (EnumHasAnyFlags(EnvironmentDefines->SubstrateMaterialBsdfFeatures, ESubstrateBsdfFeature::Toon))
+		{
+			OutEnvironment.SetDefine(TEXT("MATERIAL_SHADINGMODEL_TOON"), TEXT("1"));
+		}
+
 		if (EnvironmentDefines->bMaterialShadingModelEyeUseCurvature)
 		{
 			OutEnvironment.SetDefine(TEXT("MATERIAL_SHADINGMODEL_EYE_USE_CURVATURE"), TEXT("1"));
@@ -13534,6 +13539,7 @@ bool FHLSLMaterialTranslator::FSubstrateCompilationContext::SubstrateGenerateDer
 		bool bHasHair = false;
 		bool bHasEye = false;
 		bool bHasSLW = false;
+		bool bHasToon = false;
 		bool bHasSlab = false;
 		bool bHasUI = false;
 		bool bHasDecal = false;
@@ -13604,6 +13610,7 @@ bool FHLSLMaterialTranslator::FSubstrateCompilationContext::SubstrateGenerateDer
 					bHasHair |= CurrentOperator.BSDFType == SUBSTRATE_BSDF_TYPE_HAIR;
 					bHasEye |= CurrentOperator.BSDFType == SUBSTRATE_BSDF_TYPE_EYE;
 					bHasSLW |= CurrentOperator.BSDFType == SUBSTRATE_BSDF_TYPE_SINGLELAYERWATER;
+					bHasToon |= CurrentOperator.BSDFType == SUBSTRATE_BSDF_TYPE_TOON;
 					bHasSlab |= CurrentOperator.BSDFType == SUBSTRATE_BSDF_TYPE_SLAB;
 					break;
 				}
@@ -13632,7 +13639,7 @@ bool FHLSLMaterialTranslator::FSubstrateCompilationContext::SubstrateGenerateDer
 			}
 			bSubstrateMaterialIsUnlitNode = bHasUnlit;
 
-			if ((bHasVFogCloud || bHasHair || bHasEye || bHasSLW) && bOperatorEncountered)
+			if ((bHasVFogCloud || bHasHair || bHasEye || bHasSLW || bHasToon) && bOperatorEncountered)
 			{
 				Compiler->Errorf(TEXT("Fog/Cloud, Hair, Eye or SingleLayerWater cannot be used with operators. See %s (asset: %s).\r\n"), *CompilerMaterial->GetDebugName(), *CompilerMaterial->GetAssetPath().ToString());
 				return false;
@@ -13652,7 +13659,7 @@ bool FHLSLMaterialTranslator::FSubstrateCompilationContext::SubstrateGenerateDer
 				return false;
 			}
 
-			if ((bHasUnlit || bHasVFogCloud || bHasHair || bHasEye || bHasSLW) && SubstrateMaterialEffectiveClosureCount > 1)
+			if ((bHasUnlit || bHasVFogCloud || bHasHair || bHasEye || bHasSLW || bHasToon) && SubstrateMaterialEffectiveClosureCount > 1)
 			{
 				Compiler->Errorf(TEXT("Unlit, Fog/Cloud, Hair or SingleLayerWater must result in a single one Closure (Found %d). See %s (asset: %s).\r\n"), SubstrateMaterialEffectiveClosureCount, *CompilerMaterial->GetDebugName(), *CompilerMaterial->GetAssetPath().ToString());
 				// Even though technically we could support Unlit parameter blending with Slab.
@@ -13942,6 +13949,13 @@ bool FHLSLMaterialTranslator::FSubstrateCompilationContext::SubstrateGenerateDer
 						bIsFastWaterPath		= true;
 						break;
 					}
+					case SUBSTRATE_BSDF_TYPE_TOON:
+					{
+						// Use the complex path for now.
+						SubstrateMaterialComplexity.bIsSimple = false;
+						SubstrateMaterialComplexity.bIsSingle = false;
+						break;
+					}
 					}
 					break;
 				}
@@ -14118,6 +14132,12 @@ bool FHLSLMaterialTranslator::FSubstrateCompilationContext::SubstrateGenerateDer
 						Compiler->Errorf(TEXT("Substrate error: single layer water should go through through a dedicated fast path in %s (asset: %s).\r\n"), *CompilerMaterial->GetDebugName(), *CompilerMaterial->GetAssetPath().ToString());
 						return false;
 					}
+					case SUBSTRATE_BSDF_TYPE_TOON:
+					{
+						SubstrateMaterialRequestedSizeByte += UintByteSize;
+						SubstrateMaterialRequestedSizeByte += UintByteSize;
+						break;
+					}
 					case SUBSTRATE_BSDF_TYPE_UNLIT:
 					{
 						// Never stored, it goes directly into the scene as emitted luminance.
@@ -14226,6 +14246,10 @@ bool FHLSLMaterialTranslator::FSubstrateCompilationContext::SubstrateGenerateDer
 				{
 					Compiler->MaterialCompilationOutput.SubstrateMaterialCompilationOutput.MaterialType = SUBSTRATE_MATERIAL_TYPE_SINGLELAYERWATER;
 				}
+				else if (bHasToon)
+				{
+					Compiler->MaterialCompilationOutput.SubstrateMaterialCompilationOutput.MaterialType = SUBSTRATE_MATERIAL_TYPE_TOON;
+				}
 				else
 				{
 					Compiler->MaterialCompilationOutput.SubstrateMaterialCompilationOutput.MaterialType = SUBSTRATE_MATERIAL_TYPE_SINGLESLAB;
@@ -14898,6 +14922,29 @@ int32 FHLSLMaterialTranslator::SubstrateSingleLayerWaterBSDF(
 	);
 }
 
+int32 FHLSLMaterialTranslator::SubstrateToonBSDF(
+	int32 DiffuseColor, int32 SpecularColor, int32 Roughness, int32 Anisotropy,
+	int32 EmissiveColor, int32 Normal,
+	const FString& SharedLocalBasisIndexMacro, FSubstrateOperator* PromoteToOperator)
+{
+	return AddCodeChunk(
+		MCT_Substrate, TEXT("Parameters.%s.PromoteParameterBlendedBSDFToOperator(GetSubstrateToonBSDF(%s, %s, %s, %s, %s, %s, Parameters.%s.Types), %u, %u, %u, %u) /* Normal:%s */"),
+		*GetParametersSubstrateTreeName(CurrentSubstrateCompilationContext),
+		*SubstrateGetCastParameterCode(DiffuseColor, MCT_Float3),
+		*SubstrateGetCastParameterCode(SpecularColor, MCT_Float3),
+		*SubstrateGetCastParameterCode(Roughness, MCT_Float),
+		*SubstrateGetCastParameterCode(Anisotropy, MCT_Float),
+		*SubstrateGetCastParameterCode(EmissiveColor, MCT_Float3),
+		*SharedLocalBasisIndexMacro,
+		*GetParametersSharedLocalBasesName(CurrentSubstrateCompilationContext),
+		PromoteToOperator->Index,
+		PromoteToOperator->BSDFIndex,
+		PromoteToOperator->LayerDepth,
+		PromoteToOperator->bIsBottom ? 1 : 0,
+		*GetParameterCode(Normal)
+	);
+}
+
 int32 FHLSLMaterialTranslator::SubstrateHorizontalMixing(int32 Background, int32 Foreground, int32 Mix, int OperatorIndex, uint32 MaxDistanceFromLeaves)
 {
 	if (Foreground == INDEX_NONE || Background == INDEX_NONE || Mix == INDEX_NONE)
@@ -16523,7 +16570,7 @@ void FHLSLMaterialTranslator::PrepareEnvironmentDefines()
 						{
 							if (!It.IsDiscarded() && It.MaxDistanceFromLeaves == DistanceToLeaves)
 							{
-								SubstrateMaterialContextDescription += FString::Printf(TEXT("\tIdx=%d Op=%s ParentIdx=%d LeftIndex=%d RightIndex=%d BSDFIdx=%d LayerDepth=%d IsTop=%d IsBot=%d BSDFType=%s SSS=%d MFP=%d F90=%d Rough2=%d Fuzz=%d Aniso=%d Glint=%d SpecularProfile=%d Eye=%d Hair=%d\n GUID=%s\n"),
+								SubstrateMaterialContextDescription += FString::Printf(TEXT("\tIdx=%d Op=%s ParentIdx=%d LeftIndex=%d RightIndex=%d BSDFIdx=%d LayerDepth=%d IsTop=%d IsBot=%d BSDFType=%s SSS=%d MFP=%d F90=%d Rough2=%d Fuzz=%d Aniso=%d Glint=%d SpecularProfile=%d Eye=%d Hair=%d Toon=%d\n GUID=%s\n"),
 									It.Index, GetSubstrateOperatorStr(It.OperatorType), It.ParentIndex, It.LeftIndex, It.RightIndex, It.BSDFIndex, It.LayerDepth, It.bIsTop, It.bIsBottom,
 									*GetSubstrateBSDFName(It.BSDFType), 
 									It.Has(ESubstrateBsdfFeature::SSS), 
@@ -16536,6 +16583,7 @@ void FHLSLMaterialTranslator::PrepareEnvironmentDefines()
 									It.Has(ESubstrateBsdfFeature::SpecularProfile), 
 									It.Has(ESubstrateBsdfFeature::Eye), 
 									It.Has(ESubstrateBsdfFeature::Hair),
+									It.Has(ESubstrateBsdfFeature::Toon),
 									*It.MaterialExpressionGuid.ToString());
 
 								// Aggregate all the features used within the material
diff --git a/Engine/Source/Runtime/Engine/Private/Materials/HLSLMaterialTranslator.h b/Engine/Source/Runtime/Engine/Private/Materials/HLSLMaterialTranslator.h
--- a/Engine/Source/Runtime/Engine/Private/Materials/HLSLMaterialTranslator.h
+++ b/Engine/Source/Runtime/Engine/Private/Materials/HLSLMaterialTranslator.h
@@ -1414,6 +1414,10 @@ protected:
 		int32 BaseColor, int32 Metallic, int32 Specular, int32 Roughness, 
 		int32 EmissiveColor, int32 TopMaterialOpacity, int32 WaterAlbedo, int32 WaterExtinction, int32 WaterPhaseG, 
 		int32 ColorScaleBehindWater, int32 Normal, const FString& SharedLocalBasisIndexMacro, FSubstrateOperator* PromoteToOperator) override;
+	virtual int32 SubstrateToonBSDF(
+		int32 DiffuseColor, int32 SpecularColor, int32 Roughness, int32 Anisotropy,
+		int32 EmissiveColor, int32 Normal, 
+		const FString& SharedLocalBasisIndexMacro, FSubstrateOperator* PromoteToOperator) override;
 	virtual int32 SubstrateHorizontalMixing(int32 Background, int32 Foreground, int32 Mix, int OperatorIndex, uint32 MaxDistanceFromLeaves) override;
 	virtual int32 SubstrateHorizontalMixingParameterBlending(int32 Background, int32 Foreground, int32 HorizontalMixCodeChunk, int32 NormalMixCodeChunk, const FString& SharedLocalBasisIndexMacro, int32 BackgroundBSDFNormalCodeChunk, int32 ForegroundBSDFNormalCodeChunk, FSubstrateOperator* PromoteToOperator) override;
 	virtual int32 SubstrateVerticalLayering(int32 Top, int32 Base, int32 Thickness, int OperatorIndex, uint32 MaxDistanceFromLeaves) override;
diff --git a/Engine/Source/Runtime/Engine/Private/Materials/MaterialExpressionSubstrate.cpp b/Engine/Source/Runtime/Engine/Private/Materials/MaterialExpressionSubstrate.cpp
--- a/Engine/Source/Runtime/Engine/Private/Materials/MaterialExpressionSubstrate.cpp
+++ b/Engine/Source/Runtime/Engine/Private/Materials/MaterialExpressionSubstrate.cpp
@@ -2423,6 +2423,118 @@ FSubstrateOperator* UMaterialExpressionSubstrateSingleLayerWaterBSDF::SubstrateG
 
 
 
+
+
+
+
+
+UMaterialExpressionSubstrateToonBSDF::UMaterialExpressionSubstrateToonBSDF(const FObjectInitializer& ObjectInitializer)
+	: Super(ObjectInitializer)
+{
+	struct FConstructorStatics
+	{
+		FText NAME_Substrate;
+		FConstructorStatics() : NAME_Substrate(LOCTEXT("Substrate BSDFs", "Substrate BSDFs")) {}
+	};
+	static FConstructorStatics ConstructorStatics;
+#if WITH_EDITORONLY_DATA
+	MenuCategories.Add(ConstructorStatics.NAME_Substrate);
+#endif
+}
+
+#if WITH_EDITOR
+int32 UMaterialExpressionSubstrateToonBSDF::Compile(class FMaterialCompiler* Compiler, int32 OutputIndex)
+{
+	int32 NormalCodeChunk = CompileWithDefaultNormalWS(Compiler, Normal);
+	const FSubstrateRegisteredSharedLocalBasis NewRegisteredSharedLocalBasis = SubstrateCompilationInfoCreateSharedLocalBasis(Compiler, NormalCodeChunk);
+
+	FSubstrateOperator& SubstrateOperator = Compiler->SubstrateCompilationGetOperator(Compiler->SubstrateTreeStackGetPathUniqueId());
+	SubstrateOperator.BSDFRegisteredSharedLocalBasis = NewRegisteredSharedLocalBasis;
+
+	if (SubstrateOperator.bUseParameterBlending)
+	{
+		return Compiler->Errorf(TEXT("Substrate Toon BSDF node cannot be used with parameter blending."));
+	}
+	else if (SubstrateOperator.bRootOfParameterBlendingSubTree)
+	{
+		return Compiler->Errorf(TEXT("Substrate Toon BSDF node cannot be the root of a parameter blending sub tree."));
+	}
+
+	int32 OutputCodeChunk = Compiler->SubstrateToonBSDF(
+		CompileWithDefaultFloat3(Compiler, DiffuseColor, 0.0f, 0.0f, 0.0f),
+		CompileWithDefaultFloat1(Compiler, SpecularColor, 0.0f),
+		CompileWithDefaultFloat1(Compiler, Roughness, 0.5f),
+		CompileWithDefaultFloat1(Compiler, Anisotropy, 0.0f),
+		CompileWithDefaultFloat3(Compiler, EmissiveColor, 0.0f, 0.0f, 0.0f),
+		NormalCodeChunk,
+		Compiler->GetSubstrateSharedLocalBasisIndexMacro(NewRegisteredSharedLocalBasis),
+		&SubstrateOperator);
+
+	return OutputCodeChunk;
+}
+
+void UMaterialExpressionSubstrateToonBSDF::GetCaption(TArray<FString>& OutCaptions) const
+{
+	OutCaptions.Add(TEXT("Substrate Toon BSDF"));
+}
+
+EMaterialValueType UMaterialExpressionSubstrateToonBSDF::GetOutputValueType(int32 OutputIndex)
+{
+	return MCT_Substrate;
+}
+
+EMaterialValueType UMaterialExpressionSubstrateToonBSDF::GetInputValueType(int32 InputIndex)
+{
+	switch (InputIndex)
+	{
+	case 0:
+		return MCT_Float3; // BaseColor
+	case 1:
+		return MCT_Float3; // SpecularColor
+	case 2:
+		return MCT_Float1; // Roughness
+	case 3:
+		return MCT_Float1; // Anisotropy
+	case 4:
+		return MCT_Float3; // Normal
+	case 5:
+		return MCT_Float3; // Emissive Color
+	}
+
+	check(false);
+	return MCT_Float1;
+}
+
+bool UMaterialExpressionSubstrateToonBSDF::IsResultSubstrateMaterial(int32 OutputIndex)
+{
+	return true;
+}
+
+void UMaterialExpressionSubstrateToonBSDF::GatherSubstrateMaterialInfo(FSubstrateMaterialInfo& SubstrateMaterialInfo, int32 OutputIndex)
+{
+	SubstrateMaterialInfo.AddShadingModel(SSM_Toon);
+	SubstrateMaterialInfo.AddGuid(MaterialExpressionGuid);
+}
+
+FSubstrateOperator* UMaterialExpressionSubstrateToonBSDF::SubstrateGenerateMaterialTopologyTree(class FMaterialCompiler* Compiler, class UMaterialExpression* Parent, int32 OutputIndex)
+{
+	FSubstrateOperator& SubstrateOperator = Compiler->SubstrateCompilationRegisterOperator(SUBSTRATE_OPERATOR_BSDF, Compiler->SubstrateTreeStackGetPathUniqueId(), this->MaterialExpressionGuid, Parent, Compiler->SubstrateTreeStackGetParentPathUniqueId());
+	SubstrateOperator.BSDFType = SUBSTRATE_BSDF_TYPE_TOON;
+	SubstrateOperator.ThicknessIndex = Compiler->SubstrateThicknessStackGetThicknessIndex();
+	SubstrateOperator.bBSDFWritesEmissive = EmissiveColor.IsConnected();
+	SubstrateOperator.BSDFFeatures |= ESubstrateBsdfFeature::Toon;
+	if (Anisotropy.IsConnected())
+	{
+		SubstrateOperator.BSDFFeatures |= ESubstrateBsdfFeature::Anisotropy;
+	}
+	return &SubstrateOperator;
+}
+#endif // WITH_EDITOR
+
+
+
+
+
 UMaterialExpressionSubstrateHorizontalMixing::UMaterialExpressionSubstrateHorizontalMixing(const FObjectInitializer& ObjectInitializer)
 	: Super(ObjectInitializer)
 	, bUseParameterBlending(false)
diff --git a/Engine/Source/Runtime/Engine/Private/Materials/SubstrateMaterial.cpp b/Engine/Source/Runtime/Engine/Private/Materials/SubstrateMaterial.cpp
--- a/Engine/Source/Runtime/Engine/Private/Materials/SubstrateMaterial.cpp
+++ b/Engine/Source/Runtime/Engine/Private/Materials/SubstrateMaterial.cpp
@@ -22,6 +22,8 @@ FString GetSubstrateBSDFName(uint8 BSDFType)
 		return TEXT("EYE");
 	case SUBSTRATE_BSDF_TYPE_SINGLELAYERWATER:
 		return TEXT("SINGLELAYERWATER");
+	case SUBSTRATE_BSDF_TYPE_TOON:
+		return TEXT("TOON");
 	}
 	check(false);
 	return "";
diff --git a/Engine/Source/Runtime/Engine/Private/ShaderCompiler/ShaderGenerationUtil.cpp b/Engine/Source/Runtime/Engine/Private/ShaderCompiler/ShaderGenerationUtil.cpp
--- a/Engine/Source/Runtime/Engine/Private/ShaderCompiler/ShaderGenerationUtil.cpp
+++ b/Engine/Source/Runtime/Engine/Private/ShaderCompiler/ShaderGenerationUtil.cpp
@@ -130,6 +130,7 @@ void ApplyFetchEnvironmentInternal(FShaderMaterialPropertyDefines& SrcDefines, c
 	FETCH_COMPILE_BOOL(MATERIAL_SHADINGMODEL_EYE);
 	FETCH_COMPILE_BOOL(MATERIAL_SHADINGMODEL_SINGLELAYERWATER);
 	FETCH_COMPILE_BOOL(MATERIAL_SHADINGMODEL_THIN_TRANSLUCENT);
+	FETCH_COMPILE_BOOL(MATERIAL_SHADINGMODEL_TOON);
 
 	FETCH_COMPILE_BOOL(SINGLE_LAYER_WATER_SEPARATED_MAIN_LIGHT);
 
@@ -1879,6 +1880,12 @@ static void DetermineUsedMaterialSlots(
 		Slots[GBS_CustomData] = GetGBufferSlotUsage(bUseCustomData);
 	}
 
+	if (Mat.MATERIAL_SHADINGMODEL_TOON)
+	{
+		SetStandardGBufferSlots(Slots, bWriteEmissive, bHasTangent, bHasVelocity, bWritesVelocity, bHasStaticLighting, bIsSubstrateMaterial, bIsSubstrateNewGBuffer);
+		Slots[GBS_CustomData] = GetGBufferSlotUsage(bUseCustomData);
+	}
+
 	if (Mat.MATERIAL_SHADINGMODEL_SINGLELAYERWATER)
 	{
 		// single layer water uses standard slots
diff --git a/Engine/Source/Runtime/Engine/Public/MaterialCompiler.h b/Engine/Source/Runtime/Engine/Public/MaterialCompiler.h
--- a/Engine/Source/Runtime/Engine/Public/MaterialCompiler.h
+++ b/Engine/Source/Runtime/Engine/Public/MaterialCompiler.h
@@ -712,6 +712,10 @@ public:
 	virtual int32 SubstrateSingleLayerWaterBSDF(
 		int32 BaseColor, int32 Metallic, int32 Specular, int32 Roughness, int32 EmissiveColor, int32 TopMaterialOpacity, 
 		int32 WaterAlbedo, int32 WaterExtinction, int32 WaterPhaseG, int32 ColorScaleBehindWater, int32 Normal, const FString& SharedLocalBasisIndexMacro, FSubstrateOperator* PromoteToOperator) = 0;
+	virtual int32 SubstrateToonBSDF(
+		int32 DiffuseColor, int32 SpecularColor, int32 Roughness, int32 Anisotropy,
+		int32 EmissiveColor, int32 Normal,
+		const FString& SharedLocalBasisIndexMacro, FSubstrateOperator* PromoteToOperator) = 0;
 	virtual int32 SubstrateHorizontalMixing(int32 Background, int32 Foreground, int32 Mix, int OperatorIndex, uint32 MaxDistanceFromLeaves) = 0;
 	virtual int32 SubstrateHorizontalMixingParameterBlending(int32 Background, int32 Foreground, int32 HorizontalMixCodeChunk, int32 NormalMixCodeChunk, const FString& SharedLocalBasisIndexMacro, int32 BackgroundBSDFNormalCodeChunk, int32 ForegroundBSDFNormalCodeChunk, FSubstrateOperator* PromoteToOperator) = 0;
 	virtual int32 SubstrateVerticalLayering(int32 Top, int32 Base, int32 Thickness, int OperatorIndex, uint32 MaxDistanceFromLeaves) = 0;
@@ -1516,6 +1520,17 @@ public:
 			WaterAlbedo, WaterExtinction, WaterPhaseG, ColorScaleBehindWater, Normal, SharedLocalBasisIndexMacro, PromoteToOperator);
 	}
 
+	virtual int32 SubstrateToonBSDF(
+		int32 DiffuseColor, int32 SpecularColor, int32 Roughness, int32 Anisotropy,
+		int32 EmissiveColor, int32 Normal,
+		const FString& SharedLocalBasisIndexMacro, FSubstrateOperator* PromoteToOperator) override
+	{
+		return Compiler->SubstrateToonBSDF(
+			DiffuseColor, SpecularColor, Roughness, Anisotropy, 
+			EmissiveColor, Normal,
+			SharedLocalBasisIndexMacro, PromoteToOperator);
+	}
+
 	virtual int32 SubstrateHorizontalMixing(int32 Background, int32 Foreground, int32 Mix, int OperatorIndex, uint32 MaxDistanceFromLeaves) override
 	{
 		return Compiler->SubstrateHorizontalMixing(Background, Foreground, Mix, OperatorIndex, MaxDistanceFromLeaves);
diff --git a/Engine/Source/Runtime/Engine/Public/Materials/MaterialExpressionSubstrate.h b/Engine/Source/Runtime/Engine/Public/Materials/MaterialExpressionSubstrate.h
--- a/Engine/Source/Runtime/Engine/Public/Materials/MaterialExpressionSubstrate.h
+++ b/Engine/Source/Runtime/Engine/Public/Materials/MaterialExpressionSubstrate.h
@@ -819,6 +819,64 @@ class UMaterialExpressionSubstrateLightFunction : public UMaterialExpressionSubs
 	//~ End UMaterialExpression Interface
 };
 
+UCLASS(MinimalAPI, collapsecategories, hidecategories = Object, DisplayName = "Substrate Toon BSDF")
+class UMaterialExpressionSubstrateToonBSDF : public UMaterialExpressionSubstrateBSDF
+{
+	GENERATED_UCLASS_BODY()
+
+	/**
+	 * Surface base color. (type = float3, unit = unitless, defaults to white)
+	 */
+	UPROPERTY()
+	FExpressionInput DiffuseColor;
+
+	/**
+	 * Surface specular color. (type = float3, unit = unitless, defaults to white)
+	 */
+	UPROPERTY()
+	FExpressionInput SpecularColor;
+
+	/**
+	 * Controls how rough the Material is. Roughness of 0 (smooth) is a mirror reflection and 1 (rough) is completely matte or diffuse. When using anisotropy, it is the roughness used along the Tangent axis. (type = float, unit = unitless, defaults to 0.5)
+	 */
+	UPROPERTY()
+	FExpressionInput Roughness;
+
+	/**
+	 * Controls the anisotropy factor of the roughness. Positive value elongates the specular lobe along the Tangent vector, Negative value elongates the specular lobe along the perpendicular of the Tangent. (type = float, unit = unitless).
+	 */
+	UPROPERTY()
+	FExpressionInput Anisotropy;
+
+	/**
+	 * The normal of the surface (type = float3, unit = unitless, defaults to +Z vector)
+	 */
+	UPROPERTY()
+	FExpressionInput Normal;
+
+	/**
+	 * Emissive color on top of the surface (type = float3, unit = luminance, defaults to black)
+	 */
+	UPROPERTY()
+	FExpressionInput EmissiveColor;
+
+	/** Toon profile used to control the appearance of this BSDF. */
+//	UPROPERTY(EditAnywhere, BlueprintReadOnly, Category = Material, meta = (DisplayName = "Toon Profile"))
+//	TObjectPtr<class UToonProfile> ToonProfile;
+
+	//~ Begin UMaterialExpression Interface
+#if WITH_EDITOR
+	virtual int32 Compile(class FMaterialCompiler* Compiler, int32 OutputIndex) override;
+	virtual void GetCaption(TArray<FString>& OutCaptions) const override;
+	virtual EMaterialValueType GetOutputValueType(int32 OutputIndex) override;
+	virtual EMaterialValueType GetInputValueType(int32 InputIndex) override;
+	virtual bool IsResultSubstrateMaterial(int32 OutputIndex) override;
+	virtual void GatherSubstrateMaterialInfo(FSubstrateMaterialInfo& SubstrateMaterialInfo, int32 OutputIndex) override;
+	virtual FSubstrateOperator* SubstrateGenerateMaterialTopologyTree(class FMaterialCompiler* Compiler, class UMaterialExpression* Parent, int32 OutputIndex) override;
+#endif
+	//~ End UMaterialExpression Interface
+};
+
 UCLASS(MinimalAPI, collapsecategories, hidecategories = Object, DisplayName = "Substrate Post Process")
 class UMaterialExpressionSubstratePostProcess : public UMaterialExpressionSubstrateBSDF
 {
diff --git a/Engine/Source/Runtime/Engine/Public/Rendering/SubstrateMaterialShared.h b/Engine/Source/Runtime/Engine/Public/Rendering/SubstrateMaterialShared.h
--- a/Engine/Source/Runtime/Engine/Public/Rendering/SubstrateMaterialShared.h
+++ b/Engine/Source/Runtime/Engine/Public/Rendering/SubstrateMaterialShared.h
@@ -41,6 +41,7 @@ enum class ESubstrateBsdfFeature : uint16
 	EyeIrisNormalPluggedIn           = 1u<<9u,
 	EyeIrisTangentPluggedIn          = 1u<<10u,
 	Hair                             = 1u<<11u,
+	Toon                             = 1u<<12u,
 
 	// Complexity masks
 	SingleMask =	
@@ -54,7 +55,8 @@ enum class ESubstrateBsdfFeature : uint16
 		  Anisotropy
 		| SpecularProfile
 		| Eye
-		| Hair,
+		| Hair
+		| Toon,
 
 	ComplexSpecialMask = 
 		  Glint
@@ -172,6 +174,7 @@ public:
 #define SUBSTRATE_MATERIAL_TYPE_POSTPROCESS			8
 #define SUBSTRATE_MATERIAL_TYPE_UI					9
 #define SUBSTRATE_MATERIAL_TYPE_DECAL				10
+#define SUBSTRATE_MATERIAL_TYPE_TOON				11
 
 struct FSubstrateMaterialCompilationOutput
 {
diff --git a/Engine/Source/Runtime/RenderCore/Private/ShaderMaterialDerivedHelpers.cpp b/Engine/Source/Runtime/RenderCore/Private/ShaderMaterialDerivedHelpers.cpp
--- a/Engine/Source/Runtime/RenderCore/Private/ShaderMaterialDerivedHelpers.cpp
+++ b/Engine/Source/Runtime/RenderCore/Private/ShaderMaterialDerivedHelpers.cpp
@@ -46,7 +46,7 @@ FShaderMaterialDerivedDefines RENDERCORE_API CalculateDerivedMaterialParameters(
 	Dst.USES_GBUFFER = (FEATURE_LEVEL >= ERHIFeatureLevel::SM4_REMOVED && (Mat.MATERIALBLENDING_SOLID || Mat.MATERIALBLENDING_MASKED) && !SrcGlobal.FORWARD_SHADING);
 
 	// Only some shader models actually need custom data.
-	Dst.WRITES_CUSTOMDATA_TO_GBUFFER = (Dst.USES_GBUFFER && (Mat.MATERIAL_SHADINGMODEL_SUBSURFACE || Mat.MATERIAL_SHADINGMODEL_PREINTEGRATED_SKIN || Mat.MATERIAL_SHADINGMODEL_SUBSURFACE_PROFILE || Mat.MATERIAL_SHADINGMODEL_CLEAR_COAT || Mat.MATERIAL_SHADINGMODEL_TWOSIDED_FOLIAGE || Mat.MATERIAL_SHADINGMODEL_HAIR || Mat.MATERIAL_SHADINGMODEL_CLOTH || Mat.MATERIAL_SHADINGMODEL_EYE));
+	Dst.WRITES_CUSTOMDATA_TO_GBUFFER = (Dst.USES_GBUFFER && (Mat.MATERIAL_SHADINGMODEL_SUBSURFACE || Mat.MATERIAL_SHADINGMODEL_PREINTEGRATED_SKIN || Mat.MATERIAL_SHADINGMODEL_SUBSURFACE_PROFILE || Mat.MATERIAL_SHADINGMODEL_CLEAR_COAT || Mat.MATERIAL_SHADINGMODEL_TWOSIDED_FOLIAGE || Mat.MATERIAL_SHADINGMODEL_HAIR || Mat.MATERIAL_SHADINGMODEL_CLOTH || Mat.MATERIAL_SHADINGMODEL_EYE || Mat.MATERIAL_SHADINGMODEL_TOON));
 
 	// Based on GetPrecomputedShadowMasks()
 	// Note: WRITES_PRECSHADOWFACTOR_TO_GBUFFER is currently disabled because we use the precomputed shadow factor GBuffer outside of STATICLIGHTING_TEXTUREMASK to store UseSingleSampleShadowFromStationaryLights
diff --git a/Engine/Source/Runtime/RenderCore/Public/ShaderMaterial.h b/Engine/Source/Runtime/RenderCore/Public/ShaderMaterial.h
--- a/Engine/Source/Runtime/RenderCore/Public/ShaderMaterial.h
+++ b/Engine/Source/Runtime/RenderCore/Public/ShaderMaterial.h
@@ -119,6 +119,7 @@ struct FShaderMaterialPropertyDefines
 	uint8 MATERIAL_SHADINGMODEL_SINGLELAYERWATER : 1;
 	uint8 SINGLE_LAYER_WATER_SEPARATED_MAIN_LIGHT : 1;
 	uint8 MATERIAL_SHADINGMODEL_THIN_TRANSLUCENT : 1;
+	uint8 MATERIAL_SHADINGMODEL_TOON : 1;
 
 	uint8 TRANSLUCENCY_LIGHTING_VOLUMETRIC_NONDIRECTIONAL : 1;
 	uint8 TRANSLUCENCY_LIGHTING_VOLUMETRIC_DIRECTIONAL : 1;
diff --git a/Engine/Source/Runtime/Renderer/Private/Substrate/SubstrateVisualize.cpp b/Engine/Source/Runtime/Renderer/Private/Substrate/SubstrateVisualize.cpp
--- a/Engine/Source/Runtime/Renderer/Private/Substrate/SubstrateVisualize.cpp
+++ b/Engine/Source/Runtime/Renderer/Private/Substrate/SubstrateVisualize.cpp
@@ -530,6 +530,9 @@ void AddProcessAndPrintSubstrateMaterialPropertiesPasses(FRDGBuilder& GraphBuild
 					case SUBSTRATE_BSDF_TYPE_SINGLELAYERWATER:
 						Type = TEXT("SingleLayerWater");
 						break;
+					case SUBSTRATE_BSDF_TYPE_TOON:
+						Type = TEXT("Toon");
+						break;
 					case SUBSTRATE_BSDF_TYPE_EYE:
 						Type = TEXT("Eye");
 						break;
@@ -863,6 +866,24 @@ void AddProcessAndPrintSubstrateMaterialPropertiesPasses(FRDGBuilder& GraphBuild
 						PrintFSmall(DrawPosX + DataOffset, DrawPosY, Closure.SSSOpacity);
 						NewLine();
 					}
+					else if (Closure.Type == SUBSTRATE_BSDF_TYPE_TOON)
+					{
+						PrintS(DrawPosX + 0.0f, DrawPosY, TEXT("Diffuse Color"), Orange);
+						PrintFColorSmall(DrawPosX + DataOffset, DrawPosY, Closure.DiffuseR, Closure.DiffuseG, Closure.DiffuseB);
+						NewLine();
+
+						PrintS(DrawPosX + 0.0f, DrawPosY, TEXT("Specular Color"), Orange);
+						PrintFColorSmall(DrawPosX + DataOffset, DrawPosY, Closure.F0R, Closure.F0G, Closure.F0B);
+						NewLine();
+
+						PrintS(DrawPosX + 0.0f, DrawPosY, TEXT("Roughness"), Orange);
+						PrintFSmall(DrawPosX + DataOffset, DrawPosY, Closure.Roughness);
+						NewLine();
+
+						PrintS(DrawPosX + 0.0f, DrawPosY, TEXT("Anisotropy"), Orange);
+						PrintFSmall(DrawPosX + DataOffset, DrawPosY, Closure.Anisotropy);
+						NewLine();
+					}
 
 					DispatchBatch();
 
