[Lit Translucency] Fix on-screen stats alignment and overwriting.

#Synced-CL 48396324

diff --git a/Engine/Shaders/Private/TranslucentLightingShaders.usf b/Engine/Shaders/Private/TranslucentLightingShaders.usf
--- a/Engine/Shaders/Private/TranslucentLightingShaders.usf
+++ b/Engine/Shaders/Private/TranslucentLightingShaders.usf
@@ -409,18 +409,18 @@ void DebugTranslucencyLightingVolumeCS(
 		Print(Context, TEXT(" / "), FontYellow);
 		Print(Context, NumVoxels, FontYellow);
 		Print(Context, TEXT(" ("), FontYellow);
-		Print(Context, VisibleVoxels0 * 100.0f / NumVoxels, FontYellow, 3, 1);
+		Print(Context, VisibleVoxels0 * 100.0f / NumVoxels, FontYellow, 4, 1);
 		Print(Context, TEXT("%)"), FontYellow);
 		Newline(Context);
 
 		const uint MarkedVoxels0 = VoxelAllocator0[0];
 
-		Print(Context, TEXT("    - # Marked Voxels = "), FontWhite);
+		Print(Context, TEXT("    - # Marked Voxels  = "), FontWhite);
 		Print(Context, MarkedVoxels0, FontYellow);
 		Print(Context, TEXT(" / "), FontYellow);
 		Print(Context, NumVoxels, FontYellow);
 		Print(Context, TEXT(" ("), FontYellow);
-		Print(Context, MarkedVoxels0 * 100.0f / NumVoxels, FontYellow, 3, 1);
+		Print(Context, MarkedVoxels0 * 100.0f / NumVoxels, FontYellow, 4, 1);
 		Print(Context, TEXT("%)"), FontYellow);
 		Newline(Context);
 
@@ -438,18 +438,18 @@ void DebugTranslucencyLightingVolumeCS(
 		Print(Context, TEXT(" / "), FontYellow);
 		Print(Context, NumVoxels, FontYellow);
 		Print(Context, TEXT(" ("), FontYellow);
-		Print(Context, VisibleVoxels1 * 100.0f / NumVoxels, FontYellow, 3, 1);
+		Print(Context, VisibleVoxels1 * 100.0f / NumVoxels, FontYellow, 4, 1);
 		Print(Context, TEXT("%)"), FontYellow);
 		Newline(Context);
 
 		const uint MarkedVoxels1 = VoxelAllocator1[0];
 
-		Print(Context, TEXT("    - # Marked Voxels = "), FontWhite);
+		Print(Context, TEXT("    - # Marked Voxels  = "), FontWhite);
 		Print(Context, MarkedVoxels1, FontYellow);
 		Print(Context, TEXT(" / "), FontYellow);
 		Print(Context, NumVoxels, FontYellow);
 		Print(Context, TEXT(" ("), FontYellow);
-		Print(Context, MarkedVoxels1 * 100.0f / NumVoxels, FontYellow, 3, 1);
+		Print(Context, MarkedVoxels1 * 100.0f / NumVoxels, FontYellow, 4, 1);
 		Print(Context, TEXT("%)"), FontYellow);
 
 		AddTextBackground(Context, float4(0, 0, 0, 0.25f));
