Substrate - Implement SubstrateSanitizeBSDF for Toon.
Removed duplicated SLAB_ROUGHNESS sanitize.

#rnx
#tests PC editor

#Synced-CL 48453488

diff --git a/Engine/Shaders/Private/Substrate/Substrate.ush b/Engine/Shaders/Private/Substrate/Substrate.ush
--- a/Engine/Shaders/Private/Substrate/Substrate.ush
+++ b/Engine/Shaders/Private/Substrate/Substrate.ush
@@ -3440,7 +3440,6 @@ void FSubstrateBSDF::SubstrateSanitizeBSDF()
 				LocalHaziness.Roughness = SanitizeRoughness(LocalHaziness.Roughness);
 				SLAB_HAZINESS(this) = PackHaziness(LocalHaziness);
 			}
-			SLAB_ROUGHNESS(this) = SanitizeRoughness(SLAB_ROUGHNESS(this)); // 
 			SLAB_ANISOTROPY(this) = SanitizeAnisotropy(SLAB_ANISOTROPY(this));
 			// SLAB_HAZINESS is sanitize when packed, see PackHaziness.
 		#if SUBSTRATE_INLINE_SHADING && SUBSTRATE_OPAQUE_MATERIAL
@@ -3517,6 +3516,17 @@ void FSubstrateBSDF::SubstrateSanitizeBSDF()
 			UNLIT_NORMAL(this) = normalize(UNLIT_NORMAL(this));
 			break;
 		}
+		case SUBSTRATE_BSDF_TYPE_TOON:
+		{
+			// We do not clamp emissive luminance upper bound. This is done later in the base pass within pre-exposed luminance space.
+			BSDF_SETEMISSIVE(this, max(BSDF_GETEMISSIVE(this), 0.0));
+
+			TOON_DIFFUSEALBEDO(this) = SanitizeDiffuseAlbedo(TOON_DIFFUSEALBEDO(this));
+			TOON_SPECULARCOLOR(this) = SanitizeF0(TOON_SPECULARCOLOR(this));
+			TOON_ROUGHNESS(this) = SanitizeRoughness(TOON_ROUGHNESS(this));
+			TOON_ANISOTROPY(this) = SanitizeAnisotropy(TOON_ANISOTROPY(this));
+			break;
+		}
 	}
 }
 
