[Stochastic ligthing mark] Store normals in LDS as plain float3 instead of packing/unpacking to save some ALUs. 10 us GPU time saving on a console at 1280x720 rendering res when only 2x2 downsampling is enabled. 30 us saving when both 2x2 and 2x1 downsampling are enabled.

#tests tested in editor and on two consoles
#Synced-CL 48457700

diff --git a/Engine/Shaders/Private/StochasticLighting/StochasticLightingTileClassification.usf b/Engine/Shaders/Private/StochasticLighting/StochasticLightingTileClassification.usf
--- a/Engine/Shaders/Private/StochasticLighting/StochasticLightingTileClassification.usf
+++ b/Engine/Shaders/Private/StochasticLighting/StochasticLightingTileClassification.usf
@@ -77,9 +77,9 @@ RWTexture2D<UNORM float3> RWDownsampledWorldNormal2x1;
 RWTexture2D<UNORM float3> RWDownsampledWorldNormal2x2;
 
 groupshared float SharedSceneDepth2x1[6][10];
-groupshared uint SharedEncodedNormal2x1[6][10];
+groupshared float3 SharedEncodedNormal2x1[6][10];
 groupshared float SharedSceneDepth2x2[6][6];
-groupshared uint SharedEncodedNormal2x2[6][6];
+groupshared float3 SharedEncodedNormal2x2[6][6];
 
 void StoreDepthToCache(uint2 SharedCoord, const uint2 DownsampleFactor, float Depth)
 {
@@ -107,41 +107,26 @@ float LoadDepthFromCache(uint2 SharedCoord, const uint2 DownsampleFactor)
 
 void StoreNormalToCache(uint2 SharedCoord, const uint2 DownsampleFactor, float3 Normal)
 {
-	Normal = saturate(EncodeNormal(Normal));
-
-	uint Packed = 0;
-	Packed |= uint(Normal.x * 1023.0 + 0.5);
-	Packed |= uint(Normal.y * 1023.0 + 0.5) << 10;
-	Packed |= uint(Normal.z * 1023.0 + 0.5) << 20;
-
 	if (all(DownsampleFactor == 2))
 	{
-		SharedEncodedNormal2x2[SharedCoord.x][SharedCoord.y] = Packed;
+		SharedEncodedNormal2x2[SharedCoord.x][SharedCoord.y] = Normal;
 	}
 	else
 	{
-		SharedEncodedNormal2x1[SharedCoord.x][SharedCoord.y] = Packed;
+		SharedEncodedNormal2x1[SharedCoord.x][SharedCoord.y] = Normal;
 	}
 }
 
 float3 LoadNormalFromCache(uint2 SharedCoord, const uint2 DownsampleFactor)
 {
-	uint Packed;
 	if (all(DownsampleFactor == 2))
 	{
-		Packed = SharedEncodedNormal2x2[SharedCoord.x][SharedCoord.y];
+		return SharedEncodedNormal2x2[SharedCoord.x][SharedCoord.y];
 	}
 	else
 	{
-		Packed = SharedEncodedNormal2x1[SharedCoord.x][SharedCoord.y];
+		return SharedEncodedNormal2x1[SharedCoord.x][SharedCoord.y];
 	}
-
-	float3 Normal;
-	Normal.x = BitFieldExtractU32(Packed, 10, 0) / 1023.0;
-	Normal.y = BitFieldExtractU32(Packed, 10, 10) / 1023.0;
-	Normal.z = BitFieldExtractU32(Packed, 10, 20) / 1023.0;
-
-	return DecodeNormal(Normal);
 }
 
 float GetNormalWeight(float3 SceneWorldNormal, float3 SampleWorldNormal)
