Substrate - fixed compilation issue with PT and translucent

#rnx
#tests PC editor
#Synced-CL 48498703

diff --git a/Engine/Shaders/Private/PathTracing/PathTracingMaterialHitShader.usf b/Engine/Shaders/Private/PathTracing/PathTracingMaterialHitShader.usf
--- a/Engine/Shaders/Private/PathTracing/PathTracingMaterialHitShader.usf
+++ b/Engine/Shaders/Private/PathTracing/PathTracingMaterialHitShader.usf
@@ -1587,7 +1587,7 @@ RAY_TRACING_ENTRY_ANY_HIT(PathTracingMaterialAHS,
 						RefractMix = 1;
 						break;
 					}
-					case SUBSTRATE_BSDF_TYPE_TOON
+					case SUBSTRATE_BSDF_TYPE_TOON:
 					{
 						// No refraction
 						break;
