Remove SheenQuality settings. Substrate now always uses Sheen LTC.

RunReplay perf shows LTC is affordable. Removing Charlie Sheen as the LTC look is better and simplify maintenance.

#Synced-CL 48545202

diff --git a/Engine/Config/Linux/BaseLinuxEngine.ini b/Engine/Config/Linux/BaseLinuxEngine.ini
--- a/Engine/Config/Linux/BaseLinuxEngine.ini
+++ b/Engine/Config/Linux/BaseLinuxEngine.ini
@@ -14,7 +14,6 @@ BindlessConfiguration=All
 [/Script/Engine.RendererSettings]
 ; Substrate settings for more advanced shading on PC and editor by default.
 r.Substrate.RoughDiffuse=1
-r.Substrate.SheenQuality=1
 r.Substrate.ShadingQuality=1
 r.Substrate.TileCoord8bits=0
 r.Substrate.Glints=1
diff --git a/Engine/Config/LinuxArm64/BaseLinuxArm64Engine.ini b/Engine/Config/LinuxArm64/BaseLinuxArm64Engine.ini
--- a/Engine/Config/LinuxArm64/BaseLinuxArm64Engine.ini
+++ b/Engine/Config/LinuxArm64/BaseLinuxArm64Engine.ini
@@ -4,7 +4,6 @@ bAgreeToCrashUpload=false
 [/Script/Engine.RendererSettings]
 ; Substrate settings for more advanced shading on PC and editor by default.
 r.Substrate.RoughDiffuse=1
-r.Substrate.SheenQuality=1
 r.Substrate.ShadingQuality=1
 r.Substrate.TileCoord8bits=0
 r.Substrate.Glints=1
diff --git a/Engine/Config/Mac/BaseMacEngine.ini b/Engine/Config/Mac/BaseMacEngine.ini
--- a/Engine/Config/Mac/BaseMacEngine.ini
+++ b/Engine/Config/Mac/BaseMacEngine.ini
@@ -42,7 +42,6 @@ Slate.NoAnimationOnTabForegroundedEvent=True
 [/Script/Engine.RendererSettings]
 ; Substrate settings for more advanced shading on PC and editor by default.
 r.Substrate.RoughDiffuse=1
-r.Substrate.SheenQuality=1
 r.Substrate.ShadingQuality=1
 r.Substrate.TileCoord8bits=0
 r.Substrate.Glints=1
diff --git a/Engine/Config/Windows/BaseWindowsEngine.ini b/Engine/Config/Windows/BaseWindowsEngine.ini
--- a/Engine/Config/Windows/BaseWindowsEngine.ini
+++ b/Engine/Config/Windows/BaseWindowsEngine.ini
@@ -35,7 +35,6 @@ r.Shaders.RemoveUnusedInterpolators=1
 r.Shadow.DetectVertexShaderLayerAtRuntime=1
 ; Substrate settings for more advanced shading on PC and editor by default.
 r.Substrate.RoughDiffuse=1
-r.Substrate.SheenQuality=1
 r.Substrate.ShadingQuality=1
 r.Substrate.TileCoord8bits=0
 r.Substrate.Glints=1
diff --git a/Engine/Shaders/Private/PathTracing/Material/PathTracingSubstrateCommon.ush b/Engine/Shaders/Private/PathTracing/Material/PathTracingSubstrateCommon.ush
--- a/Engine/Shaders/Private/PathTracing/Material/PathTracingSubstrateCommon.ush
+++ b/Engine/Shaders/Private/PathTracing/Material/PathTracingSubstrateCommon.ush
@@ -17,10 +17,6 @@
 #include "/Engine/Private/Substrate/Glint/GlintThirdParty.ush"
 #endif
 
-#ifndef SUBSTRATE_SHEEN_QUALITY
-#define SUBSTRATE_SHEEN_QUALITY 0
-#endif
-
 float3 SubstrateLobeWeight(FPathTracingPayload Payload, float SatNoL)
 {
 	// See reference in LuminanceWeight
@@ -60,19 +56,14 @@ void SubstrateGlint(FPathTracingPayload Payload, float3 InTangentV, float3 InTan
 
 struct FSubstrateSheenData
 {
-#if SUBSTRATE_SHEEN_QUALITY == 1
 	float2 Tangent; 	// Captures rotation within the local tangent frame
 	float aInv; 		// LTC Coefficients
 	float bInv;
-#else
-	float Roughness;
-#endif
 	float Scale;      	// overall scale factor for the sheen lobe
 	float Attenuation; 	// amount by which other lobes have to be scaled
 
 	void Prepare(float3 V, float FuzzRoughness, float FuzzAmount)
 	{
-#if SUBSTRATE_SHEEN_QUALITY == 1
 		Tangent = V.xy;
 		const float Len2 = length2(Tangent);
 		if (Len2 > 0)
@@ -89,19 +80,12 @@ struct FSubstrateSheenData
 		const float DirectionalAlbedo = LTC.z;
 		Scale = FuzzAmount * DirectionalAlbedo;
 		Attenuation = 1 - FuzzAmount * DirectionalAlbedo; // equivalent to: lerp(1.0f, ComputeEnergyPreservation(EnergyTerms), FuzzAmount);
-#else
-		Roughness = MakeRoughnessSafe(FuzzRoughness, SUBSTRATE_MIN_FUZZ_ROUGHNESS);
-		FBxDFEnergyTermsA Fuzz = ComputeClothEnergyTermsA(Roughness, V.z);
-		Scale = FuzzAmount * Fuzz.W;
-		Attenuation = 1 - FuzzAmount * Fuzz.E;
-#endif
 	}
 
 
 
 	float4 Eval(float3 L, float3 V, float3 H, float3 FuzzColor)
 	{
-#if SUBSTRATE_SHEEN_QUALITY == 1
 		if (L.z > 0)
 		{
 			const float3 LocalL = float3(
@@ -131,23 +115,12 @@ struct FSubstrateSheenData
 #endif
 			return float4(Weight * FuzzColor, Pdf);
 		}
-#else
-		// Charlie's Cloth Lobe
-		if (V.z > 0 && L.z > 0)
-		{
-			const float DCloth = D_Charlie(Roughness, H.z);
-			const float3 FCloth = F_Schlick(FuzzColor, dot(V, H));
-			const float3 Weight = (2 * PI * L.z * DCloth * Vis_Ashikhmin(V.z, L.z)) * FCloth;
-			const float Pdf = 1.0 / (2 * PI);
-			return float4(Weight, Pdf);
-		}
-#endif
 		return 0;
 	}
 
 	float3 Sample(float2 Rand)
 	{
-#if SUBSTRATE_SHEEN_QUALITY == 1 && USE_SHEEN_LTC_SAMPLER == 1
+#if USE_SHEEN_LTC_SAMPLER == 1
 		const float3 CosSample = CosineSampleHemisphere(Rand).xyz;
 		// NOTE: The vector is scaled by aInv compared to the reference implementation to avoid unecessary divides
 		const float3 LocalL = normalize(float3(CosSample.x - CosSample.z * bInv, CosSample.y, aInv * CosSample.z));
@@ -163,12 +136,8 @@ struct FSubstrateSheenData
 
 FBxDFEnergyTermsA ComputeSheenEnergyTerms(float NoV, float SheenRoughness)
 {
-#if SUBSTRATE_SHEEN_QUALITY == 1
 	FBxDFEnergyTermsA Result;
 	Result.E = SheenLTC_DirectionalAlbedo(NoV, SheenRoughness, View.SheenLTCTexture, View.SheenLTCSampler);
 	Result.W = Result.E; // overall weight is the directional albedo
 	return Result;
-#else
-	return ComputeClothEnergyTermsA(SheenRoughness, NoV);
-#endif
 }
diff --git a/Engine/Shaders/Private/Substrate/SubstrateEvaluation.ush b/Engine/Shaders/Private/Substrate/SubstrateEvaluation.ush
--- a/Engine/Shaders/Private/Substrate/SubstrateEvaluation.ush
+++ b/Engine/Shaders/Private/Substrate/SubstrateEvaluation.ush
@@ -49,10 +49,6 @@
 #define MATERIALBLENDING_ANY_TRANSLUCENT 0
 #endif
 
-#ifndef SUBSTRATE_SHEEN_QUALITY
-#define SUBSTRATE_SHEEN_QUALITY 0
-#endif
-
 #ifndef SUBSTRATE_SSS_TRANSMISSION
 #define SUBSTRATE_SSS_TRANSMISSION 0
 #endif
@@ -1056,10 +1052,11 @@ FSubstrateEvaluateResult SubstrateEvaluateBSDFCommon(FSubstrateBSDFContext BSDFC
 				const float  FuzzRoughness = MakeRoughnessSafe(SLAB_FUZZ_ROUGHNESS(BSDFContext.BSDF), SUBSTRATE_MIN_FUZZ_ROUGHNESS);
 
 				// F0 is ignored, as energy preservation is handle solely based on the white directional albedo and FuzzAmount
-				FBxDFEnergyTermsA EnergyTerms = (FBxDFEnergyTermsA)1.f;
+				FBxDFEnergyTermsA ClothEnergyTerms = (FBxDFEnergyTermsA)1.f;
 				float3 ClothSpecularPathValueNoL = 0;
 
-				#if SUBSTRATE_SHEEN_QUALITY == 1
+				#if 1
+				// Always use Sheen LTC as its cost is affordable and has better quality
 				// Disney LTC
 				{
 					float DirectionalAlbedo = 1;
@@ -1073,25 +1070,27 @@ FSubstrateEvaluateResult SubstrateEvaluateBSDFCommon(FSubstrateBSDFContext BSDFC
 						ClothSpecularPathValueNoL = FuzzF0 * SheenLTC_Eval(BSDFContext.V, AreaLight.DiffuseL, BSDFContext.N, BSDFContext.Context.NoV, FuzzRoughness, View.SheenLTCTexture, View.SheenLTCSampler, DirectionalAlbedo);
 						ClothSpecularPathValueNoL *= DirectionalAlbedo;
 					}
-					EnergyTerms.E = DirectionalAlbedo;
-					EnergyTerms.W = 1.f;
+					ClothEnergyTerms.E = DirectionalAlbedo;
+					ClothEnergyTerms.W = 1.f;
 				}
 				#else
+				// UE_DEPRECATED: Keep Charlie Sheen for now, but code will be removed in future release
 				// Charlie cloth modle for D and Ashikhmin for visibility (more efficient than Charlie).
 				{
-					// ComputeEnergyConservation(EnergyTerms) is not taken into account because we do not want white fuzz to reflect all incoming energy
+					// ComputeEnergyConservation(ClothEnergyTerms) is not taken into account because we do not want white fuzz to reflect all incoming energy
 					// The fuzz should have some transmission to the lower lobes. In order to pass a furnace test, a combination of white diffuse+white fuzz is required
-					EnergyTerms = ComputeClothEnergyTermsA(FuzzRoughness, BSDFContext.Context.NoV);
+					ClothEnergyTerms = ComputeClothEnergyTermsA(FuzzRoughness, BSDFContext.Context.NoV);
 
 					float ClothD   = D_Charlie(FuzzRoughness, BSDFContext.Context.NoH);
 					float ClothVis = Vis_Ashikhmin(BSDFContext.Context.NoV, AreaLight.NoL);
 					float3 ClothF  = F_Schlick(FuzzF0, BSDFContext.Context.VoH);
-					ClothSpecularPathValueNoL = (ClothD * ClothVis) * ClothF * ComputeEnergyConservation(EnergyTerms) * AreaLight.NoL;
+					ClothSpecularPathValueNoL = (ClothD * ClothVis) * ClothF * ComputeEnergyConservation(ClothEnergyTerms) * AreaLight.NoL;
 				}
 				#endif //SUBSTRATE_SHEEN_QUALITY
 
 				// The specular and diffuse components below the fuzz are only attenuated linearly according to the amount of fuzz.
-				const float3 Cloth_DirectionalAlbedo_SpecularTransmission = lerp(1.0, ComputeEnergyPreservation(EnergyTerms), FuzzAmount);
+				const float ClothEnegyPreservation = 1.f - ClothEnergyTerms.E;
+				const float3 Cloth_DirectionalAlbedo_SpecularTransmission = lerp(1.0, ClothEnegyPreservation, FuzzAmount);
 
 				// Area light are not supported by the cloth BRDF
 				float3 ClothIntegratedSpecularValue = (ShadowTerms.SurfaceShadow * AreaLight.Falloff * FuzzAmount) * AreaLight.FalloffColor * ClothSpecularPathValueNoL;
@@ -1897,7 +1896,8 @@ FSubstrateEnvLightResult SubstrateEvaluateForEnvLight(FSubstrateBSDFContext BSDF
 				const float  FuzzRoughness = MakeRoughnessSafe(SLAB_FUZZ_ROUGHNESS(BSDFContext.BSDF), SUBSTRATE_MIN_FUZZ_ROUGHNESS);
 
 				FBxDFEnergyTermsA ClothEnergyTerms = (FBxDFEnergyTermsA)1.0f;
-				#if SUBSTRATE_SHEEN_QUALITY == 1
+				// Always use Sheen LTC as its cost is affordable and has better quality
+				#if 1
 				{
 					const float DirectionalAlbedo = SheenLTC_DirectionalAlbedo(BSDFContext.Context.NoV, FuzzRoughness, View.SheenLTCTexture, View.SheenLTCSampler);
 					ClothEnergyTerms.E = DirectionalAlbedo;
@@ -1905,13 +1905,15 @@ FSubstrateEnvLightResult SubstrateEvaluateForEnvLight(FSubstrateBSDFContext BSDF
 				}
 				#else
 				{
+					// UE_DEPRECATED: Keep Charlie Sheen for now, but code will be removed in future release
 					// F0 is ignored, as energy preservation is handle solely based on the white directional albedo and FuzzAmount
 					ClothEnergyTerms = ComputeClothEnergyTermsA(FuzzRoughness, BSDFContext.Context.NoV);
 				}
-				#endif // SUBSTRATE_SHEEN_QUALITY
+				#endif
 
 				// The specular and diffuse components below the fuzz are only attenuated linearly according to the amount of fuzz.
-				const float Cloth_DirectionalAlbedo_SpecularTransmission = lerp(1.0, ComputeEnergyPreservation(ClothEnergyTerms), FuzzAmount);
+				const float ClothEnegyPreservation = 1.f - ClothEnergyTerms.E;
+				const float Cloth_DirectionalAlbedo_SpecularTransmission = lerp(1.0, ClothEnegyPreservation, FuzzAmount);
 
 				SubstrateEnvLightResult.DiffuseColor		*= Cloth_DirectionalAlbedo_SpecularTransmission;
 				SubstrateEnvLightResult.DiffuseWeight		*= Cloth_DirectionalAlbedo_SpecularTransmission;
diff --git a/Engine/Shaders/Private/Substrate/SubstrateMaterialSampling.usf b/Engine/Shaders/Private/Substrate/SubstrateMaterialSampling.usf
--- a/Engine/Shaders/Private/Substrate/SubstrateMaterialSampling.usf
+++ b/Engine/Shaders/Private/Substrate/SubstrateMaterialSampling.usf
@@ -49,7 +49,7 @@ float3 GetSimpleVolumeDiffuseColor(float3 DiffuseColor, float3 MeanFreePath)
 
 FBxDFEnergyTermsA ComputeSheenEnergyTerms(float NoV, float SheenRoughness)
 {
-	#if SUBSTRATE_SHEEN_QUALITY == 1
+	#if 1
 	FBxDFEnergyTermsA Result;
 	Result.E = SheenLTC_DirectionalAlbedo(NoV, SheenRoughness, View.SheenLTCTexture, View.SheenLTCSampler);
 	Result.W = Result.E; // overall weight is the directional albedo
diff --git a/Engine/Source/Runtime/Engine/Private/ShaderCompiler/ShaderCompiler.cpp b/Engine/Source/Runtime/Engine/Private/ShaderCompiler/ShaderCompiler.cpp
--- a/Engine/Source/Runtime/Engine/Private/ShaderCompiler/ShaderCompiler.cpp
+++ b/Engine/Source/Runtime/Engine/Private/ShaderCompiler/ShaderCompiler.cpp
@@ -4022,9 +4022,6 @@ void GlobalBeginCompileShader(
 			const bool bLowQuality = SubstrateShadingQuality == 0;
 			SET_SHADER_DEFINE(Input.Environment, USE_ACHROMATIC_BXDF_ENERGY, bLowQuality ? 1u : 0u);
 
-			const uint32 SubstrateSheenQuality = FMath::Clamp(Substrate::GetSheenQuality(Target.GetPlatform()), 0u, 1u);
-			Input.Environment.SetDefine(TEXT("SUBSTRATE_SHEEN_QUALITY"), bLowQuality ? 0 : SubstrateSheenQuality);
-
 			const uint32 SubstrateNormalQuality = Substrate::GetNormalQuality();
 			SET_SHADER_DEFINE(Input.Environment, SUBSTRATE_NORMAL_QUALITY, SubstrateNormalQuality);
 			if (SubstrateNormalQuality == 0)
diff --git a/Engine/Source/Runtime/RenderCore/Private/RenderUtils.cpp b/Engine/Source/Runtime/RenderCore/Private/RenderUtils.cpp
--- a/Engine/Source/Runtime/RenderCore/Private/RenderUtils.cpp
+++ b/Engine/Source/Runtime/RenderCore/Private/RenderUtils.cpp
@@ -2056,12 +2056,6 @@ static TAutoConsoleVariable<int32> CVarSubstrateOpaqueMaterialRoughRefraction(
 	TEXT("Enable Substrate opaque material rough refractions effect from top layers over layers below. (Experimental for now since this is not a path validated in production)"),
 	ECVF_ReadOnly | ECVF_RenderThreadSafe);
 
-static TAutoConsoleVariable<int32> CVarSubstrateSheenQuality(
-	TEXT("r.Substrate.SheenQuality"),
-	0, // Substrate 60fps by default
-	TEXT("Define Substrate sheen quality (0: Charlie-based sheen, 1: Disney-based sheen). r.Substrate.ShadingQuality=0 forces SheenQuality to 0. This variable is read-only."),
-	ECVF_ReadOnly | ECVF_RenderThreadSafe | ECVF_Preview);
-
 static TAutoConsoleVariable<int32> CVarSubstrateShadingQuality(
 	TEXT("r.Substrate.ShadingQuality"),
 	0, // Substrate 60fps by default
@@ -2343,13 +2337,12 @@ namespace Substrate
 
 	uint32 GetSheenQuality()
 	{
-		return CVarSubstrateSheenQuality.GetValueOnAnyThread();
+		return 1u;
 	}
 
 	uint32 GetSheenQuality(EShaderPlatform InPlatform)
 	{
-		static FShaderPlatformCachedIniValue<int32> CVar(TEXT("r.Substrate.SheenQuality"));
-		return CVar.Get(InPlatform);
+		return 1u;
 	}
 
 	uint32 GetShadingQuality()
diff --git a/Engine/Source/Runtime/RenderCore/Private/Shader.cpp b/Engine/Source/Runtime/RenderCore/Private/Shader.cpp
--- a/Engine/Source/Runtime/RenderCore/Private/Shader.cpp
+++ b/Engine/Source/Runtime/RenderCore/Private/Shader.cpp
@@ -2502,12 +2502,6 @@ void ShaderMapAppendKey(EShaderPlatform Platform, FShaderKeyGenerator& KeyGen)
 			KeyGen.Append(Substrate::GetShadingQuality(Platform));
 		}
 
-		{
-			KeyGen.AppendSeparator();
-			KeyGen.Append(TEXT("SSHEEN"));
-			KeyGen.Append(Substrate::GetSheenQuality(Platform));
-		}
-
 		if (Substrate::IsGlintEnabled(Platform))
 		{
 			KeyGen.AppendSeparator();
diff --git a/Engine/Source/Runtime/RenderCore/Public/RenderUtils.h b/Engine/Source/Runtime/RenderCore/Public/RenderUtils.h
--- a/Engine/Source/Runtime/RenderCore/Public/RenderUtils.h
+++ b/Engine/Source/Runtime/RenderCore/Public/RenderUtils.h
@@ -638,7 +638,7 @@ namespace Substrate
 	RENDERCORE_API bool IsStochasticLightingEnabled(EShaderPlatform InPlatform);
 	RENDERCORE_API bool IsRoughDiffuseEnabled();
 	RENDERCORE_API bool IsRoughDiffuseEnabled(EShaderPlatform InPlatform);
-	UE_DEPRECATED(5.7, "Plase use IsGlintEnabled(EShaderPlatform InPlatform)")
+	UE_DEPRECATED(5.7, "Please use IsGlintEnabled(EShaderPlatform InPlatform)")
 	RENDERCORE_API bool IsGlintEnabled();
 	RENDERCORE_API bool IsGlintEnabled(EShaderPlatform InPlatform);
 	RENDERCORE_API uint32 GlintLUTIndex();
@@ -649,7 +649,7 @@ namespace Substrate
 	RENDERCORE_API bool IsSpecularProfileEnabled(EShaderPlatform InPlatform);
 	RENDERCORE_API bool IsDBufferPassEnabled(EShaderPlatform InPlatform);
 	RENDERCORE_API bool IsOpaqueRoughRefractionEnabled(EShaderPlatform InPlatform);
-	UE_DEPRECATED(5.7, "Plase use IsAdvancedVisualizationEnabled(EShaderPlatform InPlatform)")
+	UE_DEPRECATED(5.7, "Please use IsAdvancedVisualizationEnabled(EShaderPlatform InPlatform)")
 	RENDERCORE_API bool IsAdvancedVisualizationEnabled();
 	RENDERCORE_API bool IsAdvancedVisualizationEnabled(EShaderPlatform InPlatform);
 	RENDERCORE_API bool IsMaterialLayeringSupportEnabled();
@@ -666,7 +666,9 @@ namespace Substrate
 
 	RENDERCORE_API uint32 GetNormalQuality();
 
+	UE_DEPRECATED(5.8, "This function is no longer used. It will always return 1.")
 	RENDERCORE_API uint32 GetSheenQuality();
+	UE_DEPRECATED(5.8, "This function is no longer used. It will always return 1.")
 	RENDERCORE_API uint32 GetSheenQuality(EShaderPlatform InPlatform);
 
 	RENDERCORE_API uint32 GetShadingQuality();
