PR #14007: Fix GPU Crash in SetupLogLuminanceCS

#Synced-CL 48565949

diff --git a/Engine/Shaders/Private/PostProcessLocalExposure.usf b/Engine/Shaders/Private/PostProcessLocalExposure.usf
--- a/Engine/Shaders/Private/PostProcessLocalExposure.usf
+++ b/Engine/Shaders/Private/PostProcessLocalExposure.usf
@@ -23,11 +23,19 @@ RWTexture2D<float4> OutputFloat4;
 RWTexture2D<float4> OutputFloat4_1;
 
 [numthreads(THREADGROUP_SIZEX, THREADGROUP_SIZEY, 1)]
-void SetupLogLuminanceCS(uint3 DispatchThreadId : SV_DispatchThreadID)
+void SetupLogLuminanceCS(uint2 DispatchThreadId : SV_DispatchThreadID)
 {
-	float3 SceneColor = InputTexture.Load(uint3(Input_ViewportMin + DispatchThreadId.xy, 0)).rgb * View.OneOverPreExposure;
+	const uint2 InputPixelPos = Input_ViewportMin + DispatchThreadId;
+
+	BRANCH
+	if (any(InputPixelPos >= (uint2)Input_ViewportMax))
+	{
+		return;
+	}
+
+	float3 SceneColor = InputTexture.Load(uint3(InputPixelPos, 0)).rgb * View.OneOverPreExposure;
 	float LuminanceVal = CalculateEyeAdaptationLuminance(SceneColor);
-	OutputFloat[DispatchThreadId.xy] = log2(LuminanceVal);
+	OutputFloat[DispatchThreadId] = log2(LuminanceVal);
 }
 
 [numthreads(THREADGROUP_SIZEX, THREADGROUP_SIZEY, 1)]
