Substrate - added normal sanitize before normalize to avoid lighting fireflies due to our oct encoding allowing nan normals. (very visible with rect light + mfp spreading nan around)

#tests PC editor dx12
#rnx
#Synced-CL 48604927

diff --git a/Engine/Shaders/Private/BasePassPixelShader.usf b/Engine/Shaders/Private/BasePassPixelShader.usf
--- a/Engine/Shaders/Private/BasePassPixelShader.usf
+++ b/Engine/Shaders/Private/BasePassPixelShader.usf
@@ -1734,17 +1734,7 @@ void FPixelShaderInOut_MainPS(
 #if SUBSTRATE_ENABLED && !SUBSTRATE_OPTIMIZED_UNLIT
 
 #if SUBSTRATE_INLINE_SHADING
-	// We must normalize each normal and tangent to avoid non normalised vectors due to per vertex interpolation or texture filtering,
-	// for the deferred (our packing relies on normalized normal) and forward (normals are going to be used as-is from registers) paths.
-	UNROLL
-	for (uint i = 0; i < SubstratePixelHeader.SharedLocalBases.Count; ++i)
-	{
-		SubstratePixelHeader.SharedLocalBases.Normals[i] = normalize(SubstratePixelHeader.SharedLocalBases.Normals[i]);
-		if (SubstrateGetSharedLocalBasisType(SubstratePixelHeader.SharedLocalBases.Types, i) == SUBSTRATE_BASIS_TYPE_TANGENT)
-		{
-			SubstratePixelHeader.SharedLocalBases.Tangents[i] = normalize(SubstratePixelHeader.SharedLocalBases.Tangents[i]);
-		}
-	}
+	SanitizeAndNormalizeSharedLocalBases(SubstratePixelHeader.SharedLocalBases, MaterialParameters.TangentToWorld);
 #endif
 
 	// Export non-unlit materials (opaque or translucent)
diff --git a/Engine/Shaders/Private/MobileBasePassPixelShader.usf b/Engine/Shaders/Private/MobileBasePassPixelShader.usf
--- a/Engine/Shaders/Private/MobileBasePassPixelShader.usf
+++ b/Engine/Shaders/Private/MobileBasePassPixelShader.usf
@@ -663,22 +663,7 @@ void Main(
 
 
 	#if SUBSTRATE_INLINE_SHADING
-	// We must normalize each normal and tangent to avoid non normalised vectors due to per vertex interpolation or texture filtering,
-	// for the deferred (our packing relies on normalized normal) and forward (normals are going to be used as-is from registers) paths.
-	#if SUBSTRATE_FORCE_SINGLE_BSDF
-	uint i = 0;
-	if (SubstratePixelHeader.SharedLocalBases.Count > 0)
-	#else
-	UNROLL
-	for (uint i = 0; i < SubstratePixelHeader.SharedLocalBases.Count; ++i)
-	#endif // SUBSTRATE_FORCE_SINGLE_BSDF
-	{
-		SubstratePixelHeader.SharedLocalBases.Normals[i] = normalize(SubstratePixelHeader.SharedLocalBases.Normals[i]);
-		if (SubstrateGetSharedLocalBasisType(SubstratePixelHeader.SharedLocalBases.Types, i) == SUBSTRATE_BASIS_TYPE_TANGENT)
-		{
-			SubstratePixelHeader.SharedLocalBases.Tangents[i] = normalize(SubstratePixelHeader.SharedLocalBases.Tangents[i]);
-		}
-	}
+	SanitizeAndNormalizeSharedLocalBases(SubstratePixelHeader.SharedLocalBases, MaterialParameters.TangentToWorld);
 	#endif // SUBSTRATE_INLINE_SHADING
 
 	FLightAccumulator DirectLighting = (FLightAccumulator)0;
diff --git a/Engine/Shaders/Private/Substrate/Substrate.ush b/Engine/Shaders/Private/Substrate/Substrate.ush
--- a/Engine/Shaders/Private/Substrate/Substrate.ush
+++ b/Engine/Shaders/Private/Substrate/Substrate.ush
@@ -3429,6 +3429,50 @@ float SanitizeAnisotropy(float Anisotropy)
 	return clamp(Anisotropy, -1.f, 1.f);
 }
 
+#define SUBSTRATE_IS_VALID_DIRECTION(Direction) (!all(Direction == 0) && all(IsFinite(Direction)))
+
+void SanitizeAndNormalizeSharedLocalBases(inout FSharedLocalBases SharedLocalBases, in float3x3 TangentToWorld)
+{
+	// We must sanitize each normal to make sure nans are not stored which can expanlater to fire flies in the image.
+	// We must normalize each normal and tangent to avoid non normalised vectors due to per vertex interpolation or texture filtering,
+	// for the deferred (our oct packing relies on normalized normal) and forward (normals are going to be used as-is from registers) paths.
+	uint i = 0;
+#if SUBSTRATE_CLAMPED_CLOSURE_COUNT == 1
+	if(SharedLocalBases.Count > 0) // Equivalent to Count==1
+#else
+	UNROLL
+	for (; i < SharedLocalBases.Count; ++i)
+#endif
+	{
+		const float3 InputNormal = SharedLocalBases.Normals[i];
+		if(!SUBSTRATE_IS_VALID_DIRECTION(InputNormal))
+		{
+			const float3 FallbackNormal = TangentToWorld[2];
+		#if IS_NANITE_PASS // With nanite we do more check in case the TangentToWorld is wrong. (noticed in some material but we have not been able to catch the bug yet)
+			SharedLocalBases.Normals[i] = SUBSTRATE_IS_VALID_DIRECTION(FallbackNormal) ? FallbackNormal : float3(0, 0, 1);
+		#else
+			SharedLocalBases.Normals[i] = FallbackNormal;
+		#endif
+		}
+		SharedLocalBases.Normals[i] = normalize(SharedLocalBases.Normals[i]);
+
+		if (SubstrateGetSharedLocalBasisType(SharedLocalBases.Types, i) == SUBSTRATE_BASIS_TYPE_TANGENT)
+		{
+			const float3 InputTangent = SharedLocalBases.Tangents[i];
+			if (!SUBSTRATE_IS_VALID_DIRECTION(InputTangent))
+			{
+				const float3 FallbackTangent = TangentToWorld[0];
+			#if IS_NANITE_PASS
+				SharedLocalBases.Tangents[i] = SUBSTRATE_IS_VALID_DIRECTION(FallbackTangent) ? FallbackTangent : float3(1, 0, 0);
+			#else
+				SharedLocalBases.Normals[i] = FallbackTangent;
+			#endif
+			}
+			SharedLocalBases.Tangents[i] = normalize(SharedLocalBases.Tangents[i]);
+		}
+	}
+}
+
 void FSubstrateBSDF::SubstrateSanitizeBSDF()
 {
 	switch (BSDF_GETTYPE(this))
diff --git a/Engine/Shaders/Private/Substrate/SubstrateMobileForwardLighting.ush b/Engine/Shaders/Private/Substrate/SubstrateMobileForwardLighting.ush
--- a/Engine/Shaders/Private/Substrate/SubstrateMobileForwardLighting.ush
+++ b/Engine/Shaders/Private/Substrate/SubstrateMobileForwardLighting.ush
@@ -9,11 +9,6 @@
 
 #define FORWARD_PER_PIXEL_SHADING (FORWARD_SHADING || TRANSLUCENCY_LIGHTING_SURFACE_FORWARDSHADING || MATERIAL_SHADINGMODEL_SINGLELAYERWATER)
 
-// For debug or iteration purpose only
-#ifndef SUBSTRATE_FORCE_SINGLE_BSDF
-#define SUBSTRATE_FORCE_SINGLE_BSDF 1
-#endif
-
 // For debug or iteration purpose only
 #ifndef SUBSTRATE_FORCE_SINGLE_LOCAL_LIGHT
 #define SUBSTRATE_FORCE_SINGLE_LOCAL_LIGHT 0
