Force disable instanced stereo for runtime virtual texture page rendering.
It isn't relevant and causes a PSO creation failure.
#rnx

#Synced-CL 48656603

diff --git a/Engine/Shaders/Private/VirtualTextureMaterial.usf b/Engine/Shaders/Private/VirtualTextureMaterial.usf
--- a/Engine/Shaders/Private/VirtualTextureMaterial.usf
+++ b/Engine/Shaders/Private/VirtualTextureMaterial.usf
@@ -2,6 +2,12 @@
 
 // Material shader for runtime virtual texture
 
+// Remove Instanced Stereo support not needed for writing virtual texture pages.
+#ifdef INSTANCED_STEREO
+#undef INSTANCED_STEREO
+#endif
+#define INSTANCED_STEREO 0
+
 #include "Common.ush"
 #include "GammaCorrectionCommon.ush"
 #include "VirtualTextureMaterial.ush"
