Substrate - Added specular ramp.

#tests PC editor
#rnx
#Synced-CL 48664805

diff --git a/Engine/Shaders/Private/Substrate/Substrate.ush b/Engine/Shaders/Private/Substrate/Substrate.ush
--- a/Engine/Shaders/Private/Substrate/Substrate.ush
+++ b/Engine/Shaders/Private/Substrate/Substrate.ush
@@ -3575,7 +3575,7 @@ void FSubstrateBSDF::SubstrateSanitizeBSDF()
 
 			TOON_DIFFUSEALBEDO(this) = SanitizeDiffuseAlbedo(TOON_DIFFUSEALBEDO(this));
 			TOON_SPECULARCOLOR(this) = SanitizeF0(TOON_SPECULARCOLOR(this));
-			TOON_ROUGHNESS(this) = SanitizeRoughness(TOON_ROUGHNESS(this));
+			TOON_ROUGHNESS(this) = clamp(TOON_ROUGHNESS(this), 0.0, 1.0);	// We do not use SanitizeRoughness which clamps to View.MinRoughness that is higher than 0.
 			TOON_ANISOTROPY(this) = SanitizeAnisotropy(TOON_ANISOTROPY(this));
 			break;
 		}
diff --git a/Engine/Shaders/Private/Substrate/SubstrateEvaluation.ush b/Engine/Shaders/Private/Substrate/SubstrateEvaluation.ush
--- a/Engine/Shaders/Private/Substrate/SubstrateEvaluation.ush
+++ b/Engine/Shaders/Private/Substrate/SubstrateEvaluation.ush
@@ -1308,17 +1308,38 @@ FSubstrateEvaluateResult SubstrateEvaluateBSDFCommon(FSubstrateBSDFContext BSDFC
 		{
 			const float3 DiffuseColor		= TOON_DIFFUSEALBEDO(BSDFContext.BSDF);
 			const float3 SpecularColor		= TOON_SPECULARCOLOR(BSDFContext.BSDF);
+			const float  Roughness			= TOON_ROUGHNESS(BSDFContext.BSDF);
+			const float  Anisotropy			= TOON_ANISOTROPY(BSDFContext.BSDF);
 
 			const float ToonProfileId		= TOON_PROFILEID(BSDFContext.BSDF);
 
-			Sample.DiffuseColor				= DiffuseColor * EvaluateToonProfile(ToonProfileId, BSDFContext.Context.NoL); // AreaLight.NoL to handle capsule shape?
-
-			// TOON_TODO Hardcoded Toon Lambert diffuse before we implement the profile.
+			// Diffuse
+			Sample.DiffuseColor				= DiffuseColor * EvaluateToonProfileDiffuseRamp(ToonProfileId, BSDFContext.Context.NoL); // AreaLight.NoL to handle capsule shape?
 			Sample.DiffusePathValue			= Diffuse_Lambert(Sample.DiffuseColor);
 			Sample.IntegratedDiffuseValue	= ShadowTerms.SurfaceShadow * AreaLight.Falloff * AreaLight.FalloffColor * Sample.DiffusePathValue;
+			
+			const float RoughnessExponent	= saturate(1.0 - Roughness);
+			const float RoughnessExponent2 = RoughnessExponent * RoughnessExponent; // Arbitrary: this is so that the roughness control appears smooth and linear without big jumps due to 8bit encoding.
+
+			const float MaxSpecularExponent = 64.0f;
+			const float SpecularExponentScale = MaxSpecularExponent - 1.0f;
+		#if 0
+			const float Specular = pow(saturate(dot(AreaLight.SpecularL, BSDFContext.R)), max(1.0f, 1.0 + RoughnessExponent2 * SpecularExponentScale));	// NoR based is also interesting visually.
+		#else
+			const float Specular = pow(saturate(BSDFContext.Context.NoH), max(1.0f, 1.0 + RoughnessExponent2 * SpecularExponentScale));					// NoH based
+		#endif
+			// TOON_TODO add anisotropy support
+			// TOON_TODO energy conservation
+			// TOON_TODO experiment with GGX alike lob that is normalize according to peak lobe?
+			// TOON_TODO Area light, the roughness expand trick here does not make sense to use. MRP?
 
-			Sample.DiffusePDF				= BSDFContext.SatNoL / PI;
-			Sample.bPostProcessSubsurface	= false;
+			const float3 SpecularRamp = EvaluateToonProfileSpecularRamp(ToonProfileId, Specular);
+			Sample.SpecularPathValue = SpecularRamp;
+			Sample.IntegratedSpecularValue += ShadowTerms.SurfaceShadow * AreaLight.Falloff * AreaLight.FalloffColor * Sample.SpecularPathValue;
+
+			// Others
+			Sample.DiffusePDF = BSDFContext.SatNoL / PI;
+			Sample.bPostProcessSubsurface = false;
 
 			Sample.EmissivePathValue		= BSDF_GETEMISSIVE(BSDFContext.BSDF);
 
diff --git a/Engine/Shaders/Private/ToonProfile.usf b/Engine/Shaders/Private/ToonProfile.usf
--- a/Engine/Shaders/Private/ToonProfile.usf
+++ b/Engine/Shaders/Private/ToonProfile.usf
@@ -12,6 +12,7 @@ uint ToonProfileIndex;
 RWTexture2DArray<float4> ToonProfileTexture;
 
 Buffer<float4> DiffuseRampBufferSRV;
+Buffer<float4> SpecularRampBufferSRV;
 
 [numthreads(8, 8, 1)]
 void ToonProfileUpdateAtlasTileCS(uint3 DispatchThreadId : SV_DispatchThreadID)
@@ -21,10 +22,14 @@ void ToonProfileUpdateAtlasTileCS(uint3 DispatchThreadId : SV_DispatchThreadID)
 	{
 		const float2 TexCoord = float2(PixelPos + 0.5f) / float2(ToonProfileTextureResolution);
 
-		if (PixelPos.y == 0)
+		if (PixelPos.y == TOON_PROFILE_DIFFUSE_RAMP_Y)
 		{
 			ToonProfileTexture[uint3(PixelPos, ToonProfileIndex)] = DiffuseRampBufferSRV[PixelPos.x];
 		}
+		else if (PixelPos.y == TOON_PROFILE_SPECULAR_RAMP_Y)
+		{
+			ToonProfileTexture[uint3(PixelPos, ToonProfileIndex)] = SpecularRampBufferSRV[PixelPos.x];
+		}
 		else
 		{
 			ToonProfileTexture[uint3(PixelPos, ToonProfileIndex)] = float4(1, 0, 1, 0);
diff --git a/Engine/Shaders/Private/ToonProfileCommon.ush b/Engine/Shaders/Private/ToonProfileCommon.ush
--- a/Engine/Shaders/Private/ToonProfileCommon.ush
+++ b/Engine/Shaders/Private/ToonProfileCommon.ush
@@ -2,13 +2,26 @@
 
 #pragma once
 
+#define TOON_PROFILE_DIFFUSE_RAMP_Y    0
+#define TOON_PROFILE_SPECULAR_RAMP_Y   1
 
-float3 EvaluateToonProfile(in float ToonProfileId, float NoL)
+float3 EvaluateToonProfileRamp(in float ToonProfileId, in float NoL, in float TOON_PROFILE_RAMP)
 {
 	const float3 ToonProfileCoord = float3(
-		NoL * 0.5f + 0.5f,									// The toon ramp is applied over NoL in [-1, 1]
-		0.5f * View.ToonProfileTextureSizeAndInvSize.w,		// Center of the first row of texel
-		ToonProfileId										// Slice index
+		NoL,																	// The toon ramp is applied over NoL in [-1, 1]
+		(TOON_PROFILE_RAMP + 0.5f) * View.ToonProfileTextureSizeAndInvSize.w,		// Center of the first row of texel
+		ToonProfileId																		// Slice index
 	);
 	return View.ToonProfileTexture.SampleLevel(View.ToonProfileSampler, ToonProfileCoord, 0).xyz;
-}
\ No newline at end of file
+}
+
+float3 EvaluateToonProfileDiffuseRamp(in float ToonProfileId, in float NoL)
+{
+	return EvaluateToonProfileRamp(ToonProfileId, NoL * 0.5f + 0.5f, TOON_PROFILE_DIFFUSE_RAMP_Y);
+}
+
+float3 EvaluateToonProfileSpecularRamp(in float ToonProfileId, in float Specular)
+{
+	return EvaluateToonProfileRamp(ToonProfileId, Specular, TOON_PROFILE_SPECULAR_RAMP_Y);
+}
+
diff --git a/Engine/Source/Runtime/Engine/Classes/Engine/ToonProfile.h b/Engine/Source/Runtime/Engine/Classes/Engine/ToonProfile.h
--- a/Engine/Source/Runtime/Engine/Classes/Engine/ToonProfile.h
+++ b/Engine/Source/Runtime/Engine/Classes/Engine/ToonProfile.h
@@ -24,11 +24,17 @@ struct FToonProfileStruct
 	GENERATED_USTRUCT_BODY()
 
 	/**
-	* Define the diffuse ramp
+	* Define the diffuse ramp.
 	*/
 	UPROPERTY(Category = "Procedural", EditAnywhere, meta = (AllowZoomOutput = "false", ShowZoomButtons = "false", ViewMinInput = "0", ViewMaxInput = "1", ViewMinOutput = "0", ViewMaxOutput = "1", TimelineLength = "1", ShowInputGridNumbers="false", ShowOutputGridNumbers="false"))
 	FRuntimeCurveLinearColor DiffuseRamp;
 
+	/**
+	* Define the specular ramp.
+	*/
+	UPROPERTY(Category = "Procedural", EditAnywhere, meta = (AllowZoomOutput = "false", ShowZoomButtons = "false", ViewMinInput = "0", ViewMaxInput = "1", ViewMinOutput = "0", ViewMaxOutput = "1", TimelineLength = "1", ShowInputGridNumbers="false", ShowOutputGridNumbers="false"))
+	FRuntimeCurveLinearColor SpecularRamp;
+
 	FToonProfileStruct();
 
 	void Invalidate()
diff --git a/Engine/Source/Runtime/Engine/Private/Rendering/ToonProfile.cpp b/Engine/Source/Runtime/Engine/Private/Rendering/ToonProfile.cpp
--- a/Engine/Source/Runtime/Engine/Private/Rendering/ToonProfile.cpp
+++ b/Engine/Source/Runtime/Engine/Private/Rendering/ToonProfile.cpp
@@ -62,6 +62,7 @@ public:
 		SHADER_PARAMETER(uint32, ToonProfileIndex)
 		SHADER_PARAMETER_RDG_TEXTURE_UAV(RWTexture2DArray, ToonProfileTexture)
 		SHADER_PARAMETER_RDG_BUFFER_SRV(Buffer, DiffuseRampBufferSRV)
+		SHADER_PARAMETER_RDG_BUFFER_SRV(Buffer, SpecularRampBufferSRV)
 	END_SHADER_PARAMETER_STRUCT()
 
 	static int32 GetThreadGroupSize()
@@ -348,6 +349,9 @@ IPooledRenderTarget* FToonProfileTextureManager::GetAtlasTexture(FRDGBuilder& Gr
 			FRDGBufferRef DiffuseRampBuffer = CreateBufferFRomCurve(Data.DiffuseRamp);
 			PassParameters->DiffuseRampBufferSRV = GraphBuilder.CreateSRV(DiffuseRampBuffer, Format);
 
+			FRDGBufferRef SpecularRampBuffer = CreateBufferFRomCurve(Data.SpecularRamp);
+			PassParameters->SpecularRampBufferSRV = GraphBuilder.CreateSRV(SpecularRampBuffer, Format);
+
 			FToonProfileUpdateAtlasTileCS::FPermutationDomain PermutationVector;
 			TShaderMapRef<FToonProfileUpdateAtlasTileCS> Shader(GlobalShaderMap, PermutationVector);
 			FComputeShaderUtils::AddPass(GraphBuilder, 
@@ -412,6 +416,11 @@ FToonProfileStruct::FToonProfileStruct()
 	AddPointToCurve(DiffuseRamp, 0.76, 0.75);
 	AddPointToCurve(DiffuseRamp, 0.99, 0.75);
 	AddPointToCurve(DiffuseRamp, 1.00, 1.00);
+
+	AddPointToCurve(SpecularRamp, 0.00, 0.0);
+	AddPointToCurve(SpecularRamp, 0.50, 0.0);
+	AddPointToCurve(SpecularRamp, 0.51, 1.0);
+	AddPointToCurve(SpecularRamp, 1.00, 1.0);
 }
 
 ///////////////////////////////////////////////////////////////////////////////////////////////////
