Fixed ray tracing dynamic mesh update for the case when ISM VF is not used by always fetching the data from GPU-Scene (this affected any mesh that used WPO in RT except ISMs, as they were expecting a CPU-uploaded WorldToInstance transform).

#rnx
#Synced-CL 43344078

diff --git a/Engine/Shaders/Private/RayTracing/RayTracingDynamicMesh.usf b/Engine/Shaders/Private/RayTracing/RayTracingDynamicMesh.usf
--- a/Engine/Shaders/Private/RayTracing/RayTracingDynamicMesh.usf
+++ b/Engine/Shaders/Private/RayTracing/RayTracingDynamicMesh.usf
@@ -97,7 +97,7 @@ void RayTracingDynamicGeometryConverterCS(uint3 GroupId : SV_GroupID, int GroupT
 			}
 
 			float3 LocalPosition;
-	#if USE_INSTANCING || USE_INSTANCE_CULLING
+	#if USE_INSTANCING
 			//Reverse transform to neutral space
 			float3 WorldPosition = TranslatedWorldPosition.xyz - DFHackToFloat(ResolvedView.PreViewTranslation);
 			LocalPosition = mul(float4(WorldPosition, 1), WorldToInstance).xyz;
