Fix VolumetricFog on mobile renderer when substrate is enable.

#tests consoles with ODSC
#Synced-CL 48408831

diff --git a/Engine/Shaders/Private/VolumetricFog.usf b/Engine/Shaders/Private/VolumetricFog.usf
--- a/Engine/Shaders/Private/VolumetricFog.usf
+++ b/Engine/Shaders/Private/VolumetricFog.usf
@@ -1065,8 +1065,8 @@ void LightScatteringCS(
 #endif // USE_LIGHT_SOFT_FADING
 
 					Lighting = SofFade;
-#if !SHADING_PATH_MOBILE
-					// Mobile is already getting the Falloff from GetLocalLightAttenuation
+#if !SHADING_PATH_MOBILE || MOBILE_SHADING_PATH_SUPPORT_CAPSULE_LIGHT || (SUBSTRATE_ENABLED && !ALLOW_LOCAL_LIGHT_DISTANCE_ATTENUATION)
+					// Mobile is already receiving the falloff from GetLocalLightAttenuation when capsule lights are not supported and Substrate is disabled.
 					Lighting *= IntegrateLight(Capsule, LightData.bInverseSquared);
 #endif	
 				}
