Make Depth shaders use sv_position after interpolants like the hit proxy/basepass shaders
PositionOnlyDepthVertexShader.usf is not updated because it was only returning SV_Position & FStereoVSOutput, but no other interpolants
Looking at the code, FDepthOnlyPS can only be used with TDepthOnlyVS, and vice-versa

#rnx

#Synced-CL 48654158

diff --git a/Engine/Shaders/Private/DepthOnlyPixelShader.usf b/Engine/Shaders/Private/DepthOnlyPixelShader.usf
--- a/Engine/Shaders/Private/DepthOnlyPixelShader.usf
+++ b/Engine/Shaders/Private/DepthOnlyPixelShader.usf
@@ -13,16 +13,13 @@
 #include "/Engine/Generated/VertexFactory.ush"
 
 void Main(
-#if !MATERIALBLENDING_SOLID || OUTPUT_PIXEL_DEPTH_OFFSET
-	in INPUT_POSITION_QUALIFIERS float4 SvPosition : SV_Position,
-#endif
-
 #if !MATERIALBLENDING_SOLID || OUTPUT_PIXEL_DEPTH_OFFSET
 	FVertexFactoryInterpolantsVSToPS FactoryInterpolants
 	, float4 PixelPosition : TEXCOORD6 //-- Not used currently
 #if USE_WORLD_POSITION_EXCLUDING_SHADER_OFFSETS
 	, float3 PixelPositionExcludingWPO : TEXCOORD7
 #endif
+	, in INPUT_POSITION_QUALIFIERS float4 SvPosition : SV_Position
 	, in FStereoPSInput StereoInput
 	OPTIONAL_IsFrontFace
 	OPTIONAL_OutDepthConservative,
diff --git a/Engine/Shaders/Private/DepthOnlyVertexShader.usf b/Engine/Shaders/Private/DepthOnlyVertexShader.usf
--- a/Engine/Shaders/Private/DepthOnlyVertexShader.usf
+++ b/Engine/Shaders/Private/DepthOnlyVertexShader.usf
@@ -12,8 +12,6 @@
 
 struct FDepthOnlyVSToPS
 {
-	INVARIANT_OUTPUT float4 Position : SV_POSITION;
-
 	#if !MATERIALBLENDING_SOLID || OUTPUT_PIXEL_DEPTH_OFFSET
 		FVertexFactoryInterpolantsVSToPS FactoryInterpolants;
 		float4 PixelPosition : TEXCOORD6;
@@ -22,6 +20,8 @@ struct FDepthOnlyVSToPS
 	#if USE_RAW_WORLD_POSITION
 		float3 PixelPositionExcludingWPO : TEXCOORD7;
 	#endif
+
+	INVARIANT_OUTPUT float4 Position : SV_POSITION; // Put last to help with VS deduplications with other VS like the basepass ones
 };
 
 #define FDepthOnlyVSOutput FDepthOnlyVSToPS
