Support ddx/ddy in compute shader virtual texture sampling.
This is for nanite materials where analytic derivatives can't be used for virtual texture samples.
The change will prevent nanite pixel binning for that case, but that is likely much less bad then forcing mip 0.
Note that we use the DDX/DDY macros which have further configurable behavior for USE_FORCE_TEXTURE_MIP or nanite specific COMPUTE_SHADED etc.
#Synced-CL 48792737

diff --git a/Engine/Shaders/Private/VirtualTextureCommon.ush b/Engine/Shaders/Private/VirtualTextureCommon.ush
--- a/Engine/Shaders/Private/VirtualTextureCommon.ush
+++ b/Engine/Shaders/Private/VirtualTextureCommon.ush
@@ -25,6 +25,10 @@
 #define VT_FEEDBACK_RESERVOIR_SAMPLE 0
 #endif
 
+#ifndef VT_SUPPORT_DERIVATIVES
+#define VT_SUPPORT_DERIVATIVES (PIXELSHADER || COMPUTESHADER)
+#endif
+
 #ifndef VISUALIZE
 #define VISUALIZE 0
 #endif
@@ -464,9 +468,9 @@ VTPageTableResult TextureLoadVirtualPageTable(
 	{
 		UV = UV * PageTableUniform.UVScale;
 		int vLevel = GetGlobalVirtualTextureMipBias(PageTableUniform.MipBiasGroup);
-#if PIXELSHADER
+#if VT_SUPPORT_DERIVATIVES
 		vLevel = TextureComputeVirtualMipLevel(Result, DDX(UV), DDY(UV), MipBias, SvPositionXY, PageTableUniform);
-#endif // PIXELSHADER
+#endif // VT_SUPPORT_DERIVATIVES
 		UV = ApplyAddressMode(UV, AddressU, AddressV);
 		TextureLoadVirtualPageTableInternal(Result, PageTable0, PageTableUniform, UV, vLevel);
 		StoreVirtualTextureFeedback(Feedback, Result.PackedRequest, Result.PendingMips);
@@ -486,9 +490,9 @@ VTPageTableResult TextureLoadVirtualPageTable(
 	{
 		UV = UV * PageTableUniform.UVScale;
 		int vLevel = GetGlobalVirtualTextureMipBias(PageTableUniform.MipBiasGroup);
-#if PIXELSHADER
+#if VT_SUPPORT_DERIVATIVES
 		vLevel = TextureComputeVirtualMipLevel(Result, DDX(UV), DDY(UV), MipBias, SvPositionXY, PageTableUniform);
-#endif // PIXELSHADER
+#endif // VT_SUPPORT_DERIVATIVES
 		UV = ApplyAddressMode(UV, AddressU, AddressV);
 		TextureLoadVirtualPageTableInternal(Result, PageTable0, PageTableUniform, UV, vLevel);
 	}
@@ -508,9 +512,9 @@ VTPageTableResult TextureLoadVirtualPageTable(
 	{
 		UV = UV * PageTableUniform.UVScale;
 		int vLevel = GetGlobalVirtualTextureMipBias(PageTableUniform.MipBiasGroup);
-#if PIXELSHADER
+#if VT_SUPPORT_DERIVATIVES
 		vLevel = TextureComputeVirtualMipLevel(Result, DDX(UV), DDY(UV), MipBias, SvPositionXY, PageTableUniform);
-#endif // PIXELSHADER
+#endif // VT_SUPPORT_DERIVATIVES
 		UV = ApplyAddressMode(UV, AddressU, AddressV);
 		TextureLoadVirtualPageTableInternal(Result, PageTable0, PageTable1, PageTableUniform, UV, vLevel);
 		StoreVirtualTextureFeedback(Feedback, Result.PackedRequest, Result.PendingMips);
@@ -530,9 +534,9 @@ VTPageTableResult TextureLoadVirtualPageTable(
 	{
 		UV = UV * PageTableUniform.UVScale;
 		int vLevel = GetGlobalVirtualTextureMipBias(PageTableUniform.MipBiasGroup);
-#if PIXELSHADER
+#if VT_SUPPORT_DERIVATIVES
 		vLevel = TextureComputeVirtualMipLevel(Result, DDX(UV), DDY(UV), MipBias, SvPositionXY, PageTableUniform);
-#endif // PIXELSHADER
+#endif // VT_SUPPORT_DERIVATIVES
 		UV = ApplyAddressMode(UV, AddressU, AddressV);
 		TextureLoadVirtualPageTableInternal(Result, PageTable0, PageTable1, PageTableUniform, UV, vLevel);
 	}
@@ -700,9 +704,9 @@ VTPageTableResult TextureLoadVirtualPageTableAdaptive(
 	VTPageTableResult Result = (VTPageTableResult)0;
 	UV = UV * PageTableUniform.UVScale;
 	int vLevel = GetGlobalVirtualTextureMipBias(PageTableUniform.MipBiasGroup);
-#if PIXELSHADER
+#if VT_SUPPORT_DERIVATIVES
 	vLevel = TextureComputeVirtualMipLevel(Result, DDX(UV), DDY(UV), MipBias, SvPositionXY, PageTableUniform);
-#endif // PIXELSHADER
+#endif // VT_SUPPORT_DERIVATIVES
 	UV = ApplyAddressMode(UV, AddressU, AddressV);
 	ApplyAdaptivePageTableUniform(PageTableIndirection, Result, PageTableUniform, UV, vLevel);
 	TextureLoadVirtualPageTableInternal(Result, PageTable0, PageTableUniform, UV, vLevel);
