Fix hair transmittance artifact when with ViewRectOffset + VSM One-pass-projection.

#Synced-CL 48899741

diff --git a/Engine/Shaders/Private/HairStrands/HairStrandsDeepTransmittanceMask.usf b/Engine/Shaders/Private/HairStrands/HairStrandsDeepTransmittanceMask.usf
--- a/Engine/Shaders/Private/HairStrands/HairStrandsDeepTransmittanceMask.usf
+++ b/Engine/Shaders/Private/HairStrands/HairStrandsDeepTransmittanceMask.usf
@@ -383,10 +383,10 @@ void MainCS(uint2 DispatchThreadId : SV_DispatchThreadID)
 	const uint4 PackedShadowMask = ShadowMaskBitsTexture[PixelCoord];
 	const float3 Random = GetHairVoxelJitter(PixelCoord.xy, View.StateFrameIndexMod8, VirtualVoxel.JitterMode);
 
-	const uint2 LocalPixelPos = PixelCoord;
 
 	// Use VSM light grid index remapping to ensure we compute transmittance value only for shadowed light
 	// The logic to compute light index needs to be identical as in DeferredLightPixelShader.usf / ClusteredDeferredShadingPixelShader.usf
+	const uint2 LocalPixelPos = PixelCoord - View.ViewRectMinAndSize.xy;
 	const FCulledLightsGridHeader CulledLightGridHeader = VirtualShadowMapGetLightsGridHeader(LocalPixelPos, SceneDepth);
 	const uint LightCount = min(GetPackedShadowMaskMaxLightCount(), CulledLightGridHeader.NumLights);
 
@@ -522,4 +522,4 @@ void MainCS(uint2 DispatchThreadId : SV_DispatchThreadID)
 	OutputMask[Index] = PackTransmittanceMask(InitHairTransmittanceMask());
 }
 
-#endif // SHADER_CLEAR
\ No newline at end of file
+#endif // SHADER_CLEAR
