Fixed HLSL compile error due to sizeof returning ulong.

#Synced-CL 48933410

diff --git a/Engine/Shaders/Private/Skinning/AnimBankEval.usf b/Engine/Shaders/Private/Skinning/AnimBankEval.usf
--- a/Engine/Shaders/Private/Skinning/AnimBankEval.usf
+++ b/Engine/Shaders/Private/Skinning/AnimBankEval.usf
@@ -80,7 +80,7 @@ void BankEvaluateCS(uint3 GroupId : SV_GroupID, uint GroupIndex : SV_GroupIndex)
 	const FAnimBankBlockHeader BoneBlockHeader = BankBlockBuffer.Load<FAnimBankBlockHeader>(BoneBlock * uint(sizeof(FAnimBankBlockHeader)));
 	if (LocalTransformIndex < BoneBlockHeader.BlockTransformCount)
 	{
-		const uint BoneTransformSize      = sizeof(FBoneTransform);
+		const uint BoneTransformSize      = uint(sizeof(FBoneTransform));
 		const uint MappingTransformIndex  = BoneBlockHeader.MappingTransformIndex;
 		const uint SequenceTransformIndex = BoneBlockHeader.SequenceTransformIndex;
 		const uint TransformIndex         = LocalTransformIndex + (BoneBlockHeader.BlockLocalIndex * BONES_PER_GROUP);
@@ -140,12 +140,12 @@ void BankEvaluateCS(uint3 GroupId : SV_GroupID, uint GroupIndex : SV_GroupIndex)
 		MeshGlobalAnimPoseTransform = MultiplyBoneTransform(InvGlobalRefPoseTransform, MeshGlobalAnimPoseTransform);
 		const float4x3 CurrentTransform = (float4x3)MatrixFromBoneTransform(MeshGlobalAnimPoseTransform);
 
-		const uint OutputTransformOffset = BoneBlockHeader.OutputTransformIndex * sizeof(FCompressedBoneTransform);
+		const uint OutputTransformOffset = BoneBlockHeader.OutputTransformIndex * uint(sizeof(FCompressedBoneTransform));
 		StoreCompressedBoneTransform(TransformBuffer, OutputTransformOffset, TransformIndex, CurrentTransform);
 
 		if (BoneBlockHeader.OutputTransformCopyIndexOffset != 0)
 		{
-			const uint SecondaryOutputTransformOffset = (BoneBlockHeader.OutputTransformIndex + BoneBlockHeader.OutputTransformCopyIndexOffset) * sizeof(FCompressedBoneTransform);
+			const uint SecondaryOutputTransformOffset = (BoneBlockHeader.OutputTransformIndex + BoneBlockHeader.OutputTransformCopyIndexOffset) * uint(sizeof(FCompressedBoneTransform));
 			StoreCompressedBoneTransform(TransformBuffer, SecondaryOutputTransformOffset, TransformIndex, CurrentTransform);
 		}
 	}
