Mobile: Remove mobile specific BRDF context initialization, as it does not work corretly in all cases

#Synced-CL 48939854

diff --git a/Engine/Shaders/Private/BRDF.ush b/Engine/Shaders/Private/BRDF.ush
--- a/Engine/Shaders/Private/BRDF.ush
+++ b/Engine/Shaders/Private/BRDF.ush
@@ -59,26 +59,6 @@ void Init( inout BxDFContext Context, half3 N, half3 X, half3 Y, half3 V, half3
 	Context.YoH = (Context.YoL + Context.YoV) * InvLenH;
 }
 
-void InitMobile(inout BxDFContext Context, half3 N, half3 V, half3 L, half NoL)
-{
-	Context.NoL = NoL;
-	Context.NoV = dot(N, V);
-	Context.VoL = dot(V, L);
-	float3 H = normalize(float3(V + L));
-	Context.NoH = max(0, dot(N, H));
-	Context.VoH = max(0, dot(V, H));
-
-	//NoL = saturate( NoL );
-	//NoV = saturate( abs( NoV ) + 1e-5 );
-
-	Context.XoV = 0.0f;
-	Context.XoL = 0.0f;
-	Context.XoH = 0.0f;
-	Context.YoV = 0.0f;
-	Context.YoL = 0.0f;
-	Context.YoH = 0.0f;
-}
-
 // [ de Carpentier 2017, "Decima Engine: Advances in Lighting and AA" ]
 void SphereMaxNoH( inout BxDFContext Context, float SinAlpha, bool bNewtonIteration )
 {
diff --git a/Engine/Shaders/Private/ShadingModels.ush b/Engine/Shaders/Private/ShadingModels.ush
--- a/Engine/Shaders/Private/ShadingModels.ush
+++ b/Engine/Shaders/Private/ShadingModels.ush
@@ -231,12 +231,8 @@ FDirectLighting DefaultLitBxDF( FGBufferData GBuffer, half3 N, half3 V, FAreaLig
 		}
 		else
 		{
-#if SHADING_PATH_MOBILE
-			InitMobile(Context, N, V, AreaLight.SpecularL, AreaLight.NoL);
-#else
 			Init(Context, N, V, AreaLight.SpecularL);
-#endif
-
+			
 			NoV = Context.NoV;
 			VoH = Context.VoH;
 			NoH = Context.NoH;
@@ -587,11 +583,7 @@ void GetProfileDualSpecular(uint SubsurfaceProfileInt, half Roughness, half Opac
 FDirectLighting SubsurfaceProfileBxDF( FGBufferData GBuffer, half3 N, half3 V, FAreaLight AreaLight, FShadowTerms Shadow )
 {
 	BxDFContext Context;
-#if SHADING_PATH_MOBILE
-	InitMobile(Context, N, V, AreaLight.SpecularL, AreaLight.NoL);
-#else
 	Init( Context, N, V, AreaLight.SpecularL );
-#endif
 	SphereMaxNoH( Context, AreaLight.SphereSinAlpha, true );
 	Context.NoV = saturate( abs( Context.NoV ) + 1e-5 );
 
@@ -678,11 +670,7 @@ FDirectLighting ClothBxDF( FGBufferData GBuffer, half3 N, half3 V, FAreaLight Ar
 	const float  Cloth		= saturate(GBuffer.CustomData.a);
 
 	BxDFContext Context;
-#if SHADING_PATH_MOBILE
-	InitMobile(Context, N, V, AreaLight.SpecularL, AreaLight.NoL);
-#else
 	Init( Context, N, V, AreaLight.SpecularL );
-#endif
 	SphereMaxNoH( Context, AreaLight.SphereSinAlpha, true );
 	Context.NoV = saturate( abs( Context.NoV ) + 1e-5 );
 
@@ -983,11 +971,7 @@ FDirectLighting EyeBxDF( FGBufferData GBuffer, half3 N, half3 V, FAreaLight Area
 #endif
 
 	BxDFContext Context;
-#if SHADING_PATH_MOBILE
-	InitMobile(Context, N, V, AreaLight.SpecularL, AreaLight.NoL);
-#else
 	Init( Context, N, V, AreaLight.SpecularL );
-#endif
 	SphereMaxNoH( Context, AreaLight.SphereSinAlpha, false );
 	Context.NoV = saturate( abs( Context.NoV ) + 1e-5 );
 	const bool bIsRect = IsRectLight(AreaLight);
