VolumetricFog - added higher order filtering option for volumetric fog composition

Thanks Pawel Stolecki

#tests PC editor dx12
[RN] Added higher order filtering option for volumetric fog through r.VolumetricFog.Filtering.Quality
#Synced-CL 48943944

diff --git a/Engine/Shaders/Private/HeightFogCommon.ush b/Engine/Shaders/Private/HeightFogCommon.ush
--- a/Engine/Shaders/Private/HeightFogCommon.ush
+++ b/Engine/Shaders/Private/HeightFogCommon.ush
@@ -417,6 +417,78 @@ half4 CalculateHeightFog(float3 WorldPositionRelativeToCamera)
 	return CalculateHeightFog(WorldPositionRelativeToCamera, 0, GetPrimaryView());
 }
 
+float4 SampleVolumetricFog(float3 VolumeUV)
+{
+	return Texture3DSampleLevel(FogStruct.IntegratedLightScattering, SharedIntegratedLightScatteringSampler, VolumeUV, 0);
+}
+
+#if !defined(PERMUTATION_VOLUMETRIC_FOG_FILTERING_QUALITY)
+// Only passes defining this value can use higher quality filtering. Forward shading won't for instance.
+#define PERMUTATION_VOLUMETRIC_FOG_FILTERING_QUALITY 0
+#endif
+
+float4 SampleVolumetricFogFiltered(float3 VolumeUV)
+{
+#if PERMUTATION_VOLUMETRIC_FOG_FILTERING_QUALITY > 0
+		//Prerequisities for Tri & Bicubic
+		float3 PixelPos = VolumeUV * View.VolumetricFogGridSize + 0.5f;
+		uint3 PixelPosFloor = floor(PixelPos);
+		float3 Gradient = frac(PixelPos);
+		float3 OneMinusGradient = 1.0f - Gradient;
+			
+		float3 w0 = 1.0 / 6.0 * OneMinusGradient * OneMinusGradient * OneMinusGradient;
+		float3 w1 = 2.0 / 3.0 - 0.5 * Gradient * Gradient * (2.0 - Gradient);
+		float3 w2 = 2.0 / 3.0 - 0.5 * OneMinusGradient * OneMinusGradient * (2.0 - OneMinusGradient);
+		float3 w3 = 1.0 / 6.0 * Gradient * Gradient * Gradient;
+
+		float3 g0 = w0 + w1;
+		float3 g1 = w2 + w3;
+		float3 h0 = View.VolumetricFogInvGridSize * ((w1 / g0) - 1.5 + PixelPosFloor);
+		float3 h1 = View.VolumetricFogInvGridSize * ((w3 / g1) + 0.5 + PixelPosFloor);
+		h1.xy = min(h1.xy, View.VolumetricFogUVMax);
+		h0.z = VolumeUV.z;
+#endif
+	
+#if PERMUTATION_VOLUMETRIC_FOG_FILTERING_QUALITY == 2
+		//Tricubic
+		float4 tex000 = SampleVolumetricFog(h0);
+		float4 tex100 = SampleVolumetricFog(float3(h1.x, h0.y, h0.z));
+		tex000 = lerp(tex100, tex000, g0.x);
+
+		float4 tex010 = SampleVolumetricFog(float3(h0.x, h1.y, h0.z));
+		float4 tex110 = SampleVolumetricFog(float3(h1.x, h1.y, h0.z));
+		tex010 = lerp(tex110, tex010, g0.x);
+		tex000 = lerp(tex010, tex000, g0.y);
+
+		float4 tex001 = SampleVolumetricFog(float3(h0.x, h0.y, h1.z));
+		float4 tex101 = SampleVolumetricFog(float3(h1.x, h0.y, h1.z));
+		tex001 = lerp(tex101, tex001, g0.x);
+
+		float4 tex011 = SampleVolumetricFog(float3(h0.x, h1.y, h1.z));
+		float4 tex111 = SampleVolumetricFog(h1);
+		tex011 = lerp(tex111, tex011, g0.x);
+		tex001 = lerp(tex011, tex001, g0.y);
+
+		return lerp(tex001, tex000, g0.z);	
+#elif PERMUTATION_VOLUMETRIC_FOG_FILTERING_QUALITY == 1
+		//Bicubic
+		float4 tex000 = SampleVolumetricFog(h0);
+		float4 tex100 = SampleVolumetricFog(float3(h1.x, h0.y, h0.z));
+		tex000 = lerp(tex100, tex000, g0.x);
+
+		float4 tex010 = SampleVolumetricFog(float3(h0.x, h1.y, h0.z));
+		float4 tex110 = SampleVolumetricFog(float3(h1.x, h1.y, h0.z));
+		tex010 = lerp(tex110, tex010, g0.x);
+	
+		return lerp(tex010, tex000, g0.y);	
+#else
+		//Trilinear
+	return SampleVolumetricFog(VolumeUV);
+#endif
+
+	// Idea: stochastic sampling into noise resolved by TSR to make higher order filtering cheaper.
+}
+
 float4 CombineVolumetricFog(float4 GlobalFog, float3 VolumeUV, uint EyeIndex, float SceneDepth, ViewState InView)
 {
 	float4 VolumetricFogLookup = float4(0, 0, 0, 1);
@@ -443,7 +515,9 @@ float4 CombineVolumetricFog(float4 GlobalFog, float3 VolumeUV, uint EyeIndex, fl
 		}
 	#else
 		VolFogStartDistance = FogStruct.VolumetricFogStartDistance;
-		VolumetricFogLookup = Texture3DSampleLevel(FogStruct.IntegratedLightScattering, SharedIntegratedLightScatteringSampler, VolumeUV, 0);
+		// < #RED_VolumetricFiltering pawel.stolecki 2024-04-24
+		VolumetricFogLookup = SampleVolumetricFogFiltered(VolumeUV);
+		// > #RED_VolumetricFiltering pawel.stolecki 2024-04-24
 	#endif
 
 		// IntegratedLightScattering is pre-exposed, remove pre exposure now so that it can correctly be applied later
diff --git a/Engine/Source/Runtime/Engine/Public/SceneView.h b/Engine/Source/Runtime/Engine/Public/SceneView.h
--- a/Engine/Source/Runtime/Engine/Public/SceneView.h
+++ b/Engine/Source/Runtime/Engine/Public/SceneView.h
@@ -1034,6 +1034,7 @@ enum ETranslucencyVolumeCascade
 	VIEW_UNIFORM_BUFFER_MEMBER(float, ViewportScaleUI) \
 	VIEW_UNIFORM_BUFFER_MEMBER(FIntPoint, CursorPosition) \
 	VIEW_UNIFORM_BUFFER_MEMBER(float, bCheckerboardSubsurfaceProfileRendering) \
+	VIEW_UNIFORM_BUFFER_MEMBER(FVector3f, VolumetricFogGridSize) \
 	VIEW_UNIFORM_BUFFER_MEMBER(FVector3f, VolumetricFogInvGridSize) \
 	VIEW_UNIFORM_BUFFER_MEMBER(FVector3f, VolumetricFogGridZParams) \
 	VIEW_UNIFORM_BUFFER_MEMBER(FVector2f, VolumetricFogSVPosToVolumeUV) \
diff --git a/Engine/Source/Runtime/Renderer/Private/FogRendering.cpp b/Engine/Source/Runtime/Renderer/Private/FogRendering.cpp
--- a/Engine/Source/Runtime/Renderer/Private/FogRendering.cpp
+++ b/Engine/Source/Runtime/Renderer/Private/FogRendering.cpp
@@ -63,6 +63,14 @@ bool PlatformUsesBasicFogFeatures(EShaderPlatform ShaderPlatform)
 	return IsMobilePlatform(ShaderPlatform) && CVarMobilePixelFogQuality.GetValueOnAnyThread() == 0;
 }
 
+static TAutoConsoleVariable<int> CVarVolumetricFogFilteringQuality(
+	TEXT("r.VolumetricFog.Filtering.Quality"),
+	0,
+	TEXT(" 0: Trilinear\n")
+	TEXT(" 1: Bicubic\n")
+	TEXT(" 2: Tricubic"),
+	ECVF_RenderThreadSafe | ECVF_Scalability);
+
 IMPLEMENT_GLOBAL_SHADER_PARAMETER_STRUCT(FFogUniformParameters, "FogStruct");
 
 static void SetupDummyExponentialHeightFogUniformParameters(FFogUniformParameters& OutParameters)
@@ -164,7 +172,8 @@ class FExponentialHeightFogPS : public FGlobalShader
 	class FSupportVolumetricFog : SHADER_PERMUTATION_BOOL("PERMUTATION_SUPPORT_VOLUMETRIC_FOG");
 	class FSupportLocalFogVolume : SHADER_PERMUTATION_BOOL("PERMUTATION_SUPPORT_LOCAL_FOG_VOLUME");
 	class FSampleFogOnClouds : SHADER_PERMUTATION_BOOL("PERMUTATION_SAMPLE_FOG_ON_CLOUDS");
-	using FPermutationDomain = TShaderPermutationDomain<FSupportFogInScatteringTexture, FSupportFogDirectionalLightInScattering, FSupportVolumetricFog, FSupportLocalFogVolume, FSampleFogOnClouds>;
+	class FVolumetricFogFilteringQuality : SHADER_PERMUTATION_INT("PERMUTATION_VOLUMETRIC_FOG_FILTERING_QUALITY", 3);
+	using FPermutationDomain = TShaderPermutationDomain<FSupportFogInScatteringTexture, FVolumetricFogFilteringQuality, FSupportFogDirectionalLightInScattering, FSupportVolumetricFog, FSupportLocalFogVolume, FSampleFogOnClouds>;
 
 	BEGIN_SHADER_PARAMETER_STRUCT(FParameters, )
 		SHADER_PARAMETER_STRUCT_REF(FViewUniformShaderParameters, ViewUniformBuffer)
@@ -434,6 +443,11 @@ static void RenderViewFog(
 	PsPermutationVector.Set<FExponentialHeightFogPS::FSupportVolumetricFog>(bShouldRenderVolumetricFog);
 	PsPermutationVector.Set<FExponentialHeightFogPS::FSupportLocalFogVolume>(bFogComposeLocalFogVolumes);
 	PsPermutationVector.Set<FExponentialHeightFogPS::FSampleFogOnClouds>(bSampleFogOnClouds);
+
+	//Ensure Trilinear for Reflection Capture to avoid needless cost
+	int VolumetricFogFilteringQuality = View.CachedViewUniformShaderParameters->RenderingReflectionCaptureMask == 0.0f ? CVarVolumetricFogFilteringQuality.GetValueOnAnyThread() : 0;
+	PsPermutationVector.Set<FExponentialHeightFogPS::FVolumetricFogFilteringQuality>(FMath::Clamp(VolumetricFogFilteringQuality, 0, 2));
+
 	TShaderMapRef<FExponentialHeightFogPS> PixelShader(View.ShaderMap, PsPermutationVector);
 	GraphicsPSOInit.BoundShaderState.PixelShaderRHI = PixelShader.GetPixelShader();
 
diff --git a/Engine/Source/Runtime/Renderer/Private/VolumetricFog.cpp b/Engine/Source/Runtime/Renderer/Private/VolumetricFog.cpp
--- a/Engine/Source/Runtime/Renderer/Private/VolumetricFog.cpp
+++ b/Engine/Source/Runtime/Renderer/Private/VolumetricFog.cpp
@@ -1341,6 +1341,7 @@ void FViewInfo::SetupVolumetricFogUniformBufferParameters(FViewUniformShaderPara
 		const FIntVector VolumetricFogResourceGridSize = GetVolumetricFogResourceGridSize(*this, VolumetricFogResourceGridPixelSize);
 		const FIntVector VolumetricFogViewGridSize = GetVolumetricFogViewGridSize(*this, VolumetricFogViewGridPixelSize);
 
+		ViewUniformShaderParameters.VolumetricFogGridSize = FVector3f(VolumetricFogResourceGridSize.X, VolumetricFogResourceGridSize.Y, VolumetricFogResourceGridSize.Z);
 		ViewUniformShaderParameters.VolumetricFogInvGridSize = FVector3f(1.0f / VolumetricFogResourceGridSize.X, 1.0f / VolumetricFogResourceGridSize.Y, 1.0f / VolumetricFogResourceGridSize.Z);
 
 		const FVector ZParams = GetVolumetricFogGridZParams(FogInfo.VolumetricFogStartDistance, NearClippingDistance, FogInfo.VolumetricFogDistance, VolumetricFogResourceGridSize.Z);
@@ -1351,6 +1352,7 @@ void FViewInfo::SetupVolumetricFogUniformBufferParameters(FViewUniformShaderPara
 	}
 	else
 	{
+		ViewUniformShaderParameters.VolumetricFogGridSize = FVector3f::ZeroVector;
 		ViewUniformShaderParameters.VolumetricFogInvGridSize = FVector3f::ZeroVector;
 		ViewUniformShaderParameters.VolumetricFogGridZParams = FVector3f::ZeroVector;
 		ViewUniformShaderParameters.VolumetricFogSVPosToVolumeUV = FVector2f::ZeroVector;
