VolumetricFog - removed unnecessary tag.

#rnx
#Synced-CL 48943944

diff --git a/Engine/Shaders/Private/HeightFogCommon.ush b/Engine/Shaders/Private/HeightFogCommon.ush
--- a/Engine/Shaders/Private/HeightFogCommon.ush
+++ b/Engine/Shaders/Private/HeightFogCommon.ush
@@ -515,9 +515,7 @@ float4 CombineVolumetricFog(float4 GlobalFog, float3 VolumeUV, uint EyeIndex, fl
 		}
 	#else
 		VolFogStartDistance = FogStruct.VolumetricFogStartDistance;
-		// < #RED_VolumetricFiltering pawel.stolecki 2024-04-24
 		VolumetricFogLookup = SampleVolumetricFogFiltered(VolumeUV);
-		// > #RED_VolumetricFiltering pawel.stolecki 2024-04-24
 	#endif
 
 		// IntegratedLightScattering is pre-exposed, remove pre exposure now so that it can correctly be applied later
