Path Tracer: Improve shading normals on diffuse BRDFs.

Only adjust the shading normal for specular BRDFs. This allows the diffuse response to use the original shading normal, even if it points away from the camera which general matches raster behavior better.

#Synced-CL 48993796

diff --git a/Engine/Shaders/Private/MonteCarlo.ush b/Engine/Shaders/Private/MonteCarlo.ush
--- a/Engine/Shaders/Private/MonteCarlo.ush
+++ b/Engine/Shaders/Private/MonteCarlo.ush
@@ -261,6 +261,27 @@ float4 CosineSampleHemisphere( float2 E )
 	return float4(H, PDF);
 }
 
+// PDF = NoL / PI
+// Simpler than cosine sampling + tangent transform if you only want a world space direction
+float4 CosineSampleHemisphereWorld( float2 E, float3 N )
+{
+	float Phi = 2 * PI * E.x;
+	float CosTheta = 1 - 2 * E.y;
+	float SinTheta = sqrt( 1 - CosTheta * CosTheta );
+
+	float3 U;
+	U.x = SinTheta * cos( Phi );
+	U.y = SinTheta * sin( Phi );
+	U.z = CosTheta;
+
+	float3 V = normalize(U + N);
+
+	float PDF = saturate(dot(V, N)) * (1.0 / PI);
+
+	return float4( V, PDF );
+}
+
+
 // PDF = NoL / PI
 float4 CosineSampleHemisphereConcentric(float2 E)
 {
diff --git a/Engine/Shaders/Private/PathTracing/Material/PathTracingClearCoat.ush b/Engine/Shaders/Private/PathTracing/Material/PathTracingClearCoat.ush
--- a/Engine/Shaders/Private/PathTracing/Material/PathTracingClearCoat.ush
+++ b/Engine/Shaders/Private/PathTracing/Material/PathTracingClearCoat.ush
@@ -42,13 +42,12 @@ FClearCoatData CreateClearCoatData(FPathTracingPayload Payload, float3 V_World)
 {
 	FClearCoatData Data = (FClearCoatData)0;
 
-	const float3 CoatN = Payload.WorldNormal;
-	const float3 BaseN = Payload.GetClearCoatBottomNormal();
+	const float3 CoatN = AdjustShadingNormal(Payload.WorldNormal               , Payload.WorldGeoNormal, V_World);
+	const float3 BaseN = AdjustShadingNormal(Payload.GetClearCoatBottomNormal(), Payload.WorldGeoNormal, V_World);
 
 	Data.BaseBasis = GetGGXBasis(Payload.Roughness, Payload.Anisotropy, BaseN, Payload.WorldTangent, Data.BaseAlpha);
 	Data.CoatBasis = GetGGXBasis(Payload.GetClearCoatRoughness(), CoatN, Data.CoatAlpha);
 	
-
 	Data.CoatV = mul(Data.CoatBasis, V_World);
 	Data.BaseV = mul(Data.BaseBasis, V_World);
 
@@ -129,9 +128,16 @@ FMaterialEval ClearCoat_EvalMaterial(
 	const float SpecPdf = SpecGGXResult.y;
 
 	// Base Diffuse lobe
-	const float3 Diffuse = GetPathTracingDiffuseModel(Data.DiffColor, Payload.Roughness, BaseNoV, BaseNoL, BaseVoH, BaseNoH);
-	const float3 DiffWeight = Transmission * (1 - Data.Spec.E) * Diffuse * ShadowTerminatorTerm(L_World, Payload.WorldNormal, Payload.WorldSmoothNormal);
-	const float DiffPdf = BaseNoL / PI;
+	const float3 N_World = Payload.GetClearCoatBottomNormal();
+	const float DiffNoV = saturate(dot(N_World, V_World));
+	const float DiffNoL = saturate(dot(N_World, L_World));
+	const float3 H_World = normalize(V_World + L_World);
+	const float DiffVoH = saturate(dot(V_World, H_World));
+	const float DiffNoH = saturate(dot(N_World, H_World));
+	const float3 Diffuse = GetPathTracingDiffuseModel(Data.DiffColor, Payload.Roughness, DiffNoV, DiffNoL, DiffVoH, DiffNoH);
+	const float3 DiffWeight = Transmission * (1 - Data.Spec.E) * Diffuse * ShadowTerminatorTerm(L_World, N_World, Payload.WorldSmoothNormal);
+	const float DiffPdf = DiffNoL / PI;
+
 
 	// Combine all lobes together
 	FMaterialEval Result = NullMaterialEval();
@@ -190,8 +196,7 @@ FMaterialSample ClearCoat_SampleMaterial(
 		// Sampled base diffuse
 		RandSample.x = RescaleRandomNumber(RandSample.x, Data.LobeCdf.y, 1.0);
 
-		const float3 L = CosineSampleHemisphere(RandSample.xy).xyz;
-		Result.Direction = mul(L, Data.BaseBasis);
+		Result.Direction = CosineSampleHemisphereWorld(RandSample.xy, Payload.WorldNormal).xyz;
 		Result.Roughness = 1.0;
 		Result.ScatterType = PATHTRACER_SCATTER_DIFFUSE;
 	}
@@ -221,7 +226,9 @@ FMaterialSample ClearCoat_SampleMaterial(
 	const float2 SpecGGXResult = GGXEvalReflection(BaseL, Data.BaseV, BaseH, Data.BaseAlpha);
 	const float CoatPdf = CoatGGXResult.y;
 	const float SpecPdf = SpecGGXResult.y;
-	const float DiffPdf = BaseNoL / PI;
+	const float3 N_World = Payload.GetClearCoatBottomNormal();
+	const float DiffNoL = saturate(dot(N_World, L_World));
+	const float DiffPdf = DiffNoL / PI;
 
 	// evaluate the lobe we sampled
 	if (bSampledSpecular)
@@ -237,8 +244,13 @@ FMaterialSample ClearCoat_SampleMaterial(
 	else
 	{
 		// Base Diffuse lobe
-		const float3 Diffuse = GetPathTracingDiffuseModel(Data.DiffColor, Payload.Roughness, BaseNoV, BaseNoL, BaseVoH, BaseNoH);
-		const float3 DiffWeight = Transmission * (1 - Data.Spec.E) * Diffuse * ShadowTerminatorTerm(L_World, Payload.WorldNormal, Payload.WorldSmoothNormal);
+		const float DiffNoV = saturate(dot(N_World, V_World));
+		const float3 H_World = normalize(V_World + L_World);
+		const float DiffVoH = saturate(dot(V_World, H_World));
+		const float DiffNoH = saturate(dot(N_World, H_World));
+
+		const float3 Diffuse = GetPathTracingDiffuseModel(Data.DiffColor, Payload.Roughness, DiffNoV, DiffNoL, DiffVoH, DiffNoH);
+		const float3 DiffWeight = Transmission * (1 - Data.Spec.E) * Diffuse * ShadowTerminatorTerm(L_World, N_World, Payload.WorldSmoothNormal);
 
 		Result.AddLobeWithMIS(DiffWeight, DiffPdf, Data.LobePdf.z);
 
diff --git a/Engine/Shaders/Private/PathTracing/Material/PathTracingCloth.ush b/Engine/Shaders/Private/PathTracing/Material/PathTracingCloth.ush
--- a/Engine/Shaders/Private/PathTracing/Material/PathTracingCloth.ush
+++ b/Engine/Shaders/Private/PathTracing/Material/PathTracingCloth.ush
@@ -28,7 +28,9 @@ FClothData PrepareClothData(FPathTracingPayload Payload, float3 V_World)
 {
 	FClothData Data = (FClothData)0;
 
-	Data.Basis = GetGGXBasis(Payload.Roughness, Payload.Anisotropy, Payload.WorldNormal, Payload.WorldTangent, Data.Alpha);
+	float3 NSpec = AdjustShadingNormal(Payload.WorldNormal, Payload.WorldGeoNormal, V_World);
+
+	Data.Basis = GetGGXBasis(Payload.Roughness, Payload.Anisotropy, NSpec, Payload.WorldTangent, Data.Alpha);
 
 	Data.V = mul(Data.Basis, V_World);
 
@@ -72,7 +74,8 @@ FMaterialEval Cloth_EvalMaterial(
 
 	// Diffuse Lobe
 	const float ShadowTerminator = ShadowTerminatorTerm(L_World, Payload.WorldNormal, Payload.WorldSmoothNormal);
-	Result.AddLobeWithMIS(Data.DiffWeight * ShadowTerminator * DiffuseSpecularScale.x, NoL / PI, Data.LobePdf.x);
+	const float DiffNoL = saturate(dot(L_World, Payload.WorldNormal));
+	Result.AddLobeWithMIS(Data.DiffWeight * ShadowTerminator * DiffuseSpecularScale.x, DiffNoL / PI, Data.LobePdf.x);
 
 	const float ClothBlend = Payload.GetClothAmount();
 
@@ -112,18 +115,15 @@ FMaterialSample Cloth_SampleMaterial(
 	float OutRoughness = 1.0;
 	const bool bSampledDiffuse = RandSample.x < Data.LobeCdf.x;
 	const bool bSampledCloth   = RandSample.x < Data.LobeCdf.y;
-	if (bSampledDiffuse || bSampledCloth)
+	if (bSampledDiffuse)
 	{
-		if (bSampledDiffuse)
-		{
-			RandSample.x = RescaleRandomNumber(RandSample.x, 0.0, Data.LobeCdf.x);
-		}
-		else
-		{
-			RandSample.x = RescaleRandomNumber(RandSample.x, Data.LobeCdf.x, Data.LobeCdf.y);
-		}
-
-		// Lambert
+		RandSample.x = RescaleRandomNumber(RandSample.x, 0.0, Data.LobeCdf.x);
+		float3 L_World = CosineSampleHemisphereWorld(RandSample.xy, Payload.WorldNormal).xyz;
+		L = mul(Data.Basis, L_World);
+	}
+	else if (bSampledCloth)
+	{
+		RandSample.x = RescaleRandomNumber(RandSample.x, Data.LobeCdf.x, Data.LobeCdf.y);
 		L = CosineSampleHemisphere(RandSample.xy).xyz;
 		H = normalize(L + Data.V);
 	}
@@ -153,7 +153,8 @@ FMaterialSample Cloth_SampleMaterial(
 	const float ClothBlend = Payload.GetClothAmount();
 	const float2 GGXResult = GGXEvalReflection(L, V, H, Data.Alpha);
 	const float SpecPdf = GGXResult.y;
-	const float DiffPdf = NoL / PI;
+	const float DiffNoL = saturate(dot(Payload.WorldNormal, L_World));
+	const float DiffPdf = DiffNoL / PI;
 	const float ClothPdf = DiffPdf;
 
 
diff --git a/Engine/Shaders/Private/PathTracing/Material/PathTracingDefaultLit.ush b/Engine/Shaders/Private/PathTracing/Material/PathTracingDefaultLit.ush
--- a/Engine/Shaders/Private/PathTracing/Material/PathTracingDefaultLit.ush
+++ b/Engine/Shaders/Private/PathTracing/Material/PathTracingDefaultLit.ush
@@ -25,7 +25,9 @@ FDefaultLitData CreateDefaultLitData(FPathTracingPayload Payload, float3 V_World
 {
 	FDefaultLitData Data = (FDefaultLitData)0;
 
-	Data.Basis = GetGGXBasis(Payload.Roughness, Payload.Anisotropy, Payload.WorldNormal, Payload.WorldTangent, Data.Alpha);
+	float3 NSpec = AdjustShadingNormal(Payload.WorldNormal, Payload.WorldGeoNormal, V_World);
+
+	Data.Basis = GetGGXBasis(Payload.Roughness, Payload.Anisotropy, NSpec, Payload.WorldTangent, Data.Alpha);
 	Data.V = mul(Data.Basis, V_World);
 	const float NoV = saturate(Data.V.z);
 	Data.Spec = ComputeGGXSpecEnergyTermsRGB(Payload.Roughness, NoV, Payload.SpecularColor);
@@ -58,8 +60,14 @@ FMaterialEval DefaultLit_EvalMaterial(
 	FMaterialEval Result = NullMaterialEval();
 
 	// Diffuse Lobe
-	const float3 Diffuse = GetPathTracingDiffuseModel(Data.DiffuseWeight, Payload.Roughness, NoV, NoL, VoH, NoH);
-	Result.AddLobeWithMIS(DiffuseSpecularScale.x * Diffuse * ShadowTerminatorTerm(L_World, Payload.WorldNormal, Payload.WorldSmoothNormal), NoL / PI, Data.LobeProb);
+	const float DiffNoV = saturate(dot(Payload.WorldNormal, V_World));
+	const float DiffNoL = saturate(dot(Payload.WorldNormal, L_World));
+	const float3 H_World = normalize(V_World + L_World);
+	const float DiffVoH = saturate(dot(V_World, H_World));
+	const float DiffNoH = saturate(dot(Payload.WorldNormal, H_World));
+	const float3 Diffuse = GetPathTracingDiffuseModel(Data.DiffuseWeight, Payload.Roughness, DiffNoV, DiffNoL, DiffVoH, DiffNoH);
+	const float DiffPdf = DiffNoL / PI;
+	Result.AddLobeWithMIS(DiffuseSpecularScale.x * Diffuse * ShadowTerminatorTerm(L_World, Payload.WorldNormal, Payload.WorldSmoothNormal), DiffPdf, Data.LobeProb);
 
 	// Specular lobe
 	const float2 GGXResult = GGXEvalReflection(L, V, H, Data.Alpha);
@@ -84,6 +92,8 @@ FMaterialSample DefaultLit_SampleMaterial(
 
 	const float3 V = Data.V;
 	float3 L = 0, H = 0;
+	FMaterialSample Result = NullMaterialSample();
+	Result.PositionBiasSign = 1;
 
 	// Randomly choose to sample diffuse or specular
 	const bool bSampledDiffuse = RandSample.x < Data.LobeProb;
@@ -92,8 +102,8 @@ FMaterialSample DefaultLit_SampleMaterial(
 		RandSample.x = RescaleRandomNumber(RandSample.x, 0.0, Data.LobeProb);
 
 		// Lambert
-		// TODO: evaluate CosineSampleHemisphereConcentric
-		L = CosineSampleHemisphere(RandSample.xy).xyz;
+		Result.Direction = CosineSampleHemisphereWorld(RandSample.xy, Payload.WorldNormal).xyz;
+		L = mul(Data.Basis, Result.Direction);
 		H = normalize(L + V);
 	}
 	else
@@ -107,10 +117,11 @@ FMaterialSample DefaultLit_SampleMaterial(
 			// invalid output direction, exit early
 			return NullMaterialSample();
 		}
+		Result.Direction = normalize(mul(L, Data.Basis));
 	}
 
 	// transform to world space
-	const float3 L_World = normalize(mul(L, Data.Basis));
+	const float3 L_World = Result.Direction;
 
 	const float NoV = saturate(V.z);
 	const float NoL = saturate(L.z);
@@ -119,14 +130,20 @@ FMaterialSample DefaultLit_SampleMaterial(
 
 	const float2 GGXResult = GGXEvalReflection(L, V, H, Data.Alpha);
 	const float SpecPdf = GGXResult.y;
-	const float DiffPdf = NoL / PI;
+	const float DiffNoL = saturate(dot(Payload.WorldNormal, L_World));
+	const float DiffPdf = DiffNoL / PI;
 
-	FMaterialSample Result = CreateMaterialSample(L_World, 0.0, 0.0, 1.0, 1.0, PATHTRACER_SCATTER_DIFFUSE);
 	if (bSampledDiffuse)
 	{
-		const float3 Diffuse = GetPathTracingDiffuseModel(Data.DiffuseWeight, Payload.Roughness, NoV, NoL, VoH, NoH);
+		const float DiffNoV = saturate(dot(Payload.WorldNormal, V_World));
+		const float3 H_World = normalize(V_World + L_World);
+		const float DiffVoH = saturate(dot(V_World, H_World));
+		const float DiffNoH = saturate(dot(Payload.WorldNormal, H_World));
+		const float3 Diffuse = GetPathTracingDiffuseModel(Data.DiffuseWeight, Payload.Roughness, DiffNoV, DiffNoL, DiffVoH, DiffNoH);
 		Result.AddLobeWithMIS(Diffuse * ShadowTerminatorTerm(L_World, Payload.WorldNormal, Payload.WorldSmoothNormal), DiffPdf, Data.LobeProb);
 		Result.Pdf += (1.0 - Data.LobeProb) * SpecPdf;
+		Result.Roughness = 1;
+		Result.ScatterType = PATHTRACER_SCATTER_DIFFUSE;
 	}
 	else
 	{
diff --git a/Engine/Shaders/Private/PathTracing/Material/PathTracingLambert.ush b/Engine/Shaders/Private/PathTracing/Material/PathTracingLambert.ush
--- a/Engine/Shaders/Private/PathTracing/Material/PathTracingLambert.ush
+++ b/Engine/Shaders/Private/PathTracing/Material/PathTracingLambert.ush
@@ -12,9 +12,9 @@ FMaterialSample Lambert_SampleMaterial(
 	const float3 N_World = Payload.WorldNormal;
 	const float3 BaseColor = Payload.GetBaseColor();
 
-	float4 SampledValue = CosineSampleHemisphere(RandSample.xy);
+	float4 SampledValue = CosineSampleHemisphereWorld(RandSample.xy, N_World);
 
-	float3 L_World = TangentToWorld(SampledValue.xyz, N_World);
+	float3 L_World = SampledValue.xyz;
 #if SIMPLIFIED_MATERIAL_SHADER
 	float ShadowTerminator = 1;
 #else
@@ -54,16 +54,15 @@ FMaterialSample Simplified_SampleMaterial(
 	float3 FrontColor = Payload.BaseColor;
 	float3 BackColor = Payload.SubsurfaceColor;
 
-	float4 SampledValue = CosineSampleHemisphere(RandSample.xy);
-
 	float TransmissionProb = LobeSelectionProb(BackColor, FrontColor);
 	bool bDoTransmission = RandSample.z < TransmissionProb;
 	float3 Weight = bDoTransmission ? BackColor / TransmissionProb : FrontColor / (1 - TransmissionProb);
+	float Sign = bDoTransmission ? -1.0 : 1.0;
+	float4 SampledValue = CosineSampleHemisphereWorld(RandSample.xy, Sign * N_World);
 	float Pdf = SampledValue.w * (bDoTransmission ? TransmissionProb : 1.0 - TransmissionProb);
 
 	Weight *= Payload.BSDFOpacity;
-	float Sign = bDoTransmission ? -1.0 : 1.0;
-	return CreateMaterialSample(TangentToWorld(SampledValue.xyz, N_World * Sign), Weight, Pdf, Sign, 1.0, PATHTRACER_SCATTER_DIFFUSE);
+	return CreateMaterialSample(SampledValue.xyz, Weight, Pdf, Sign, 1.0, PATHTRACER_SCATTER_DIFFUSE);
 }
 
 FMaterialEval Simplified_EvalMaterial(
diff --git a/Engine/Shaders/Private/PathTracing/Material/PathTracingSubstrate.ush b/Engine/Shaders/Private/PathTracing/Material/PathTracingSubstrate.ush
--- a/Engine/Shaders/Private/PathTracing/Material/PathTracingSubstrate.ush
+++ b/Engine/Shaders/Private/PathTracing/Material/PathTracingSubstrate.ush
@@ -31,7 +31,8 @@ struct FSubstrateSingleSlabData
 
 	void PrepareBSDF(float3 V_World, FPathTracingPayload Payload, bool bUseSSSColor = false)
 	{
-		Basis = GetGGXBasis(Payload.RoughnessData.x, Payload.Anisotropy, Payload.WorldNormal, Payload.WorldTangent, Alpha0);
+		float3 NSpec = AdjustShadingNormal(Payload.WorldNormal, Payload.WorldGeoNormal, V_World);
+		Basis = GetGGXBasis(Payload.RoughnessData.x, Payload.Anisotropy, NSpec, Payload.WorldTangent, Alpha0);
 		Alpha1 = GetGGXAlpha(Payload.RoughnessData.y, Payload.Anisotropy);
 
 		V = mul(Basis, V_World);
@@ -80,32 +81,31 @@ FMaterialSample Substrate_SampleMaterial(
 
 
 	float3 L = 0, H = 0, V = Data.V;
-	float OutRoughness = 1;
 
 	const bool bSampledDiffuse  = RandSample.x < Data.LobeCdf.x;
 	const bool bSampledSpecular = RandSample.x < Data.LobeCdf.z;
+
+	FMaterialSample Result = NullMaterialSample();
+	Result.Roughness = 1;
+	Result.PositionBiasSign = 1;
 	if (bSampledDiffuse)
 	{
-		RandSample.x = RescaleRandomNumber(RandSample.x, 0.0, Data.LobeCdf.x);
 		// diffuse lobe
-		L = CosineSampleHemisphere(RandSample.xy).xyz;
+		RandSample.x = RescaleRandomNumber(RandSample.x, 0.0, Data.LobeCdf.x);
+		// cosine sampling without needing to form an orthonormal basis
+		Result.Direction = CosineSampleHemisphereWorld(RandSample.xy, Payload.WorldNormal).xyz;
+
+		// setup vectors for MIS math
+		L = mul(Data.Basis, Result.Direction);
 		H = normalize(L + V);
 	}
 	else if (bSampledSpecular)
 	{
 		// specular lobes
 		const bool bUseSpec0 = RandSample.x < Data.LobeCdf.y;
-		if (bUseSpec0)
-		{
-			RandSample.x = RescaleRandomNumber(RandSample.x, Data.LobeCdf.x, Data.LobeCdf.y);
-			OutRoughness = Payload.RoughnessData.x;
-		}
-		else
-		{
-			RandSample.x = RescaleRandomNumber(RandSample.x, Data.LobeCdf.y, Data.LobeCdf.z);
-			OutRoughness = Payload.RoughnessData.y;
-		}
 
+		Result.Roughness = bUseSpec0 ? Payload.RoughnessData.x : Payload.RoughnessData.y;
+		RandSample.x = RescaleRandomNumber(RandSample.x, bUseSpec0 ? Data.LobeCdf.x : Data.LobeCdf.y, bUseSpec0 ? Data.LobeCdf.y : Data.LobeCdf.z);
 		H = ImportanceSampleVisibleGGX(RandSample.xy, bUseSpec0 ? Data.Alpha0 : Data.Alpha1, V).xyz;
 		L = reflect(-V, H);
 		if (L.z <= 0)
@@ -113,6 +113,7 @@ FMaterialSample Substrate_SampleMaterial(
 			// invalid output direction, exit early
 			return NullMaterialSample();
 		}
+		Result.Direction = mul(L, Data.Basis);
 	}
 	else
 	{
@@ -120,32 +121,36 @@ FMaterialSample Substrate_SampleMaterial(
 		RandSample.x = RescaleRandomNumber(RandSample.x, Data.LobeCdf.z, 1.0);
 		L = Data.Fuzz.Sample(RandSample.xy);
 		H = normalize(L + V);
+		Result.Direction = mul(L, Data.Basis);
 	}
 
 	// transform to world space
-	const float3 L_World = normalize(mul(L, Data.Basis));
-	const float3 N_World = Payload.WorldNormal;
+	const float3 L_World = Result.Direction;
 	const float NoV = saturate(V.z);
 	const float NoL = saturate(L.z);
 	const float VoH = saturate(dot(V, H));
 	const float NoH = saturate(H.z);
 
-	FMaterialSample Result = CreateMaterialSample(L_World, 0.0, 0.0, 1.0, OutRoughness, PATHTRACER_SCATTER_DIFFUSE);
-
-	const float DiffPdf = NoL / PI;
+	const float ShadowTerminator = ShadowTerminatorTerm(L_World, Payload.WorldNormal, Payload.WorldSmoothNormal);
+	const float DiffNoL = saturate(dot(Result.Direction, Payload.WorldNormal));
+	const float DiffPdf = DiffNoL / PI;
 	float2 GGXResult0 = GGXEvalReflection(L, V, H, Data.Alpha0);
 	const float2 GGXResult1 = GGXEvalReflection(L, V, H, Data.Alpha1);
 	const float Spec0Pdf = GGXResult0.y;
 	const float Spec1Pdf = GGXResult1.y;
 	const float4 ClothResult = Data.Fuzz.Eval(L, V, H, Payload.FuzzColor);
-	const float ClothPdf = ClothResult.w;
+	const float ClothPdf = V.z > 0 ? ClothResult.w : 0;
 
-	const float ShadowTerminator = ShadowTerminatorTerm(L_World, N_World, Payload.WorldSmoothNormal);
 	if (bSampledDiffuse)
 	{
-		const float3 Diffuse = Data.DiffuseWeight * GetPathTracingDiffuseModel(Payload.DiffuseColor, Payload.RoughnessData.x, NoV, NoL, VoH, NoH);
+		float DiffNoL = saturate(dot(Payload.WorldNormal, L_World));
+		float DiffNoV = saturate(dot(Payload.WorldNormal, V_World));
+		float3 H_World = normalize(V_World + L_World);
+		float DiffVoH = saturate(dot(V_World, H_World));
+		float DiffNoH = saturate(dot(Payload.WorldNormal, H_World));
 
-		Result.AddLobeWithMIS(Diffuse * ShadowTerminator, NoL / PI, Data.LobePdf.x);
+		const float3 Diffuse = Data.DiffuseWeight * GetPathTracingDiffuseModel(Payload.DiffuseColor, Payload.RoughnessData.x, DiffNoV, DiffNoL, DiffVoH, DiffNoH);
+		Result.AddLobeWithMIS(Diffuse * ShadowTerminator, DiffPdf, Data.LobePdf.x);
 
 		Result.Pdf += Data.LobePdf.y * Spec0Pdf;
 		Result.Pdf += Data.LobePdf.z * Spec1Pdf;
@@ -167,21 +172,19 @@ FMaterialSample Substrate_SampleMaterial(
 		Result.AddLobeWithMIS((      Payload.RoughnessData.z) * F * GGXResult1.x * Data.Spec1W, Spec1Pdf, Data.LobePdf.z);
 		Result.Pdf += Data.LobePdf.x * DiffPdf;
 		Result.Pdf += Data.LobePdf.w * ClothPdf;
-
-		Result.ScatterType = PATHTRACER_SCATTER_SPECULAR;
 	}
 	else
 	{
 		// Cloth Lobe
 		Result.AddLobeWithMIS(ClothResult.xyz * ShadowTerminator, ClothPdf, Data.LobePdf.w);
-		Result.ScatterType = PATHTRACER_SCATTER_SPECULAR;
-
 		Result.Pdf += Data.LobePdf.x * DiffPdf;
 		Result.Pdf += Data.LobePdf.y * Spec0Pdf;
 		Result.Pdf += Data.LobePdf.z * Spec1Pdf;
 	}
 
-	Result.Weight *= SubstrateLobeWeight(Payload, NoL);
+	Result.ScatterType = bSampledDiffuse ? PATHTRACER_SCATTER_DIFFUSE : PATHTRACER_SCATTER_SPECULAR;
+
+	Result.Weight *= SubstrateLobeWeight(Payload, DiffNoL);
 
 	return Result;
 }
@@ -196,8 +199,6 @@ FMaterialEval Substrate_EvalMaterial(
 	FSubstrateSingleSlabData Data = (FSubstrateSingleSlabData)0;
 	Data.PrepareBSDF(V_World, Payload);
 
-	const float3 N_World = Payload.WorldNormal;
-
 	// move vectors into right shading frame
 	const float3 V = Data.V;
 	const float3 L = mul(Data.Basis, L_World);
@@ -209,10 +210,16 @@ FMaterialEval Substrate_EvalMaterial(
 
 	FMaterialEval Result = NullMaterialEval();
 
-	const float ShadowTerminator = ShadowTerminatorTerm(L_World, N_World, Payload.WorldSmoothNormal);
+	const float ShadowTerminator = ShadowTerminatorTerm(L_World, Payload.WorldNormal, Payload.WorldSmoothNormal);
+	const float DiffNoL = saturate(dot(Payload.WorldNormal, L_World));
+	const float DiffPdf = DiffNoL / PI;
 	{
-		const float3 Diffuse = Data.DiffuseWeight * GetPathTracingDiffuseModel(Payload.DiffuseColor, Payload.RoughnessData.x, NoV, NoL, VoH, NoH);
-		Result.AddLobeWithMIS(DiffuseSpecularScale.x * Diffuse * ShadowTerminator, NoL / PI, Data.LobePdf.x);
+		const float DiffNoV = saturate(dot(Payload.WorldNormal, V_World));
+		const float3 H_World = normalize(V_World + L_World);
+		const float DiffVoH = saturate(dot(V_World, H_World));
+		const float DiffNoH = saturate(dot(Payload.WorldNormal, H_World));
+		const float3 Diffuse = Data.DiffuseWeight * GetPathTracingDiffuseModel(Payload.DiffuseColor, Payload.RoughnessData.x, DiffNoV, DiffNoL, DiffVoH, DiffNoH);
+		Result.AddLobeWithMIS(DiffuseSpecularScale.x * Diffuse * ShadowTerminator, DiffPdf, Data.LobePdf.x);
 	}
 	{
 		float2 GGXResult0 = GGXEvalReflection(L, V, H, Data.Alpha0);
@@ -235,7 +242,7 @@ FMaterialEval Substrate_EvalMaterial(
 		Result.AddLobeWithMIS(DiffuseSpecularScale.y * ClothResult.xyz * ShadowTerminator, ClothPdf, Data.LobePdf.w);
 	}
 
-	Result.Weight *= SubstrateLobeWeight(Payload, NoL);
+	Result.Weight *= SubstrateLobeWeight(Payload, DiffNoL);
 
 	return Result;
 }
@@ -265,6 +272,7 @@ struct FSubstrateClearCoatData
 	void PrepareBSDF(float3 V_World, FPathTracingPayload Payload, bool bUseSSSColor = false)
 	{
 		float3 BaseWorldNormal = Payload.WorldNormal;
+		float3 CoatWorldNormal = Payload.WorldNormal;
 #if CLEAR_COAT_BOTTOM_NORMAL
 		if ((Payload.BottomNormalOct16bits & 0xFFFF) != 0)
 		{
@@ -273,10 +281,13 @@ struct FSubstrateClearCoatData
 		}
 #endif
 
+		BaseWorldNormal = AdjustShadingNormal(BaseWorldNormal, Payload.WorldGeoNormal, V_World);
+		CoatWorldNormal = AdjustShadingNormal(CoatWorldNormal, Payload.WorldGeoNormal, V_World);
+
 		BaseBasis = GetGGXBasis(Payload.RoughnessData.x, Payload.Anisotropy, BaseWorldNormal, Payload.WorldTangent, BaseAlpha);
 		BaseAlpha = GetGGXAlpha(Payload.RoughnessData.x, Payload.Anisotropy);
 
-		CoatBasis = GetGGXBasis(Payload.RoughnessData.y, Payload.Anisotropy, Payload.WorldNormal, Payload.WorldTangent, CoatAlpha);
+		CoatBasis = GetGGXBasis(Payload.RoughnessData.y, Payload.Anisotropy, CoatWorldNormal, Payload.WorldTangent, CoatAlpha);
 		CoatAlpha = GetGGXAlpha(Payload.RoughnessData.y, Payload.Anisotropy);
 
 		BaseV = mul(BaseBasis, V_World);
@@ -322,23 +333,16 @@ FMaterialSample SubstrateClearCoat_SampleMaterial(
 	FSubstrateClearCoatData Data = (FSubstrateClearCoatData)0;
 	Data.PrepareBSDF(V_World, Payload);
 
-	float OutRoughness = 1;
-	float3 OutNormal = Payload.WorldNormal;
-	float3 OutL = 0;
-	float3 OutV = 0;
-	float3 OutH = 0;
-	float3x3 OutBasis;
-
+	FMaterialSample Result = NullMaterialSample();
+	Result.PositionBiasSign = 1;
+	Result.Roughness = 1;
 	const bool bSampledDiffuse = RandSample.x < Data.LobeCdf.x;
 	const bool bSampledSpecular = RandSample.x < Data.LobeCdf.z;
 	if (bSampledDiffuse)
 	{
 		RandSample.x = RescaleRandomNumber(RandSample.x, 0.0, Data.LobeCdf.x);
 		// diffuse lobe
-		OutL = CosineSampleHemisphere(RandSample.xy).xyz;
-		OutV = Data.BaseV;
-		OutH = normalize(OutL + OutV);
-		OutBasis = Data.BaseBasis;
+		Result.Direction = CosineSampleHemisphereWorld(RandSample.xy, Payload.WorldNormal).xyz;
 	}
 	else if (bSampledSpecular)
 	{
@@ -347,74 +351,70 @@ FMaterialSample SubstrateClearCoat_SampleMaterial(
 		if (bUseBaseSpec)
 		{
 			RandSample.x = RescaleRandomNumber(RandSample.x, Data.LobeCdf.x, Data.LobeCdf.y);
-			OutRoughness = Payload.RoughnessData.x;
-
-			const float2 BottomNormalOct = UnpackBottomNormalOct(Payload.BottomNormalOct16bits);
-			OutNormal = SubstrateGetSimpleCoatBottomNormal(BottomNormalOct, Payload.WorldNormal);
-
-			OutBasis = Data.BaseBasis;
-			OutV = Data.BaseV;
+			Result.Roughness = Payload.RoughnessData.x;
 		}
 		else
 		{
 			RandSample.x = RescaleRandomNumber(RandSample.x, Data.LobeCdf.y, Data.LobeCdf.z);
-			OutRoughness = Payload.RoughnessData.y;
-
-			OutBasis = Data.CoatBasis;
-			OutV = Data.CoatV;
+			Result.Roughness = Payload.RoughnessData.y;
 		}
-
-		OutV = bUseBaseSpec ? Data.BaseV : Data.CoatV;
-		OutH = ImportanceSampleVisibleGGX(RandSample.xy, bUseBaseSpec ? Data.BaseAlpha : Data.CoatAlpha, OutV).xyz;
-		OutL = reflect(-OutV, OutH);
-		if (OutL.z <= 0)
+		float3 V = bUseBaseSpec ? Data.BaseV : Data.CoatV;
+		float3 H = ImportanceSampleVisibleGGX(RandSample.xy, bUseBaseSpec ? Data.BaseAlpha : Data.CoatAlpha, V).xyz;
+		float3 L = reflect(-V, H);
+		if (L.z <= 0)
 		{
 			// invalid output direction, exit early
 			return NullMaterialSample();
 		}
+		Result.Direction = mul(L, bUseBaseSpec ? Data.BaseBasis : Data.CoatBasis);
 	}
 	else
 	{
 		// cloth lobe
 		RandSample.x = RescaleRandomNumber(RandSample.x, Data.LobeCdf.z, 1.0);
-		OutV = Data.CoatV;
-		OutL = Data.Fuzz.Sample(RandSample.xy);
-		OutH = normalize(OutL + OutV);
-
-		OutBasis = Data.CoatBasis;
+		float3 L = Data.Fuzz.Sample(RandSample.xy);
+		Result.Direction = mul(L, Data.CoatBasis);
 	}
 
-	// transform to world space
-	float3 L = OutL;
-	float3 V = OutV;
-	float3 H = OutH;
-	const float3 L_World = normalize(mul(L, OutBasis)); 
-	const float3 N_World = OutNormal;
-	const float NoV = saturate(V.z);
-	const float NoL = saturate(L.z);
-	const float VoH = saturate(dot(V, H));
-	const float NoH = saturate(H.z);
+	const float3 L_World = Result.Direction;
+	const float3 BaseL = mul(Data.BaseBasis, L_World);
+	const float3 CoatL = mul(Data.CoatBasis, L_World);
+	const float3 BaseH = normalize(BaseL + Data.BaseV);
+	const float3 CoatH = normalize(CoatL + Data.CoatV);
+
+	const float BaseNoV = saturate(Data.BaseV.z);
+	const float BaseNoL = saturate(BaseL.z);
+	const float BaseVoH = saturate(dot(Data.BaseV, BaseH));
+	const float BaseNoH = saturate(BaseH.z);
 
-	FMaterialSample Result = CreateMaterialSample(L_World, 0.0, 0.0, 1.0, OutRoughness, PATHTRACER_SCATTER_DIFFUSE);
+	const float CoatNoV = saturate(Data.CoatV.z);
+	const float CoatNoL = saturate(CoatL.z);
+	const float CoatVoH = saturate(dot(Data.CoatV, CoatH));
+	const float CoatNoH = saturate(CoatH.z);
 
 	const float Metallic = F0RGBToMetallic(Payload.SpecularColor);
 	const float3 BaseColor = lerp(Payload.DiffuseColor, Payload.SpecularColor, Metallic);
-	const float3 Transmission = lerp(1.0, SimpleClearCoatTransmittance(NoL, NoV, Metallic, BaseColor), Payload.RoughnessData.z);
+	const float3 Transmission = lerp(1.0, SimpleClearCoatTransmittance(CoatNoL, CoatNoV, Metallic, BaseColor), Payload.RoughnessData.z);
 
-	const float DiffPdf = NoL / PI;
-	const float2 BaseGGXResult = GGXEvalReflection(L, V, H, Data.BaseAlpha);
-	const float2 CoatGGXResult = GGXEvalReflection(L, V, H, Data.CoatAlpha);
+	const float DiffNoL = saturate(dot(Payload.WorldNormal, L_World));
+	const float DiffPdf = DiffNoL / PI;
+	const float2 BaseGGXResult = GGXEvalReflection(BaseL, Data.BaseV, BaseH, Data.BaseAlpha);
+	const float2 CoatGGXResult = GGXEvalReflection(CoatL, Data.CoatV, CoatH, Data.CoatAlpha);
 	const float BaseSpecPdf = BaseGGXResult.y;
 	const float CoatSpecPdf = CoatGGXResult.y;
-	const float4 ClothResult = Data.Fuzz.Eval(L, V, H, Payload.FuzzColor);
+	const float4 ClothResult = Data.Fuzz.Eval(CoatL, Data.CoatV, CoatH, Payload.FuzzColor);
 	const float ClothPdf = ClothResult.w;
 
-	const float ShadowTerminator = ShadowTerminatorTerm(L_World, N_World, Payload.WorldSmoothNormal);
+	const float ShadowTerminator = ShadowTerminatorTerm(L_World, Payload.WorldNormal, Payload.WorldSmoothNormal);
 	if (bSampledDiffuse)
 	{
-		const float3 Diffuse = Transmission * Data.DiffuseWeight * GetPathTracingDiffuseModel(Payload.DiffuseColor, Payload.RoughnessData.x, NoV, NoL, VoH, NoH);
+		const float DiffNoV = saturate(dot(Payload.WorldNormal, V_World));
+		const float3 H_World = normalize(V_World + L_World);
+		const float DiffVoH = saturate(dot(V_World, H_World));
+		const float DiffNoH = saturate(dot(Payload.WorldNormal, H_World));
+		const float3 Diffuse = Data.DiffuseWeight * GetPathTracingDiffuseModel(Payload.DiffuseColor, Payload.RoughnessData.x, DiffNoV, DiffNoL, DiffVoH, DiffNoH);
 
-		Result.AddLobeWithMIS(Diffuse * ShadowTerminator, NoL / PI, Data.LobePdf.x);
+		Result.AddLobeWithMIS(Diffuse * ShadowTerminator, DiffPdf, Data.LobePdf.x);
 
 		Result.Pdf += Data.LobePdf.y * BaseSpecPdf;
 		Result.Pdf += Data.LobePdf.z * CoatSpecPdf;
@@ -422,34 +422,32 @@ FMaterialSample SubstrateClearCoat_SampleMaterial(
 	}
 	else if (bSampledSpecular)
 	{
-		const float3 FTint = SubstrateSpecularTint(Payload, NoV, NoL, VoH, NoH);
+		const float3 FTint = SubstrateSpecularTint(Payload, BaseNoV, BaseNoL, BaseVoH, BaseNoH);
 #if SUBSTRATE_FRESNEL_F82
 		// SUBSTRATE_TODO: Need to tune the logic for specular micro-occlusion ...
-		const float3 FBase = Data.Fuzz.Attenuation * F_AdobeF82(Data.F0, Data.F90, VoH) * FTint;
+		const float3 FBase = Data.Fuzz.Attenuation * F_AdobeF82(Data.F0, Data.F90, BaseVoH) * FTint;
 #else
-		const float3 FBase = Data.Fuzz.Attenuation * F_Schlick(Data.F0, Data.F90, VoH) * FTint;
+		const float3 FBase = Data.Fuzz.Attenuation * F_Schlick(Data.F0, Data.F90, BaseVoH) * FTint;
 #endif
-		const float3 FCoat = Payload.RoughnessData.z * ClearCoatFresnel(VoH) * FTint;
+		const float3 FCoat = Payload.RoughnessData.z * ClearCoatFresnel(CoatVoH) * FTint;
 
-		Result.AddLobeWithMIS(Data.CoatAttenuation * FBase * BaseGGXResult.x * Data.BaseSpec.W * Transmission, BaseSpecPdf, Data.LobePdf.y);
+		Result.AddLobeWithMIS(Data.CoatAttenuation  * FBase * BaseGGXResult.x * Data.BaseSpec.W * Transmission, BaseSpecPdf, Data.LobePdf.y);
 		Result.AddLobeWithMIS(Data.Fuzz.Attenuation * FCoat * CoatGGXResult.x * Data.CoatSpec.W               , CoatSpecPdf, Data.LobePdf.z);
 		Result.Pdf += Data.LobePdf.x * DiffPdf;
 		Result.Pdf += Data.LobePdf.w * ClothPdf;
-
-		Result.ScatterType = PATHTRACER_SCATTER_SPECULAR;
 	}
 	else
 	{
 		// Cloth Lobe
 		Result.AddLobeWithMIS(ClothResult.xyz * ShadowTerminator, ClothPdf, Data.LobePdf.w);
-		Result.ScatterType = PATHTRACER_SCATTER_SPECULAR;
 
 		Result.Pdf += Data.LobePdf.x * DiffPdf;
 		Result.Pdf += Data.LobePdf.y * BaseSpecPdf;
 		Result.Pdf += Data.LobePdf.z * CoatSpecPdf;
 	}
 
-	Result.Weight *= SubstrateLobeWeight(Payload, NoL);
+	Result.Weight *= SubstrateLobeWeight(Payload, DiffNoL);
+	Result.ScatterType = bSampledDiffuse ? PATHTRACER_SCATTER_DIFFUSE : PATHTRACER_SCATTER_SPECULAR;
 
 	return Result;
 }
@@ -488,11 +486,18 @@ FMaterialEval SubstrateClearCoat_EvalMaterial(
 	const float3 BaseColor = lerp(Payload.DiffuseColor, Payload.SpecularColor, Metallic);
 	const float3 Transmission = lerp(1.0, SimpleClearCoatTransmittance(CoatNoL, CoatNoV, Metallic, BaseColor), Payload.RoughnessData.z);
 
-	const float BasehadowTerminator = ShadowTerminatorTerm(L_World,  Data.BaseBasis[2], Payload.WorldSmoothNormal);
-	const float CoatShadowTerminator = ShadowTerminatorTerm(L_World, Data.CoatBasis[2], Payload.WorldSmoothNormal);
+	const float BasehadowTerminator  = ShadowTerminatorTerm(L_World, Payload.WorldNormal, Payload.WorldSmoothNormal);
+	const float CoatShadowTerminator = ShadowTerminatorTerm(L_World, Data.CoatBasis[2] , Payload.WorldSmoothNormal);
+
+	const float DiffNoL = saturate(dot(Payload.WorldNormal, L_World));
 	{
-		const float3 Diffuse = Transmission * Data.DiffuseWeight * GetPathTracingDiffuseModel(Payload.DiffuseColor, Payload.RoughnessData.x, BaseNoV, BaseNoL, BaseVoH, BaseNoH);
-		Result.AddLobeWithMIS(DiffuseSpecularScale.x * Diffuse * BasehadowTerminator, BaseNoL / PI, Data.LobePdf.x);
+		const float DiffNoV = saturate(dot(Payload.WorldNormal, V_World));
+		const float3 H_World = normalize(V_World + L_World);
+		const float DiffVoH = saturate(dot(V_World, H_World));
+		const float DiffNoH = saturate(dot(Payload.WorldNormal, H_World));
+
+		const float3 Diffuse = Transmission * Data.DiffuseWeight * GetPathTracingDiffuseModel(Payload.DiffuseColor, Payload.RoughnessData.x, DiffNoV, DiffNoL, DiffVoH, DiffNoH);
+		Result.AddLobeWithMIS(DiffuseSpecularScale.x * Diffuse * BasehadowTerminator, DiffNoL / PI, Data.LobePdf.x);
 	}
 	{
 		const float2 BaseGGXResult = GGXEvalReflection(BaseL, BaseV, BaseH, Data.BaseAlpha);
@@ -500,16 +505,16 @@ FMaterialEval SubstrateClearCoat_EvalMaterial(
 		const float BaseSpecPdf = BaseGGXResult.y;
 		const float CoatSpecPdf = CoatGGXResult.y;
 
-		const float3 FTint = SubstrateSpecularTint(Payload, CoatNoV, CoatNoL, CoatVoH, CoatNoH);
+		const float3 FTint = SubstrateSpecularTint(Payload, BaseNoV, BaseNoL, BaseVoH, BaseNoH);
 #if SUBSTRATE_FRESNEL_F82
 		// SUBSTRATE_TODO: Need to tune the logic for specular micro-occlusion ...
-		const float3 FBase = Data.Fuzz.Attenuation * F_AdobeF82(Data.F0, Data.F90, CoatVoH) * FTint;
+		const float3 FBase = Data.Fuzz.Attenuation * F_AdobeF82(Data.F0, Data.F90, BaseVoH) * FTint;
 #else
-		const float3 FBase = Data.Fuzz.Attenuation * F_Schlick(Data.F0, Data.F90, CoatVoH) * FTint;
+		const float3 FBase = Data.Fuzz.Attenuation * F_Schlick(Data.F0, Data.F90, BaseVoH) * FTint;
 #endif
 		const float3 FCoat = Payload.RoughnessData.z * ClearCoatFresnel(CoatVoH) * FTint;
 
-		Result.AddLobeWithMIS(Data.CoatAttenuation * FBase * BaseGGXResult.x * Data.BaseSpec.W * Transmission, BaseSpecPdf, Data.LobePdf.y);
+		Result.AddLobeWithMIS(Data.CoatAttenuation  * FBase * BaseGGXResult.x * Data.BaseSpec.W * Transmission, BaseSpecPdf, Data.LobePdf.y);
 		Result.AddLobeWithMIS(Data.Fuzz.Attenuation * FCoat * CoatGGXResult.x * Data.CoatSpec.W               , CoatSpecPdf, Data.LobePdf.z);
 	}
 	{
@@ -519,7 +524,7 @@ FMaterialEval SubstrateClearCoat_EvalMaterial(
 		Result.AddLobeWithMIS(DiffuseSpecularScale.y * ClothResult.xyz * CoatShadowTerminator, ClothPdf, Data.LobePdf.w);
 	}
 
-	Result.Weight *= SubstrateLobeWeight(Payload, CoatNoL);
+	Result.Weight *= SubstrateLobeWeight(Payload, DiffNoL);
 
 	return Result;
 }
diff --git a/Engine/Shaders/Private/PathTracing/Material/PathTracingSubstrateFoliage.ush b/Engine/Shaders/Private/PathTracing/Material/PathTracingSubstrateFoliage.ush
--- a/Engine/Shaders/Private/PathTracing/Material/PathTracingSubstrateFoliage.ush
+++ b/Engine/Shaders/Private/PathTracing/Material/PathTracingSubstrateFoliage.ush
@@ -32,7 +32,8 @@ struct FSubstrateFoliageSlabData
 
 	void PrepareBSDF(float3 V_World, FPathTracingPayload Payload)
 	{
-		Basis = GetGGXBasis(Payload.RoughnessData.x, Payload.Anisotropy, Payload.WorldNormal, Payload.WorldTangent, Alpha0);
+		float3 NSpec = AdjustShadingNormal(Payload.WorldNormal, Payload.WorldGeoNormal, V_World);
+		Basis = GetGGXBasis(Payload.RoughnessData.x, Payload.Anisotropy, NSpec, Payload.WorldTangent, Alpha0);
 		Alpha1 = GetGGXAlpha(Payload.RoughnessData.y, Payload.Anisotropy);
 
 		V = mul(Basis, V_World);
@@ -82,40 +83,30 @@ FMaterialSample SubstrateFoliage_SampleMaterial(
 	FSubstrateFoliageSlabData Data = (FSubstrateFoliageSlabData)0;
 	Data.PrepareBSDF(V_World, Payload);
 
-	float3 L = 0, H = 0, V = Data.V;
-	float OutRoughness = 1;
+	FMaterialSample Result = NullMaterialSample();
+	Result.Roughness = 1;
+	Result.PositionBiasSign = 1;
+	float3 V = Data.V, L = 0, H = 0;
 
-	const bool bSampledTransmit = RandSample.x < Data.LobeCdf.x * Data.TransmissionProb;
+	const float TransmitCdf = Data.LobeCdf.x * Data.TransmissionProb;
+	const bool bSampledTransmit = RandSample.x < TransmitCdf;
 	const bool bSampledDiffuse  = RandSample.x < Data.LobeCdf.x;
 	const bool bSampledSpecular = RandSample.x < Data.LobeCdf.z;
-	if (bSampledTransmit)
-	{
-		RandSample.x = RescaleRandomNumber(RandSample.x, 0.0, Data.LobeCdf.x * Data.TransmissionProb);
-		// diffuse lobe
-		L = CosineSampleHemisphere(RandSample.xy).xyz;
-		L.z = -L.z;
-	}
-	else if (bSampledDiffuse)
+	if (bSampledDiffuse)
 	{
-		RandSample.x = RescaleRandomNumber(RandSample.x, Data.LobeCdf.x * Data.TransmissionProb, Data.LobeCdf.x);
+		RandSample.x = RescaleRandomNumber(RandSample.x, bSampledTransmit ? 0.0 : TransmitCdf, bSampledTransmit ? TransmitCdf : Data.LobeCdf.x);
 		// diffuse lobe
-		L = CosineSampleHemisphere(RandSample.xy).xyz;
+		Result.Direction = CosineSampleHemisphereWorld(RandSample.xy, Payload.WorldNormal * (bSampledTransmit ? -1 : 1)).xyz;
+		L = mul(Result.Direction, Data.Basis);
 		H = normalize(L + V);
 	}
 	else if (bSampledSpecular)
 	{
 		// specular lobes
 		const bool bUseSpec0 = RandSample.x < Data.LobeCdf.y;
-		if (bUseSpec0)
-		{
-			RandSample.x = RescaleRandomNumber(RandSample.x, Data.LobeCdf.x, Data.LobeCdf.y);
-			OutRoughness = Payload.RoughnessData.x;
-		}
-		else
-		{
-			RandSample.x = RescaleRandomNumber(RandSample.x, Data.LobeCdf.y, Data.LobeCdf.z);
-			OutRoughness = Payload.RoughnessData.y;
-		}
+
+		Result.Roughness = bUseSpec0 ? Payload.RoughnessData.x : Payload.RoughnessData.y;
+		RandSample.x = RescaleRandomNumber(RandSample.x, bUseSpec0 ? Data.LobeCdf.x : Data.LobeCdf.y, bUseSpec0 ? Data.LobeCdf.y : Data.LobeCdf.z);
 
 		H = ImportanceSampleVisibleGGX(RandSample.xy, bUseSpec0 ? Data.Alpha0 : Data.Alpha1, V).xyz;
 		L = reflect(-V, H);
@@ -124,6 +115,7 @@ FMaterialSample SubstrateFoliage_SampleMaterial(
 			// invalid output direction, exit early
 			return NullMaterialSample();
 		}
+		Result.Direction = mul(Data.Basis, L);
 	}
 	else
 	{
@@ -131,36 +123,39 @@ FMaterialSample SubstrateFoliage_SampleMaterial(
 		RandSample.x = RescaleRandomNumber(RandSample.x, Data.LobeCdf.z, 1.0);
 		L = Data.Fuzz.Sample(RandSample.xy);
 		H = normalize(L + V);
+		Result.Direction = mul(Data.Basis, L);
 	}
 
 	// transform to world space
 	const float3 L_World = normalize(mul(L, Data.Basis));
-	const float3 N_World = Payload.WorldNormal;
 	const float NoV = saturate(V.z);
 	const float NoL = abs(L.z);
 	const float VoH = saturate(dot(V, H));
 	const float NoH = saturate(H.z);
 
-	FMaterialSample Result = CreateMaterialSample(L_World, 0.0, 0.0, 1.0, OutRoughness, PATHTRACER_SCATTER_DIFFUSE);
-
-	const float DiffPdf = NoL / PI;
-	float2 GGXResult0 = GGXEvalReflection(L, V, H, Data.Alpha0);
+	const float DiffNoL = abs(dot(Payload.WorldNormal, L_World));
+	const float DiffPdf = DiffNoL / PI;
+	const float2 GGXResult0 = GGXEvalReflection(L, V, H, Data.Alpha0);
 	const float2 GGXResult1 = GGXEvalReflection(L, V, H, Data.Alpha1);
 	const float Spec0Pdf = GGXResult0.y;
 	const float Spec1Pdf = GGXResult1.y;
 	const float4 ClothResult = Data.Fuzz.Eval(L, V, H, Payload.FuzzColor);
 	const float ClothPdf = ClothResult.w;
 
-	const float ShadowTerminator = ShadowTerminatorTerm(L_World, N_World, Payload.WorldSmoothNormal);
+	const float ShadowTerminator = ShadowTerminatorTerm(L_World, Payload.WorldNormal, Payload.WorldSmoothNormal);
 	if (bSampledTransmit)
 	{
-		Result.AddLobeWithMIS(Data.DiffuseWeight * Payload.GetFoliageTransmissionColor(), abs(L.z) / PI, Data.LobePdf.x * Data.TransmissionProb);
+		Result.AddLobeWithMIS(Data.DiffuseWeight * Payload.GetFoliageTransmissionColor(), DiffPdf, Data.LobePdf.x * Data.TransmissionProb);
 	}
 	else if (bSampledDiffuse)
 	{
-		const float3 Diffuse = Data.DiffuseWeight * GetPathTracingDiffuseModel(Payload.DiffuseColor, Payload.RoughnessData.x, NoV, NoL, VoH, NoH);
+		const float DiffNoV = saturate(dot(Payload.WorldNormal, V_World));
+		const float3 H_World = normalize(V_World + L_World);
+		const float DiffVoH = saturate(dot(V_World, H_World));
+		const float DiffNoH = saturate(dot(Payload.WorldNormal, H_World));
+		const float3 Diffuse = Data.DiffuseWeight * GetPathTracingDiffuseModel(Payload.DiffuseColor, Payload.RoughnessData.x, DiffNoV, DiffNoL, DiffVoH, DiffNoH);
 
-		Result.AddLobeWithMIS(Diffuse * ShadowTerminator, NoL / PI, Data.LobePdf.x - Data.TransmissionProb * Data.LobePdf.x);
+		Result.AddLobeWithMIS(Diffuse * ShadowTerminator, DiffPdf, Data.LobePdf.x - Data.TransmissionProb * Data.LobePdf.x);
 
 		Result.Pdf += Data.LobePdf.y * Spec0Pdf;
 		Result.Pdf += Data.LobePdf.z * Spec1Pdf;
@@ -180,22 +175,18 @@ FMaterialSample SubstrateFoliage_SampleMaterial(
 		Result.AddLobeWithMIS((      Payload.RoughnessData.z) * F * GGXResult1.x * Data.Spec1W, Spec1Pdf, Data.LobePdf.z);
 		Result.Pdf += Data.LobePdf.x * DiffPdf;
 		Result.Pdf += Data.LobePdf.w * ClothPdf;
-
-		Result.ScatterType = PATHTRACER_SCATTER_SPECULAR;
 	}
 	else
 	{
 		// Cloth Lobe
 		Result.AddLobeWithMIS(ClothResult.xyz * ShadowTerminator, ClothPdf, Data.LobePdf.w);
-		Result.ScatterType = PATHTRACER_SCATTER_SPECULAR;
-
 		Result.Pdf += Data.LobePdf.x * DiffPdf;
 		Result.Pdf += Data.LobePdf.y * Spec0Pdf;
 		Result.Pdf += Data.LobePdf.z * Spec1Pdf;
 	}
 
-	Result.Weight *= SubstrateLobeWeight(Payload, NoL);
-
+	Result.Weight *= SubstrateLobeWeight(Payload, DiffNoL);
+	Result.ScatterType = bSampledDiffuse ? PATHTRACER_SCATTER_DIFFUSE : PATHTRACER_SCATTER_SPECULAR;
 	return Result;
 }
 
@@ -222,19 +213,25 @@ FMaterialEval SubstrateFoliage_EvalMaterial(
 
 	FMaterialEval Result = NullMaterialEval();
 
+	const float DiffNoL = dot(Payload.WorldNormal, L_World);
+	const float DiffPdf = DiffNoL / PI;
 	const float ShadowTerminator = ShadowTerminatorTerm(L_World, N_World, Payload.WorldSmoothNormal);
-	if (L.z < 0)
+	if (DiffNoL < 0)
 	{
 		const float3 Diffuse = Data.DiffuseWeight * Payload.GetFoliageTransmissionColor();
-		Result.AddLobeWithMIS(DiffuseSpecularScale.x * Diffuse, NoL / PI, Data.LobePdf.x * Data.TransmissionProb);
+		Result.AddLobeWithMIS(DiffuseSpecularScale.x * Diffuse, -DiffPdf, Data.LobePdf.x * Data.TransmissionProb);
 	}
-	else if (L.z > 0)
+	else
+	{
+		const float DiffNoV = saturate(dot(Payload.WorldNormal, V_World));
+		const float3 H_World = normalize(V_World + L_World);
+		const float DiffVoH = saturate(dot(V_World, H_World));
+		const float DiffNoH = saturate(dot(Payload.WorldNormal, H_World));
+		const float3 Diffuse = Data.DiffuseWeight * GetPathTracingDiffuseModel(Payload.DiffuseColor, Payload.RoughnessData.x, DiffNoV, DiffNoL, DiffVoH, DiffNoH);
+		Result.AddLobeWithMIS(DiffuseSpecularScale.x * Diffuse * ShadowTerminator, DiffPdf, Data.LobePdf.x - Data.LobePdf.x * Data.TransmissionProb);
+	}
+	if (L.z > 0)
 	{
-		{
-			const float3 Diffuse = Data.DiffuseWeight * GetPathTracingDiffuseModel(Payload.DiffuseColor, Payload.RoughnessData.x, NoV, NoL, VoH, NoH);
-			Result.AddLobeWithMIS(DiffuseSpecularScale.x * Diffuse * ShadowTerminator, NoL / PI, Data.LobePdf.x - Data.LobePdf.x * Data.TransmissionProb);
-		}
-
 		{
 			float2 GGXResult0 = GGXEvalReflection(L, V, H, Data.Alpha0);
 			const float2 GGXResult1 = GGXEvalReflection(L, V, H, Data.Alpha1);
@@ -256,7 +253,7 @@ FMaterialEval SubstrateFoliage_EvalMaterial(
 		}
 	}
 
-	Result.Weight *= SubstrateLobeWeight(Payload, NoL);
+	Result.Weight *= SubstrateLobeWeight(Payload, DiffNoL);
 
 	return Result;
 }
\ No newline at end of file
diff --git a/Engine/Shaders/Private/PathTracing/Material/PathTracingSubstrateSolidGlass.ush b/Engine/Shaders/Private/PathTracing/Material/PathTracingSubstrateSolidGlass.ush
--- a/Engine/Shaders/Private/PathTracing/Material/PathTracingSubstrateSolidGlass.ush
+++ b/Engine/Shaders/Private/PathTracing/Material/PathTracingSubstrateSolidGlass.ush
@@ -33,7 +33,8 @@ struct FSubstrateSolidGlassSlabData
 
 	void PrepareBSDF(float3 V_World, FPathTracingPayload Payload)
 	{
-		Basis = GetGGXBasis(Payload.RoughnessData.x, Payload.Anisotropy, Payload.WorldNormal, Payload.WorldTangent, Alpha0);
+		float3 N = AdjustShadingNormal(Payload.WorldNormal, Payload.WorldGeoNormal, V_World);
+		Basis = GetGGXBasis(Payload.RoughnessData.x, Payload.Anisotropy, N, Payload.WorldTangent, Alpha0);
 		Alpha1 = GetGGXAlpha(Payload.RoughnessData.y, Payload.Anisotropy);
 
 		V = mul(Basis, V_World);
@@ -129,7 +130,7 @@ FMaterialSample SubstrateSolidGlass_SampleMaterial(
 			// and account for probability of picking the fuzz layer
 			const float ClothPdf = Data.Fuzz.Eval(L, V, H, Payload.FuzzColor).w;
 			Result.Pdf += Data.LobePdf.z * ClothPdf;
-			Result.Weight *= Data.Fuzz.Attenuation * SubstrateLobeWeight(Payload, saturate(L.z)) * FTint;
+			Result.Weight *= Data.Fuzz.Attenuation * SubstrateLobeWeight(Payload, NoL) * FTint;
 		}
 
 		return Result;
diff --git a/Engine/Shaders/Private/PathTracing/Material/PathTracingSubstrateThinGlass.ush b/Engine/Shaders/Private/PathTracing/Material/PathTracingSubstrateThinGlass.ush
--- a/Engine/Shaders/Private/PathTracing/Material/PathTracingSubstrateThinGlass.ush
+++ b/Engine/Shaders/Private/PathTracing/Material/PathTracingSubstrateThinGlass.ush
@@ -41,7 +41,9 @@ struct FSubstrateThinGlassSlabData
 		RoughnessT.x = ComputeThinTransmittedRoughness(RoughnessR.x, Payload.Ior);
 		RoughnessT.y = ComputeThinTransmittedRoughness(RoughnessR.y, Payload.Ior);
 
-		Basis   = GetGGXBasis(RoughnessR.x, Payload.Anisotropy, Payload.WorldNormal, Payload.WorldTangent, AlphaR0);
+		float3 N = AdjustShadingNormal(Payload.WorldNormal, Payload.WorldGeoNormal, V_World);
+
+		Basis   = GetGGXBasis(RoughnessR.x, Payload.Anisotropy, N, Payload.WorldTangent, AlphaR0);
 		AlphaR1 = GetGGXAlpha(RoughnessR.y, Payload.Anisotropy);
 
 		AlphaT0 = GetGGXAlpha(RoughnessT.x, Payload.Anisotropy);
diff --git a/Engine/Shaders/Private/PathTracing/Material/PathTracingThinGlass.ush b/Engine/Shaders/Private/PathTracing/Material/PathTracingThinGlass.ush
--- a/Engine/Shaders/Private/PathTracing/Material/PathTracingThinGlass.ush
+++ b/Engine/Shaders/Private/PathTracing/Material/PathTracingThinGlass.ush
@@ -35,7 +35,8 @@ FRoughThinGlassData PrepareRoughThinGlassData(FPathTracingPayload Payload, float
 	const float RoughnessR = Payload.Roughness;
 	const float RoughnessT = ComputeThinTransmittedRoughness(Payload.Roughness, Payload.Ior);
 
-	Data.Basis = GetGGXBasis(RoughnessR, Payload.Anisotropy, Payload.WorldNormal, Payload.WorldTangent, Data.AlphaR);
+	float3 N = AdjustShadingNormal(Payload.WorldNormal, Payload.WorldGeoNormal, V_World);
+	Data.Basis = GetGGXBasis(RoughnessR, Payload.Anisotropy, N, Payload.WorldTangent, Data.AlphaR);
 	Data.AlphaT = GetGGXAlpha(RoughnessT, Payload.Anisotropy);
 
 	Data.V = mul(Data.Basis, V_World);
diff --git a/Engine/Shaders/Private/PathTracing/PathTracingCommon.ush b/Engine/Shaders/Private/PathTracing/PathTracingCommon.ush
--- a/Engine/Shaders/Private/PathTracing/PathTracingCommon.ush
+++ b/Engine/Shaders/Private/PathTracing/PathTracingCommon.ush
@@ -39,13 +39,14 @@
 #define PATH_TRACING_PAYLOAD_TANGENT_ANGLE_BITS 16 						// Could be dropped to 12 or even 10 if we need more room in the future
 
 // shading model extensions beyond the default count
-#define SHADINGMODELID_MEDIUM  			(SHADINGMODELID_NUM + 1)
-#define SHADINGMODELID_SOLID_GLASS 		(SHADINGMODELID_NUM + 2)
+// TODO: instead of re-using SHADINGMODELID_* values, make a distint enum for the path tracer
+#define SHADINGMODELID_MEDIUM  			(SHADINGMODELID_SUBSTRATE_TOON) // alias with unused value in the path tracer
+#define SHADINGMODELID_LAMBERT          (SHADINGMODELID_NUM    )        // pure lambert response (used for lighting after SSS random walk)
+#define SHADINGMODELID_SOLID_GLASS 		(SHADINGMODELID_NUM + 1) 		// use a sparse bit
 
-#define PATH_TRACING_SHADINGMODEL_ID_MASK 0x3F
-
-HLSL_STATIC_ASSERT(SHADINGMODELID_MEDIUM      <= PATH_TRACING_SHADINGMODEL_ID_MASK, "");
-HLSL_STATIC_ASSERT(SHADINGMODELID_SOLID_GLASS <= PATH_TRACING_SHADINGMODEL_ID_MASK, "");
+HLSL_STATIC_ASSERT(SHADINGMODELID_MEDIUM      <= SHADINGMODELID_MASK, "");
+HLSL_STATIC_ASSERT(SHADINGMODELID_LAMBERT     <= SHADINGMODELID_MASK, "");
+HLSL_STATIC_ASSERT(SHADINGMODELID_SOLID_GLASS <= SHADINGMODELID_MASK, "");
 
 #define PATH_TRACING_ABSORPTION_SCALE  0.01        // inverse distance at which we hit the transmission color (Beer's law)
 
@@ -207,12 +208,12 @@ void PayloadDecodeFuzzAndBottomNormal(uint rgba, inout float FuzzAmount, inout f
 	BottomNormalOct16bits = rgba & 0xFFFF;
 }
 
-float3 AdjustShadingNormal(float3 ShadingNormal, float3 GeoNormal, float3 RayDirection)
+float3 AdjustShadingNormal(float3 ShadingNormal, float3 GeoNormal, float3 V_World)
 {
 	// Clip shading normal in a view dependent way such that the reflection stays above the geometric normal
 	// This introduces a bit of view-dependency to the shading normal but fixes dark artifacts around grazing angles
 
-	float3 D = RayDirection;
+	float3 D = -V_World;
 	float3 R = reflect(D, ShadingNormal);
 
 	// https://iquilezles.org/www/articles/dontflip/dontflip.htm
@@ -251,8 +252,8 @@ struct FPackedPathTracingPayload : FMinimalPayload
 	// 64 bytes total
 #endif
 
-	float UnpackBSDFOpacity() { return Unpack10F(PackedData[0] & 0x3FF); }
-	uint GetShadingModelID() { return (PackedData[0] >> 10) & PATH_TRACING_SHADINGMODEL_ID_MASK; }
+	float UnpackBSDFOpacity() { return f16tof32(PackedData[0] & 0xFFFF); }
+	uint GetShadingModelID() { return (PackedData[0] >> 16) & SHADINGMODELID_MASK; }
 	uint GetFlags() { return (PackedData[0] >> 20) & 0x1FF; }
 	void SetFlag(uint Flag) { PackedData[0] |= Flag << 20; }
 	void RemoveFlag(uint Flag) { PackedData[0] &= ~(Flag << 20); }
@@ -777,7 +778,7 @@ struct FPathTracingPayload : FMinimalPayload
 	}
 
 
-	// SHADINGMODELID_NUM (Lambert diffuse) and SHADINGMODELID_MEDIUM
+	// SHADINGMODELID_LAMBERT and SHADINGMODELID_MEDIUM
 	float3 GetBaseColor() { return DiffuseColor; }
 	void SetBaseColor(float3 BaseColor) { DiffuseColor = BaseColor; }
 
@@ -901,7 +902,7 @@ struct FPathTracingPayload : FMinimalPayload
 	void SetGlassPhase(float G) { CustomData0.w = G; }
 
 
-	// SHADINGMODELID_NUM (Lambert diffuse) and SHADINGMODELID_MEDIUM
+	// SHADINGMODELID_LAMBERT and SHADINGMODELID_MEDIUM
 	float3 GetBaseColor() { return BaseColor; }
 	void SetBaseColor(float3 C)
 	{
@@ -1007,8 +1008,8 @@ FPackedPathTracingPayload PackPathTracingPayload(FPathTracingPayload Input)
 	FPackedPathTracingPayload Output = (FPackedPathTracingPayload)0;
 	Output.HitT = Input.HitT;
 
-	Output.PackedData[0]  = Pack10F(Input.BSDFOpacity);                                       // 10 bits
-	Output.PackedData[0] |= (Input.ShadingModelID & PATH_TRACING_SHADINGMODEL_ID_MASK) << 10; // 10 bits - That is a lot, some bits from there could be reused for other things.
+	Output.PackedData[0]  = f32tof16(Input.BSDFOpacity);                                      // 16 bits
+	Output.PackedData[0] |= (Input.ShadingModelID & SHADINGMODELID_MASK) << 16;               // 4  bits
 	Output.PackedData[0] |= (Input.Flags & 0x1FF) << 20;                                      // 9  bits
 	Output.PackedData[0] |= (Input.PrimitiveLightingChannelMask & 0x7) << 29;                 // 3  bits
 	                                                                                          // 32 bits total
diff --git a/Engine/Shaders/Private/PathTracing/PathTracingCore.ush b/Engine/Shaders/Private/PathTracing/PathTracingCore.ush
--- a/Engine/Shaders/Private/PathTracing/PathTracingCore.ush
+++ b/Engine/Shaders/Private/PathTracing/PathTracingCore.ush
@@ -393,9 +393,6 @@ void ApplyDecal(FDecalShaderPayload DecalPayload, inout FPackedPathTracingPayloa
 		// We normalize the normal to get smoother visual result (it helps to avoid having D_GGX explodes toward infinity, and matches ApplyDBufferData(...))
 		float3 WorldNormal = PackedPayload.UnpackWorldNormal();
 		WorldNormal = normalize(WorldNormal * DecalWorldNormal.a + DecalWorldNormal.rgb);
-		// re-apply this (normally happens as part of shading, but since decal changed the normal, apply it again)
-		WorldNormal = AdjustShadingNormal(WorldNormal, PackedPayload.UnpackWorldGeoNormal(), RayDirection);
-
 		PackedPayload.PackWorldNormal(WorldNormal);
 	}
 
@@ -1639,7 +1636,7 @@ bool ProcessSubsurfaceRandomWalk(inout FPathTracingPayload Payload, inout float3
 			Payload.WorldSmoothNormal = SignFlip * ProbeResult.WorldSmoothNormal;
 			Payload.WorldGeoNormal    = SignFlip * ProbeResult.WorldGeoNormal;
 			
-			Payload.ShadingModelID = SHADINGMODELID_NUM; // invalid value so that we get diffuse shading
+			Payload.ShadingModelID = SHADINGMODELID_LAMBERT; // diffuse shading
 			Payload.BSDFOpacity = 1;
 			Payload.SetBaseColor(1.0);
 			Payload.TransparencyColor = 0;
diff --git a/Engine/Shaders/Private/PathTracing/PathTracingDebug.usf b/Engine/Shaders/Private/PathTracing/PathTracingDebug.usf
--- a/Engine/Shaders/Private/PathTracing/PathTracingDebug.usf
+++ b/Engine/Shaders/Private/PathTracing/PathTracingDebug.usf
@@ -1101,6 +1101,20 @@ FPathTracingResult PathTracingDebugCore(FRayDesc Ray, uint2 LaunchIndex)
 				case PATH_TRACER_DEBUG_VIZ_EYE_IRIS_MASK:      Color = (HitPayload.ShadingModelID == SHADINGMODELID_EYE) ?             HitPayload.GetEyeIrisMask()      : 0.0; break;
 				case PATH_TRACER_DEBUG_VIZ_EYE_IRIS_NORMAL:    Color = (HitPayload.ShadingModelID == SHADINGMODELID_EYE) ? 0.5 + 0.5 * HitPayload.GetEyeIrisNormal()    : 0.0; break;
 				case PATH_TRACER_DEBUG_VIZ_EYE_CAUSTIC_NORMAL: Color = (HitPayload.ShadingModelID == SHADINGMODELID_EYE) ? 0.5 + 0.5 * HitPayload.GetEyeCausticNormal() : 0.0; break;
+				case PATH_TRACER_DEBUG_VIZ_BACKFACING_NORMAL:
+				{
+					Color = FakeShading;
+					float K1 = dot(WorldNormal, Ray.Direction);
+					float K2 = dot(HitPayload.WorldGeoNormal, reflect(Ray.Direction, WorldNormal));
+					if (K1 > 1e-3)
+					{
+						Color = float3(0, 1, 0);
+					}
+					else if (K2 < -1e-3)
+					{
+						Color = float3(1, 0, 0);
+					}
+				}
 				default:
 				{
 					break;
diff --git a/Engine/Shaders/Private/PathTracing/PathTracingMaterialHitShader.usf b/Engine/Shaders/Private/PathTracing/PathTracingMaterialHitShader.usf
--- a/Engine/Shaders/Private/PathTracing/PathTracingMaterialHitShader.usf
+++ b/Engine/Shaders/Private/PathTracing/PathTracingMaterialHitShader.usf
@@ -352,11 +352,6 @@ RAY_TRACING_ENTRY_CLOSEST_HIT(PathTracingMaterialCHS,
 #endif
 	Payload.WorldGeoNormal = GeoNormalSign * GeoNormal;
 	Payload.WorldSmoothNormal = MaterialParameters.TwoSidedSign * TransformTangentNormalToWorld(MaterialParameters.TangentToWorld, float3(0, 0, 1));
-#if PATH_TRACER_USE_VOLUMETRIC_SHADING_MODEL
-	// no need to adjust normal on volumetric shading models
-#else
-	Payload.WorldSmoothNormal = AdjustShadingNormal(Payload.WorldSmoothNormal, Payload.WorldGeoNormal, WorldRayDirection());
-#endif
 
 #if SUBSTRATE_ENABLED && !USE_LEGACY_CODEPATH_FOR_SUBSTRATE
 	// unless proven otherwise, we are a basic unlit material
@@ -377,7 +372,7 @@ RAY_TRACING_ENTRY_CLOSEST_HIT(PathTracingMaterialCHS,
 #if !MATERIAL_TANGENTSPACENORMAL
 			N *= MaterialParameters.TwoSidedSign;
 #endif
-			SubstratePixelHeader.SharedLocalBases.Normals[Index] = AdjustShadingNormal(N, Payload.WorldGeoNormal, WorldRayDirection());
+			SubstratePixelHeader.SharedLocalBases.Normals[Index] = N;
 		}
 
 
@@ -828,12 +823,6 @@ RAY_TRACING_ENTRY_CLOSEST_HIT(PathTracingMaterialCHS,
 	Payload.WorldNormal *= MaterialParameters.TwoSidedSign;
 #endif
 
-#if PATH_TRACER_USE_VOLUMETRIC_SHADING_MODEL
-	// volumetric shading doesn't need to adjust normal
-#else
-	Payload.WorldNormal = AdjustShadingNormal(Payload.WorldNormal, Payload.WorldGeoNormal, WorldRayDirection());
-#endif
-
 	// Store the results in local variables and reuse instead of calling the functions multiple times.
 	half3 BaseColor = GetMaterialBaseColor(PixelMaterialInputs);
 	half  Metallic = GetMaterialMetallic(PixelMaterialInputs);
@@ -928,9 +917,6 @@ RAY_TRACING_ENTRY_CLOSEST_HIT(PathTracingMaterialCHS,
 			#else
 				float3 BottomNormal = ClearCoatBottomNormal0(MaterialParameters);
 			#endif
-		
-			// if we got a custom normal for the clearcoat, adjust it now
-			BottomNormal = AdjustShadingNormal(BottomNormal, Payload.WorldGeoNormal, WorldRayDirection());
 		#else
 			float3 BottomNormal = Payload.WorldNormal;
 		#endif
diff --git a/Engine/Shaders/Private/PathTracing/Volume/PathTracingAtmosphere.ush b/Engine/Shaders/Private/PathTracing/Volume/PathTracingAtmosphere.ush
--- a/Engine/Shaders/Private/PathTracing/Volume/PathTracingAtmosphere.ush
+++ b/Engine/Shaders/Private/PathTracing/Volume/PathTracingAtmosphere.ush
@@ -129,7 +129,7 @@ FPackedPathTracingPayload AtmosphereGetBlockerHit(float3 Origin, float3 Directio
 	Payload.SetBaseColor(Atmosphere.GroundAlbedo.rgb);
 	Payload.BSDFOpacity = 1.0;
 	Payload.WorldNormal = Payload.WorldGeoNormal = Payload.WorldSmoothNormal = WorldN;
-	Payload.ShadingModelID = SHADINGMODELID_NUM; // defaults to diffuse
+	Payload.ShadingModelID = SHADINGMODELID_LAMBERT; // defaults to diffuse
 	Payload.HitT = HitT;
 	Payload.PrimitiveLightingChannelMask = 7; // visible to all lights
 	Payload.SetFrontFace();
diff --git a/Engine/Shaders/Shared/PathTracingDefinitions.h b/Engine/Shaders/Shared/PathTracingDefinitions.h
--- a/Engine/Shaders/Shared/PathTracingDefinitions.h
+++ b/Engine/Shaders/Shared/PathTracingDefinitions.h
@@ -120,4 +120,5 @@
 #define PATH_TRACER_DEBUG_VIZ_SUBSTRATE_F90 			38
 #define PATH_TRACER_DEBUG_VIZ_EYE_IRIS_MASK             39
 #define PATH_TRACER_DEBUG_VIZ_EYE_CAUSTIC_NORMAL        40
-#define PATH_TRACER_DEBUG_VIZ_EYE_IRIS_NORMAL           41
\ No newline at end of file
+#define PATH_TRACER_DEBUG_VIZ_EYE_IRIS_NORMAL           41
+#define PATH_TRACER_DEBUG_VIZ_BACKFACING_NORMAL         42
diff --git a/Engine/Source/Runtime/Renderer/Private/PathTracing.cpp b/Engine/Source/Runtime/Renderer/Private/PathTracing.cpp
--- a/Engine/Source/Runtime/Renderer/Private/PathTracing.cpp
+++ b/Engine/Source/Runtime/Renderer/Private/PathTracing.cpp
@@ -103,6 +103,7 @@ struct FPathTracingDebugModeInfo
 		REGISTER_VIZ_TYPE(TEXT("EyeIrisMask")           , PATH_TRACER_DEBUG_VIZ_EYE_IRIS_MASK           );
 		REGISTER_VIZ_TYPE(TEXT("EyeIrisNormal")         , PATH_TRACER_DEBUG_VIZ_EYE_IRIS_NORMAL         );
 		REGISTER_VIZ_TYPE(TEXT("EyeCausticNormal")      , PATH_TRACER_DEBUG_VIZ_EYE_CAUSTIC_NORMAL      );
+		REGISTER_VIZ_TYPE(TEXT("BackfacingNormal")      , PATH_TRACER_DEBUG_VIZ_BACKFACING_NORMAL       );
 #undef REGISTER_VIZ_TYPE
 	}
 
