[HWRT] Include far field in all debug visualizations.
- removed far field specific visualization mode and instead added cvars to include whether far field is included and tinted in the visualizations.

#Synced-CL 49007722

diff --git a/Engine/Shaders/Private/RayTracing/RayTracingDebug.usf b/Engine/Shaders/Private/RayTracing/RayTracingDebug.usf
--- a/Engine/Shaders/Private/RayTracing/RayTracingDebug.usf
+++ b/Engine/Shaders/Private/RayTracing/RayTracingDebug.usf
@@ -41,6 +41,8 @@ uint VisualizationMode;
 uint PickerDomain;
 uint ShouldUsePreExposure;
 uint OpaqueOnly;
+uint TraceFarField;
+uint TintFarField;
 float TimingScale;
 
 float MaxTraceDistance;
@@ -140,7 +142,6 @@ float4 DebugRayTracingMaterial(uint2 PixelCoord, FRayDesc Ray, bool bIsFarField)
 
 	if (Payload.IsHit())
 	{
-
 		// 1. Non-material dependent debug view
 		switch (VisualizationMode)
 		{
@@ -191,8 +192,7 @@ float4 DebugRayTracingMaterial(uint2 PixelCoord, FRayDesc Ray, bool bIsFarField)
 			const FSubstrateBSDF BSDF = UnpackSubstrateBSDFIn(Payload.SubstrateData, SubstrateAddressing, SubstratePixelHeader); // Take the parameter of the first BSDF
 			const float3x3 TangentBasis = SubstrateGetBSDFSharedBasis(SubstratePixelHeader, BSDF, SubstrateAddressing);
 			switch (VisualizationMode)
-			{				
-			case RAY_TRACING_DEBUG_VIZ_FAR_FIELD:		return float4(SubstrateGetBSDFBaseColor(BSDF), Payload.HitT);
+			{
 			case RAY_TRACING_DEBUG_VIZ_GBUFFER_AO:		return float4(SubstrateGetIrradianceAndAO(SubstratePixelHeader).MaterialAO.xxx, 0.0f);
 			case RAY_TRACING_DEBUG_VIZ_SHADING_MODEL:	return float4(GetShadingModelColor(SubstrateGetLegacyShadingModels(BSDF)), 0.0f);
 			case RAY_TRACING_DEBUG_VIZ_BASE_COLOR:		return float4(SubstrateGetBSDFBaseColor(BSDF), 0.0f);
@@ -212,7 +212,6 @@ float4 DebugRayTracingMaterial(uint2 PixelCoord, FRayDesc Ray, bool bIsFarField)
 		switch (VisualizationMode)
 		{
 			case RAY_TRACING_DEBUG_VIZ_WORLD_NORMAL:	return float4(Payload.WorldNormal * 0.5 + 0.5, 0.0f);
-			case RAY_TRACING_DEBUG_VIZ_FAR_FIELD:		return float4(Payload.BaseColor.rgb, Payload.HitT);
 			case RAY_TRACING_DEBUG_VIZ_GBUFFER_AO:		return float4(Payload.GBufferAO, Payload.GBufferAO, Payload.GBufferAO, 0.0f);
 			case RAY_TRACING_DEBUG_VIZ_SHADING_MODEL:	return float4(GetShadingModelColor(Payload.ShadingModelID), 0.0f);
 			case RAY_TRACING_DEBUG_VIZ_BASE_COLOR:		return float4(Payload.BaseColor.rgb, 0.0f);
@@ -230,14 +229,14 @@ float4 DebugRayTracingMaterial(uint2 PixelCoord, FRayDesc Ray, bool bIsFarField)
 	}
 
 	// HitT is packed into the alpha for far field compositing
-	return (VisualizationMode == RAY_TRACING_DEBUG_VIZ_FAR_FIELD) ? float4(0, 0, 0, -1.0) : float4(0, 0, 0, 1);
+	return float4(0, 0, 0, -1.0);
 }
 
 #else // USE_DEBUG_CHS
 
 // Add functions that use the debug payload only
 
-float4 DebugRayTracingInstance(FRayDesc Ray, out float OutDepth)
+float4 DebugRayTracingInstance(FRayDesc Ray, bool bIsFarField, out float OutDepth)
 {
 	float4 Result = float4(0, 0, 0, 1);
 
@@ -255,11 +254,24 @@ float4 DebugRayTracingInstance(FRayDesc Ray, out float OutDepth)
 		RayFlags |= RAY_FLAG_FORCE_OPAQUE;
 	}
 
+	RaytracingAccelerationStructure TLASToTrace = TLAS;
+	if (bIsFarField)
+	{
+		TLASToTrace = FarFieldTLAS;
+
+		Ray.TMin = MaxTraceDistance;
+		Ray.TMax = FarFieldMaxTraceDistance;
+	}
+	else
+	{
+		Ray.TMax = (MaxTraceDistance > 0.0f) ? MaxTraceDistance : Ray.TMax;
+	}
+
 #if RAYTRACING_DEBUG_INLINE
 	FTraceRayInlineContext TraceRayInlineContext = CreateTraceRayInlineContext();
 
 	FTraceRayInlineResult TraceRayResult = TraceRayInline(
-		TLAS,
+		TLASToTrace,
 		RayFlags,
 		InstanceInclusionMask,
 		Ray.GetNativeDesc(),
@@ -322,7 +334,7 @@ float4 DebugRayTracingInstance(FRayDesc Ray, out float OutDepth)
 	}
 #else
 	TraceRay(
-		TLAS,
+		TLASToTrace,
 		RayFlags  /*RayFlags*/,
 		InstanceInclusionMask /*InstanceInclusionMask*/,
 		0 /*RayContributionToHitGroupIndex*/,
@@ -432,6 +444,7 @@ float4 DebugRayTracingInstance(FRayDesc Ray, out float OutDepth)
 	}
 	else
 	{
+		Result = float4(0, 0, 0, -1.0);
 		OutDepth = 0.0;
 	}
 
@@ -456,27 +469,39 @@ RAY_TRACING_ENTRY_RAYGEN_OR_INLINE(RayTracingDebugMain)
 	float OutDepth = 1.0f;
 #endif
 
+	bool bIsFarField = false;
+
 	bool bWriteColorOuput = VisualizationMode != RAY_TRACING_DEBUG_VIZ_SUBSTRATE_DATA;
 	bool bWriteDepthOutput = VisualizationMode == RAY_TRACING_DEBUG_VIZ_INSTANCE_OVERLAP;
 #if USE_DEBUG_CHS || RAYTRACING_DEBUG_INLINE
 	{
-		Result = DebugRayTracingInstance(Ray, OutDepth);
+		Result = DebugRayTracingInstance(Ray, bIsFarField, OutDepth);
+		
+		const bool bNearFieldMiss = Result.w < 0.0;
+		if (bNearFieldMiss && TraceFarField > 0)
+		{
+			bIsFarField = true;
+			Result = DebugRayTracingInstance(Ray, bIsFarField, OutDepth);
+		}
 	}
 #else // USE_DEBUG_CHS
 	{
-		bool bIsFarField = false;
 		Result = DebugRayTracingMaterial(PixelCoord, Ray, bIsFarField);
 		
 		const bool bNearFieldMiss = Result.w < 0.0;
-		if (VisualizationMode == RAY_TRACING_DEBUG_VIZ_FAR_FIELD && bNearFieldMiss)
+		if (bNearFieldMiss && TraceFarField > 0)
 		{
 			bIsFarField = true;
-			const float4 FarFieldTint = float4(1.0f, 0.5f, 0.5f, 1.0f);
-			Result = DebugRayTracingMaterial(PixelCoord, Ray, bIsFarField) * FarFieldTint;
+			Result = DebugRayTracingMaterial(PixelCoord, Ray, bIsFarField);
 		}
 	}
 #endif // USE_DEBUG_CHS
 
+	if (bIsFarField && TintFarField)
+	{
+		Result *= float4(1.0f, 0.5f, 0.5f, 1.0f);
+	}
+
 	if (ShouldUsePreExposure)
 	{
 		// Only  when the output is tonemapped
@@ -494,7 +519,7 @@ RAY_TRACING_ENTRY_RAYGEN_OR_INLINE(RayTracingDebugMain)
 	
 	if (bWriteColorOuput)
 	{
-		Output[PixelCoord] = Result;
+		Output[PixelCoord] = float4(Result.rgb, 1.0f);
 	}
 
 	if (bWriteDepthOutput)
diff --git a/Engine/Shaders/Shared/RaytracingDebugDefinitions.h b/Engine/Shaders/Shared/RaytracingDebugDefinitions.h
--- a/Engine/Shaders/Shared/RaytracingDebugDefinitions.h
+++ b/Engine/Shaders/Shared/RaytracingDebugDefinitions.h
@@ -34,7 +34,7 @@
 #define RAY_TRACING_DEBUG_VIZ_TIMING_ANY_HIT					24 
 #define RAY_TRACING_DEBUG_VIZ_TIMING_MATERIAL					25 
 #define RAY_TRACING_DEBUG_VIZ_TRIANGLES							26
-#define RAY_TRACING_DEBUG_VIZ_FAR_FIELD							27
+#define RAY_TRACING_DEBUG_VIZ_UNUSED_0							27
 #define RAY_TRACING_DEBUG_VIZ_TRAVERSAL_PRIMARY_NODE			28
 #define RAY_TRACING_DEBUG_VIZ_TRAVERSAL_PRIMARY_CLUSTER			29
 #define RAY_TRACING_DEBUG_VIZ_TRAVERSAL_PRIMARY_TRIANGLE		30
diff --git a/Engine/Source/Runtime/Engine/Private/RayTracingVisualizationData.cpp b/Engine/Source/Runtime/Engine/Private/RayTracingVisualizationData.cpp
--- a/Engine/Source/Runtime/Engine/Private/RayTracingVisualizationData.cpp
+++ b/Engine/Source/Runtime/Engine/Private/RayTracingVisualizationData.cpp
@@ -53,7 +53,6 @@ FRayTracingVisualizationData::FRayTracingVisualizationData()
 		// debugging the geometry itself
 		AddVisualizationMode(TEXT("Instances"), LOCTEXT("Instances", "Instances"), FModeType::Standard, RAY_TRACING_DEBUG_VIZ_INSTANCES, false);
 		AddVisualizationMode(TEXT("Triangles"), LOCTEXT("Triangles", "Triangles"), FModeType::Standard, RAY_TRACING_DEBUG_VIZ_TRIANGLES, false);
-		AddVisualizationMode(TEXT("FarField"), LOCTEXT("FarField", "Far Field"), FModeType::Standard, RAY_TRACING_DEBUG_VIZ_FAR_FIELD, false);
 		AddVisualizationMode(TEXT("DynamicInstances"), LOCTEXT("DynamicInstances", "Dynamic Instances"), FModeType::Standard, RAY_TRACING_DEBUG_VIZ_DYNAMIC_INSTANCES, false);
 		AddVisualizationMode(TEXT("ProxyType"), LOCTEXT("ProxyType", "Proxy Type"), FModeType::Standard, RAY_TRACING_DEBUG_VIZ_PROXY_TYPE, false);
 		AddVisualizationMode(TEXT("Picker"), LOCTEXT("Picker", "Picker"), FModeType::Standard, RAY_TRACING_DEBUG_VIZ_PICKER, false);
diff --git a/Engine/Source/Runtime/Renderer/Private/RayTracing/RayTracingDebug.cpp b/Engine/Source/Runtime/Renderer/Private/RayTracing/RayTracingDebug.cpp
--- a/Engine/Source/Runtime/Renderer/Private/RayTracing/RayTracingDebug.cpp
+++ b/Engine/Source/Runtime/Renderer/Private/RayTracing/RayTracingDebug.cpp
@@ -75,6 +75,20 @@ static TAutoConsoleVariable<int32> CVarRayTracingVisualizeOpaqueOnly(
 	ECVF_RenderThreadSafe
 );
 
+static TAutoConsoleVariable<bool> CVarRayTracingVisualizeFarField(
+	TEXT("r.RayTracing.Visualize.FarField"),
+	true,
+	TEXT("Whether to include far field geometry in ray tracing visualizations."),
+	ECVF_RenderThreadSafe
+);
+
+static TAutoConsoleVariable<bool> CVarRayTracingVisualizeFarFieldTint(
+	TEXT("r.RayTracing.Visualize.FarField.Tint"),
+	true,
+	TEXT("Whether to tint far field geometry in ray tracing visualizations."),
+	ECVF_RenderThreadSafe
+);
+
 static FAutoConsoleVariableDeprecated CVarRayTracingDebugModeOpaqueOnly_Deprecated(TEXT("r.RayTracing.DebugVisualizationMode.OpaqueOnly"), TEXT("r.RayTracing.Visualize.OpaqueOnly"), TEXT("5.6"));
 
 static TAutoConsoleVariable<float> CVarRayTracingVisualizeTimingScale(
@@ -188,6 +202,8 @@ BEGIN_SHADER_PARAMETER_STRUCT(FRayTracingDebugParameters, )
 	SHADER_PARAMETER(float, TimingScale)
 	SHADER_PARAMETER(float, MaxTraceDistance)
 	SHADER_PARAMETER(float, FarFieldMaxTraceDistance)
+	SHADER_PARAMETER(uint32, TraceFarField)
+	SHADER_PARAMETER(uint32, TintFarField)
 	SHADER_PARAMETER(uint32, OpaqueOnly)
 	SHADER_PARAMETER(float, TriangleHitCountMaxThreshold)
 	SHADER_PARAMETER(float, TriangleHitCountPerInstanceMaxThreshold)
@@ -1746,6 +1762,8 @@ void RenderRayTracingDebug(FRDGBuilder& GraphBuilder, const FScene& Scene, const
 	SharedParameters.ShouldUsePreExposure = View.Family->EngineShowFlags.Tonemapper;
 	SharedParameters.TimingScale = GetRayTracingDebugTimingScale();
 	SharedParameters.OpaqueOnly = CVarRayTracingVisualizeOpaqueOnly.GetValueOnRenderThread();
+	SharedParameters.TraceFarField = CVarRayTracingVisualizeFarField.GetValueOnRenderThread() ? 1 : 0;
+	SharedParameters.TintFarField = CVarRayTracingVisualizeFarFieldTint.GetValueOnRenderThread() ? 1 : 0;
 	SharedParameters.TriangleHitCountMaxThreshold = FMath::Clamp((float)CVarRayTracingVisualizeHitCountMaxThreshold.GetValueOnRenderThread(), 1, 100000);
 	SharedParameters.TriangleHitCountPerInstanceMaxThreshold = FMath::Max(1, CVarRayTracingVisualizeHitCountPerInstanceMaxThreshold.GetValueOnRenderThread());
 	SharedParameters.RayTracingDebugHitStatsUniformBuffer = DebugHitStatsUniformBuffer;
