Path Tracer: Fix overly dark glass absorption when materials have no internal medium.

The numeric threshold on the mean free path introduced a tiny amount of absorption that could become significant over long distances. Instead, rely on floating point semantics to encode a SigmaT of 0 as an MFP as +inf when needed.

#Synced-CL 49038315

diff --git a/Engine/Shaders/Private/PathTracing/PathTracingCommon.ush b/Engine/Shaders/Private/PathTracing/PathTracingCommon.ush
--- a/Engine/Shaders/Private/PathTracing/PathTracingCommon.ush
+++ b/Engine/Shaders/Private/PathTracing/PathTracingCommon.ush
@@ -821,7 +821,7 @@ struct FPathTracingPayload : FMinimalPayload
 	void SetHairPrimitiveUV(float2 UV) { SpecularEdgeColor.xy = UV; }
 
 	// SHADINGMODELID_SOLID_GLASS
-	float3 GetExtinction() { return rcp(max(MeanFreePath, 1e-5)); }
+	float3 GetExtinction() { return rcp(MeanFreePath); }
 	float3 GetGlassAlbedo() { return DiffuseColor; }
 	void SetGlassAlbedo(float3 A) { DiffuseColor = A; }
 	float GetGlassPhase() { return PhaseG; }
diff --git a/Engine/Shaders/Private/PathTracing/PathTracingMaterialHitShader.usf b/Engine/Shaders/Private/PathTracing/PathTracingMaterialHitShader.usf
--- a/Engine/Shaders/Private/PathTracing/PathTracingMaterialHitShader.usf
+++ b/Engine/Shaders/Private/PathTracing/PathTracingMaterialHitShader.usf
@@ -1284,7 +1284,7 @@ RAY_TRACING_ENTRY_CLOSEST_HIT(PathTracingMaterialCHS,
 			#endif
 
 			#if USE_LEGACY_CODEPATH_FOR_SUBSTRATE
-				Payload.MeanFreePath = rcp(max(Extinction, 1e-5f));
+				Payload.MeanFreePath = rcp(Extinction);
 			#else
 				Payload.SetExtinction(Extinction);
 			#endif
