Mobile: Fixed several issues with local light shading, where it was applying double attenuation in some cases.
Cleaned up mobile shading code, making it more inline with desktop version
Fixed issue with a rect light attenuation
Fixed issue with a SSS profile shading model
#Synced-CL 49063504

diff --git a/Engine/Shaders/Private/DeferredLightingCommon.ush b/Engine/Shaders/Private/DeferredLightingCommon.ush
--- a/Engine/Shaders/Private/DeferredLightingCommon.ush
+++ b/Engine/Shaders/Private/DeferredLightingCommon.ush
@@ -19,10 +19,6 @@
 #include "LightFunctionAtlas/LightFunctionAtlasCommon.usf"
 #endif
 
-#ifndef ALLOW_LOCAL_LIGHT_DISTANCE_ATTENUATION
-#define ALLOW_LOCAL_LIGHT_DISTANCE_ATTENUATION 0
-#endif
-
 // This Substrate include assumes that if inline shading is required, it needs to 
 // be defined prior to DeferredLightingCommong.ush
 // SubstrateStruct.MaterialTextureArray is only available to translucent materials or global shaders computing contact shadows.
@@ -258,9 +254,8 @@ float GetLocalLightAttenuation(
 	if (LightData.bInverseSquared)
 	{
 		LightMask = Square( saturate( 1 - Square( DistanceSqr * Square(LightData.InvRadius) ) ) );
-		// This extra attenuation has been added for supporting existing 'legacy' shading model on mobile. Not needed for capsule lights
-		// This is not needed for Substrate which unifies all lighting paths
-		#if SHADING_PATH_MOBILE && !MOBILE_SHADING_PATH_SUPPORT_CAPSULE_LIGHT && (!SUBSTRATE_ENABLED || ALLOW_LOCAL_LIGHT_DISTANCE_ATTENUATION)
+		// This extra attenuation has been added for supporting existing translucent shading model on mobile
+		#if SHADING_PATH_MOBILE && TRANSLUCENCY_SH_LIGHTING
 		if (!LightData.bRectLight)
 		{ 
 			LightMask *= 1.0f / (DistanceSqr + 1.0f);
@@ -356,49 +351,6 @@ FLightAccumulator AccumulateDynamicLighting(
 
 		LightAccumulator.EstimatedCost += 0.3f;		// add the cost of getting the shadow terms
 
-#if SHADING_PATH_MOBILE
-		const bool bNeedsSeparateSubsurfaceLightAccumulation = UseSubsurfaceProfile(GBuffer.ShadingModelID);
-		
-		FDirectLighting Lighting = (FDirectLighting)0;
-
-		half NoL = saturate(dot(N, L));
-	#if TRANSLUCENCY_NON_DIRECTIONAL
-		NoL = 1.0f;
-	#endif
-		BRANCH
-		if (LightData.bRectLight)
-		{
-			FRect Rect = GetRect( ToLight, LightData );
-			const FRectTexture SourceTexture = ConvertToRectTexture(LightData);
-
-		#if REFERENCE_QUALITY
-			Lighting = IntegrateBxDF( GBuffer, N, V, Rect, Shadow, SourceTexture, SVPos );
-		#else
-			Lighting = IntegrateBxDF( GBuffer, N, V, Rect, Shadow, SourceTexture);
-		#endif		
-		}
-		else
-		{
-		#if MOBILE_SHADING_PATH_SUPPORT_CAPSULE_LIGHT
-			FCapsuleLight Capsule = GetCapsule( ToLight, LightData );
-
-			#if REFERENCE_QUALITY
-				Lighting = IntegrateBxDF( GBuffer, N, V, Capsule, Shadow, SVPos );
-			#else
-				Lighting = IntegrateBxDF( GBuffer, N, V, Capsule, Shadow, LightData.bInverseSquared );
-			#endif
-		#else 
-			Lighting = EvaluateBxDF(GBuffer, N, V, L, NoL, Shadow);
-		#endif
-		}
-
-		Lighting.Specular *= LightData.SpecularScale;
-		Lighting.Diffuse *= LightData.DiffuseScale;
-				
-		LightAccumulator_AddSplit( LightAccumulator, Lighting.Diffuse, Lighting.Specular, Lighting.Diffuse, MaskedLightColor * Shadow.SurfaceShadow, bNeedsSeparateSubsurfaceLightAccumulation );
-		LightAccumulator_AddSplit( LightAccumulator, Lighting.Transmission, 0.0f, Lighting.Transmission, MaskedLightColor * Shadow.TransmissionShadow, bNeedsSeparateSubsurfaceLightAccumulation );
-#else // SHADING_PATH_MOBILE
-
 	#if !COMPILER_FXC // Work around "error X3524: duplicate attribute branch" in cases where this branch is constant
 		BRANCH
 	#endif
@@ -462,7 +414,6 @@ FLightAccumulator AccumulateDynamicLighting(
 			LightAccumulator.EstimatedCost += 0.4f;		// add the cost of the lighting computations (should sum up to 1 form one light)
 		#endif
 		}
-#endif // SHADING_PATH_MOBILE
 	}
 	return LightAccumulator;
 }
diff --git a/Engine/Shaders/Private/LightData.ush b/Engine/Shaders/Private/LightData.ush
--- a/Engine/Shaders/Private/LightData.ush
+++ b/Engine/Shaders/Private/LightData.ush
@@ -539,6 +539,42 @@ FRectTexture ConvertToRectTexture(FDeferredLightData In)
 	return Output;
 }
 
+void MobileOverrideDeferredLightParameters(inout FDeferredLightData Out)
+{
+#if	SHADING_PATH_MOBILE
+	#if RECT_LIGHT_AS_SPOTLIGHT && !MOBILE_DEFERRED_LIGHTING
+		// Render RectLight approximately as SpotLight. This simplification only applied to a forward shading
+		Out.bSpotLight = Out.bSpotLight || Out.bRectLight;
+		Out.bRectLight = false;
+	#endif
+		// Override values for performance on mobile
+	#if !MOBILE_SHADING_PATH_SUPPORT_CAPSULE_LIGHT
+		Out.SpecularScale = 1.0; // If enabled -> UnpackLightSpecularScale(In);
+		Out.DiffuseScale = 1.0; // If enabled -> UnpackLightDiffuseScale(In);
+		
+		if (!Out.bRectLight)
+		{
+			Out.Tangent = 0;
+			Out.SourceRadius = 0;
+			Out.SourceLength = 0;
+			Out.SoftSourceRadius = 0;
+		}
+	#endif
+#endif
+}
+
+void MobileOverrideDeferredLightParameters(inout FLightShaderParameters Out)
+{
+#if	SHADING_PATH_MOBILE && !MOBILE_SHADING_PATH_SUPPORT_CAPSULE_LIGHT
+		Out.Tangent = 0;
+		Out.SourceRadius = 0;
+		Out.SourceLength = 0;
+		Out.SoftSourceRadius = 0;
+		Out.SpecularScale = 1.0; // If enabled -> UnpackLightSpecularScale(In);
+		Out.DiffuseScale = 1.0; // If enabled -> UnpackLightDiffuseScale(In);
+#endif
+}
+
 // Helper function for converting FForwardLightData into FDeferredLightData
 FDeferredLightData ConvertToDeferredLight(
 	const FForwardLightData In, 
@@ -589,25 +625,7 @@ FDeferredLightData ConvertToDeferredLight(
 	OutPreviewShadowMapChannelMask	= UnpackShadowMapChannelMask(LightSceneInfoExtraDataPacked >> 4);
 	OutLightingChannelMask			= UnpackLightingChannelMask(LightSceneInfoExtraDataPacked);
 
-#if	SHADING_PATH_MOBILE
-	#if RECT_LIGHT_AS_SPOTLIGHT
-		// Render RectLight approximately as SpotLight
-		Out.bSpotLight = Out.bSpotLight || Out.bRectLight;
-		Out.bRectLight = false;
-	#endif
-	// Override values for performance on mobile
-	#if !MOBILE_SHADING_PATH_SUPPORT_CAPSULE_LIGHT
-		Out.SourceRadius = 0;
-		Out.SourceLength = 0;
-		Out.SoftSourceRadius = 0;
-		Out.SpecularScale = 1.0; // If enabled -> UnpackLightSpecularScale(In);
-		Out.DiffuseScale = 1.0; // If enabled -> UnpackLightDiffuseScale(In);
-		if (!Out.bRectLight)
-		{
-			Out.Tangent = 0;
-		}
-	#endif
-#endif
+	MobileOverrideDeferredLightParameters(Out);
 
 	return Out;
 }
@@ -693,6 +711,8 @@ FLightShaderParameters ConvertFromForward(const FForwardLightData In)
 	Out.RectLightAtlasUVScale	= RectLightData.AtlasData.AtlasUVScale;
 	Out.RectLightAtlasMaxLevel	= RectLightData.AtlasData.AtlasMaxLevel;
 
+	MobileOverrideDeferredLightParameters(Out);
+
 	return Out;
 }
 
diff --git a/Engine/Shaders/Private/MobileDeferredShading.usf b/Engine/Shaders/Private/MobileDeferredShading.usf
--- a/Engine/Shaders/Private/MobileDeferredShading.usf
+++ b/Engine/Shaders/Private/MobileDeferredShading.usf
@@ -2,11 +2,6 @@
 
 #define MOBILE_DEFERRED_LIGHTING 1
 
-#if SUBSTRATE_ENABLED
-// This is needed when lights are rendered deferred on mobile with Substrate because we use the legacy GBuffer representation and code process.
-#define ALLOW_LOCAL_LIGHT_DISTANCE_ATTENUATION 1
-#endif
-
 #include "Common.ush"
 #include "SHCommon.ush"
 
@@ -295,6 +290,8 @@ FDeferredLightData InitDeferredLightFromLightParameters(FLightShaderParameters I
 	Out.HairTransmittance		= InitHairTransmittanceData();
 	Out.ShadowedBits						= ShadowedBits;
 
+	MobileOverrideDeferredLightParameters(Out);
+
 #if SUPPORT_SPOTLIGHTS_SHADOW
 	Out.ShadowMapChannelMask	= 1.f;
 #endif
@@ -326,14 +323,6 @@ void MobileRadialLightPS(
 			
 	FLightAccumulator DirectLighting = (FLightAccumulator)0;
 	FDeferredLightData LightData = InitDeferredLightFromLightParameters(GetRootLightShaderParameters(), RADIAL_LIGHT_TYPE);
-		
-#if MOBILE_SUPPORTS_SM5_MATERIAL_NODES
-	if(LightData.bRectLight && GBuffer.ShadingModelID == SHADINGMODELID_SUBSURFACE_PROFILE)
-	{
-		// SSS has issues with rect light. Will be fixed in upcoming version
-		 GBuffer.ShadingModelID = SHADINGMODELID_DEFAULT_LIT;
-	}
-#endif
 
 #if !SIMPLE_LIGHT
 	// Affect the light color by any other shadow/transmittance before the lighting is evaluated
diff --git a/Engine/Shaders/Private/MobileLocalLightsBuffer.usf b/Engine/Shaders/Private/MobileLocalLightsBuffer.usf
--- a/Engine/Shaders/Private/MobileLocalLightsBuffer.usf
+++ b/Engine/Shaders/Private/MobileLocalLightsBuffer.usf
@@ -1,7 +1,6 @@
 // Copyright Epic Games, Inc. All Rights Reserved.
 
 #define SUPPORT_CONTACT_SHADOWS 0
-#define ALLOW_LOCAL_LIGHT_DISTANCE_ATTENUATION 1
 #define MAX_UINT 4294967295
 
 #if !COMPUTE_SHADER
diff --git a/Engine/Shaders/Private/ShadingModels.ush b/Engine/Shaders/Private/ShadingModels.ush
--- a/Engine/Shaders/Private/ShadingModels.ush
+++ b/Engine/Shaders/Private/ShadingModels.ush
@@ -613,7 +613,7 @@ FDirectLighting SubsurfaceProfileBxDF( FGBufferData GBuffer, half3 N, half3 V, F
 	half3 Tint = GetSubsurfaceProfileTexture(SSSS_TINT_SCALE_OFFSET, SubsurfaceProfileId).rgb;
 	
 	// Needs to apply shadow at here, since the preintegrated value has already counted it, and we need to skip applying shadow outside.
-	Lighting.Diffuse = lerp(AreaLight.FalloffColor * (AreaLight.Falloff * AreaLight.NoL) * Shadow.SurfaceShadow, BurleyDiffuse.rgb, Tint) * Diffuse_Lambert(GBuffer.DiffuseColor);
+	Lighting.Diffuse = (AreaLight.Falloff * AreaLight.NoL) * lerp(AreaLight.FalloffColor * Shadow.SurfaceShadow, BurleyDiffuse.rgb, Tint) * Diffuse_Lambert(GBuffer.DiffuseColor);
 #else
 
 #if MATERIAL_ROUGHDIFFUSE
diff --git a/Engine/Shaders/Private/VolumetricFog.usf b/Engine/Shaders/Private/VolumetricFog.usf
--- a/Engine/Shaders/Private/VolumetricFog.usf
+++ b/Engine/Shaders/Private/VolumetricFog.usf
@@ -1065,10 +1065,7 @@ void LightScatteringCS(
 #endif // USE_LIGHT_SOFT_FADING
 
 					Lighting = SofFade;
-#if !SHADING_PATH_MOBILE || MOBILE_SHADING_PATH_SUPPORT_CAPSULE_LIGHT || (SUBSTRATE_ENABLED && !ALLOW_LOCAL_LIGHT_DISTANCE_ATTENUATION)
-					// Mobile is already receiving the falloff from GetLocalLightAttenuation when capsule lights are not supported and Substrate is disabled.
 					Lighting *= IntegrateLight(Capsule, LightData.bInverseSquared);
-#endif	
 				}
 
 				float CombinedAttenuation = Lighting * LightMask;
