[TSR] Add thin geometry tiling optimization. Reduce the thin geometry detection cost.

#Synced-CL 48521690

diff --git a/Engine/Shaders/Private/TemporalSuperResolution/TSRDetectThinGeometry.usf b/Engine/Shaders/Private/TemporalSuperResolution/TSRDetectThinGeometry.usf
--- a/Engine/Shaders/Private/TemporalSuperResolution/TSRDetectThinGeometry.usf
+++ b/Engine/Shaders/Private/TemporalSuperResolution/TSRDetectThinGeometry.usf
@@ -92,16 +92,12 @@ tsr_short2 ComputeElementOutputPixelPos(const uint ElementIndex)
 		ElementIndex);
 }
 
-void FetchInput(
+void FetchCommonInput(
 	tsr_tensor_short2 PixelPos
-	,out tsr_tensor_half SceneDepth
-	,out tsr_tensor_half DeviceZThinGeometryDetectionError
 	,out tsr_tensor_half CurrentCoverage
-	,out tsr_tensor_half ReprojectedHistoryCoverage
 	,out tsr_tensor_bool bHasPixelAnimation
 	)
 {
-	tsr_tensor_float DepthCache;
 	tsr_tensor_bool bHasReprojectionOffset;
 #if CONFIG_THIN_GEOMETRY_EDGE_REPROJECTION
 	tsr_tensor_short2 ReprojectionOffset;
@@ -112,8 +108,6 @@ void FetchInput(
 	{
 		tsr_short2 ClampedFetchPixelPos = PixelPos.GetElement(ElementIndex);
 		CurrentCoverage.SetElement(ElementIndex, CurrentCoverageTexture[ClampedFetchPixelPos]);
-		ReprojectedHistoryCoverage.SetElement(ElementIndex, ReprojectedHistoryCoverageTexture[ClampedFetchPixelPos]);
-		DepthCache.SetElement(ElementIndex,SceneDepthTexture[ClampedFetchPixelPos]);
 
 		uint EncodedDilateMask = R8Output[int3(ClampedFetchPixelPos,0)];
 		bHasReprojectionOffset.SetElement(ElementIndex, (EncodedDilateMask & 0x80u) != 0u);
@@ -138,6 +132,23 @@ void FetchInput(
 #endif
 	
 	CurrentCoverage = ToBinaryCoverage(CurrentCoverage);
+}// FetchCommonInput
+
+void FetchRegularPassInput(
+	tsr_tensor_short2 PixelPos
+	,out tsr_tensor_half SceneDepth
+	,out tsr_tensor_half DeviceZThinGeometryDetectionError
+	,out tsr_tensor_half ReprojectedHistoryCoverage)
+{
+	tsr_tensor_float DepthCache;
+
+	UNROLL_N(SIMD_SIZE)
+	for (uint ElementIndex = 0; ElementIndex < SIMD_SIZE; ElementIndex++)
+	{
+		tsr_short2 ClampedFetchPixelPos = PixelPos.GetElement(ElementIndex);
+		ReprojectedHistoryCoverage.SetElement(ElementIndex, ReprojectedHistoryCoverageTexture[ClampedFetchPixelPos]);
+		DepthCache.SetElement(ElementIndex,SceneDepthTexture[ClampedFetchPixelPos]);
+	}
 
 	UNROLL_N(SIMD_SIZE)
 	for (uint ElementIndex = 0; ElementIndex < SIMD_SIZE; ElementIndex++)
@@ -148,8 +159,7 @@ void FetchInput(
 		SceneDepth.SetElement(ElementIndex, tsr_half(Depth));
 		DeviceZThinGeometryDetectionError.SetElement(ElementIndex, tsr_half(ErrorMultiplier * DeviceZError));
 	}
-} // FetchInput
-
+} // FetchRegularPassInput
 
 //------------------------------------------------------- FUNCTIONS
 
@@ -284,40 +294,41 @@ void DetectThinGeometryCS(
 	}
 	#endif
 
-
 	tsr_tensor_short2 PixelPos = ComputePixelPos();
 	tsr_tensor_half DeviceZ;
 	tsr_tensor_half DeviceZThinGeometryDetectionError;
 	tsr_tensor_half CurrentCoverage;
 	tsr_tensor_half ReprojectedHistoryCoverage;
 	tsr_tensor_bool bHasPixelAnimation;
-	FetchInput(PixelPos,
-		DeviceZ, 
-		DeviceZThinGeometryDetectionError
-		,CurrentCoverage
-		,ReprojectedHistoryCoverage
-		,bHasPixelAnimation
-	);
-	
-	//Debug[0].SetComponent(0, DeviceZ);
-	//Debug[0].SetComponent(1, DeviceZThinGeometryDetectionError);
 	
+	// Fetch data for tiling optimization.
+	FetchCommonInput(PixelPos, 
+		CurrentCoverage,
+		bHasPixelAnimation);
+		
+	bool bRegularPass = AnyThread(AnyComponent(and(CurrentCoverage > tsr_tensor_half::Const(0.0), not(bHasPixelAnimation))));
 	tsr_tensor_half ThinGeometryEdge = tsr_tensor_half::Const(0.0);
-
-#if !CONFIG_SKY_RELAXATION
-	{
-		ThinGeometryEdge = HorizontalThinEdge3x3(DeviceZ, DeviceZThinGeometryDetectionError);
-		ThinGeometryEdge = ThinGeometryEdge + VerticalThinEdge3x3(DeviceZ, DeviceZThinGeometryDetectionError);
-		//Debug[0].SetComponent(2, tsr_tensor_half(ThinGeometryEdge));
-	}
-#endif
-
 	tsr_tensor_half HistoryRelaxationWeight = tsr_tensor_half::Const(0);
 	tsr_tensor_bool IsClusterHoleRegion = tsr_tensor_bool::Const(false);
 	tsr_tensor_half NewSampleWeight = tsr_tensor_half::Const(0.1);
 
 	// Thin geometry coverage estimation
+	if (bRegularPass)
 	{
+		FetchRegularPassInput(PixelPos,
+			DeviceZ, 
+			DeviceZThinGeometryDetectionError,
+			ReprojectedHistoryCoverage
+		);
+		
+		#if !CONFIG_SKY_RELAXATION
+		{
+			ThinGeometryEdge = HorizontalThinEdge3x3(DeviceZ, DeviceZThinGeometryDetectionError);
+			ThinGeometryEdge = ThinGeometryEdge + VerticalThinEdge3x3(DeviceZ, DeviceZThinGeometryDetectionError);
+			//Debug[0].SetComponent(2, tsr_tensor_half(ThinGeometryEdge));
+		}
+		#endif
+
 		// Student t test of current observation and history to check if they are of the same distribution.
 		tsr_tensor_bool bSameDistribution = tsr_tensor_bool::Const(false);
 		tsr_tensor_bool NotAbsoluteMatch = tsr_tensor_bool::Const(false);
@@ -391,6 +402,10 @@ void DetectThinGeometryCS(
 		// Valid pixels for cluster hole stability improvement.
 		IsClusterHoleRegion = ReprojectedHistoryCoverage > tsr_tensor_half::Const(half(0.99f/127.0f));
 	}
+	else
+	{	// Fast pass does not have any thin geometry indicators. Update coverage to current.
+		ReprojectedHistoryCoverage = CurrentCoverage;
+	}
 
 	tsr_tensor_uint EncodedThinGeometryEdge = EncodeThinGeometryWeight(
 		ThinGeometryEdge
diff --git a/Engine/Shaders/Private/TemporalSuperResolution/TSRRejectShading.usf b/Engine/Shaders/Private/TemporalSuperResolution/TSRRejectShading.usf
--- a/Engine/Shaders/Private/TemporalSuperResolution/TSRRejectShading.usf
+++ b/Engine/Shaders/Private/TemporalSuperResolution/TSRRejectShading.usf
@@ -625,11 +625,13 @@ void RejectShadingCS(
 	tsr_tensor_half RejectionClampBlend = tsr_tensor_half::Const(1.0);
 	tsr_tensor_bool bAliasingIsVisible = tsr_tensor_bool::Const(false);
 #if CONFIG_THIN_GEOMETRY_DETECTION
-	tsr_tensor_half HistoryRelaxationFactor = tsr_tensor_half::Const(0);
+	bool bApplyThinGeometryRelaxation = false;
+	tsr_tensor_half ThinGeometryRelaxationFactor = tsr_tensor_half::Const(0);
 	{
 		tsr_tensor_short2 PixelPos = ComputePixelPos();
-		tsr_tensor_bool IsClusterHoleRegion;
-		FetchThinGeometryEdge(PixelPos, HistoryRelaxationFactor, IsClusterHoleRegion);
+		tsr_tensor_bool IsThinGeometryRegion;
+		FetchThinGeometryEdge(PixelPos, ThinGeometryRelaxationFactor, IsThinGeometryRegion);
+		bApplyThinGeometryRelaxation = AnyThread(AnyComponent(ThinGeometryRelaxationFactor > tsr_tensor_half::Const(0)));
 	}
 #endif
 	#if CONFIG_HISTORY_RESURRECTION
@@ -682,8 +684,8 @@ void RejectShadingCS(
 			/* HistoryC2 = */  ClampedHistory,
 			/* MoireError = */ MoireError,
 #if CONFIG_THIN_GEOMETRY_DETECTION
-			HistoryRelaxationFactor ,
-			true,
+			ThinGeometryRelaxationFactor ,
+			bApplyThinGeometryRelaxation,
 #endif
 			/* bEnableMoireError = */ CONFIG_MOIRE_DETECTION != 0 && bEnableFlickeringHeuristic != 0,
 			/* out */ RejectionBlendFinal,
@@ -742,7 +744,7 @@ void RejectShadingCS(
 			/* HistoryC2 = */  ClampedResurrectedHistory,
 			/* MoireError = */ tsr_tensor_half::Const(0.0),
 #if CONFIG_THIN_GEOMETRY_DETECTION
-			/*HistoryRelaxationFactor */tsr_tensor_half::Const(0),
+			/*ThinGeometryRelaxationFactor */tsr_tensor_half::Const(0),
 			false,
 #endif
 			/* bEnableMoireError = */ CONFIG_MOIRE_DETECTION != 0 && bEnableFlickeringHeuristic != 0,
