[TSR] Add weight modification optimization for the weight relaxation pass.

#Synced-CL 49128105

diff --git a/Engine/Shaders/Private/TemporalSuperResolution/TSRWeightRelaxation.usf b/Engine/Shaders/Private/TemporalSuperResolution/TSRWeightRelaxation.usf
--- a/Engine/Shaders/Private/TemporalSuperResolution/TSRWeightRelaxation.usf
+++ b/Engine/Shaders/Private/TemporalSuperResolution/TSRWeightRelaxation.usf
@@ -221,16 +221,25 @@ void WeightRelaxationCS(
 	tsr_tensor_half RatioToLerpWeight = max(tsr_tensor_half::Const(0.6)-TranslucencyRatio,tsr_tensor_half::Const(0));
 	RatioToLerpWeight = min(RatioToLerpWeight * RatioToLerpWeight * tsr_tensor_half::Const(3.33), tsr_tensor_half::Const(1));
 	RatioToLerpWeight = select(or(IsPartialCoverageRegion,TranslucencyRatio > tsr_tensor_half::Const(0.6)), RatioToLerpWeight, tsr_tensor_half::Const(1));
-	HistoryRelaxationWeight = lerp(tsr_tensor_half::Const(0), HistoryRelaxationWeight, RatioToLerpWeight); 
 	
-	tsr_tensor_ushort EncodedThinGeometry = EncodeThinGeometry(HistoryRelaxationWeight,IsPartialCoverageRegion);
-	// Output
-	UNROLL_N(SIMD_SIZE)
-		for(uint ElementIndex = 0; ElementIndex < SIMD_SIZE; ElementIndex++)
-		{
-			tsr_short2 OutputPixelPos = ComputeElementOutputPixelPos(ElementIndex);
-			R8Output[int3(OutputPixelPos, ThinGeometryTextureIndex)] = EncodedThinGeometry.GetElement(ElementIndex);
-		}
+	// Write only when the weights has been modified.
+	#if !CONFIG_SKY_RELAXATION
+	bool bRelaxationWeightsModified = AnyComponent(and(HistoryRelaxationWeight > tsr_tensor_half::Const(tsr_half(0)), RatioToLerpWeight < tsr_tensor_half::Const(tsr_half(1.0))));
+	
+	BRANCH
+	if (bRelaxationWeightsModified)
+	#endif
+	{
+		HistoryRelaxationWeight = lerp(tsr_tensor_half::Const(0), HistoryRelaxationWeight, RatioToLerpWeight); 
+		tsr_tensor_ushort EncodedThinGeometry = EncodeThinGeometry(HistoryRelaxationWeight,IsPartialCoverageRegion);
+		// Output
+		UNROLL_N(SIMD_SIZE)
+			for(uint ElementIndex = 0; ElementIndex < SIMD_SIZE; ElementIndex++)
+			{
+				tsr_short2 OutputPixelPos = ComputeElementOutputPixelPos(ElementIndex);
+				R8Output[int3(OutputPixelPos, ThinGeometryTextureIndex)] = EncodedThinGeometry.GetElement(ElementIndex);
+			}
+	}
 
 #if DEBUG_OUTPUT
 	{
