Modified ComputeAverageLuminanceWithoutOutlier(...) to return MinAverageLuminance when histogram is empty.
- can happen due to r.EyeAdaptation.BlackHistogramBucketInfluence.

#Synced-CL 49186836

diff --git a/Engine/Shaders/Private/PostProcessHistogramCommon.ush b/Engine/Shaders/Private/PostProcessHistogramCommon.ush
--- a/Engine/Shaders/Private/PostProcessHistogramCommon.ush
+++ b/Engine/Shaders/Private/PostProcessHistogramCommon.ush
@@ -153,7 +153,12 @@ float ComputeAverageLuminanceWithoutOutlier(Texture2D HistogramTexture, float Mi
 		SumWithoutOutliers += float2(LogLuminanceAtBucket, 1) * LocalValue;
 	}
 
-	float AvgLogLuminance = SumWithoutOutliers.x / max(0.0001f, SumWithoutOutliers.y);
+	if (SumWithoutOutliers.y == 0)
+	{
+		return EyeAdaptation_MinAverageLuminance / 0.18f;
+	}
+
+	float AvgLogLuminance = SumWithoutOutliers.x / SumWithoutOutliers.y;
 
 	return exp2(AvgLogLuminance);
 }
