Substrate - fixed bug in SanitizeAndNormalizeSharedLocalBases from 48650197.

#rnx
#tests PC editor dx12.
#Synced-CL 49271765

diff --git a/Engine/Shaders/Private/Substrate/Substrate.ush b/Engine/Shaders/Private/Substrate/Substrate.ush
--- a/Engine/Shaders/Private/Substrate/Substrate.ush
+++ b/Engine/Shaders/Private/Substrate/Substrate.ush
@@ -3465,7 +3465,7 @@ void SanitizeAndNormalizeSharedLocalBases(inout FSharedLocalBases SharedLocalBas
 			#if IS_NANITE_PASS
 				SharedLocalBases.Tangents[i] = SUBSTRATE_IS_VALID_DIRECTION(FallbackTangent) ? FallbackTangent : float3(1, 0, 0);
 			#else
-				SharedLocalBases.Normals[i] = FallbackTangent;
+				SharedLocalBases.Tangents[i] = FallbackTangent;
 			#endif
 			}
 			SharedLocalBases.Tangents[i] = normalize(SharedLocalBases.Tangents[i]);
