[AutoExposure] Maintain UE 5.7 ComputeAverageLuminanceWithoutOutlier(...) behavior when ExposureLowPercent == ExposureHighPercent.

#Synced-CL 49278935

diff --git a/Engine/Shaders/Private/PostProcessHistogramCommon.ush b/Engine/Shaders/Private/PostProcessHistogramCommon.ush
--- a/Engine/Shaders/Private/PostProcessHistogramCommon.ush
+++ b/Engine/Shaders/Private/PostProcessHistogramCommon.ush
@@ -132,6 +132,11 @@ float ComputeHistogramSum(Texture2D HistogramTexture)
 // @param MaxFractionSum e.g. ComputeHistogramSum() * 0.9f for 90% percentil
 float ComputeAverageLuminanceWithoutOutlier(Texture2D HistogramTexture, float MinFractionSum, float MaxFractionSum)
 {
+	if (abs(MinFractionSum - MaxFractionSum) < 0.0001f)
+	{
+		return 1.0f; // special case to avoid affecting existing projects (eg: Lyra)
+	}
+
 	float2 SumWithoutOutliers = 0;
 	
 	UNROLL for(uint i = 0; i < HISTOGRAM_SIZE; ++i)
