Ensure Depth samples don't get truncated to half's on mobile.Texture2DSample* helpers return MaterialFloat4 which can be reduced to half4 when full precision isn't enabled.

#Synced-CL 48848898

diff --git a/Engine/Shaders/Private/DeferredShadingCommon.ush b/Engine/Shaders/Private/DeferredShadingCommon.ush
--- a/Engine/Shaders/Private/DeferredShadingCommon.ush
+++ b/Engine/Shaders/Private/DeferredShadingCommon.ush
@@ -837,7 +837,7 @@ void MobileFetchGBuffer(in float2 UV, out half4 GBufferA, out half4 GBufferB, ou
 #elif (FRAMEBUFFER_FETCH_ACCESS && !MAC && (USE_SCENE_DEPTH_AUX || USE_GLES_FBF_DEFERRED))
 	SceneDepth = ConvertFromDeviceZ(DepthbufferFetchES2());
 #else
-	SceneDepth = ConvertFromDeviceZ(Texture2DSampleLevel(MobileSceneTextures.SceneDepthTexture, MobileSceneTextures.SceneDepthTextureSampler, UV, 0).r);
+	SceneDepth = ConvertFromDeviceZ(MobileSceneTextures.SceneDepthTexture.SampleLevel(MobileSceneTextures.SceneDepthTextureSampler, UV, 0).r);
 #endif
 }
 #endif //MOBILE_DEFERRED_SHADING
@@ -858,10 +858,10 @@ FGBufferData MobileFetchAndDecodeGBuffer(in float2 UV, in float2 PixelPos, uint
 	GBuffer.Depth = CalcSceneDepth(UV);
 #endif
 #if	MOBILE_MULTI_VIEW
-	GBuffer.CustomDepth = ConvertFromDeviceZ(Texture2DArraySample(MobileSceneTextures.CustomDepthTextureArray, MobileSceneTextures.CustomDepthTextureSampler, float3(UV, ViewId)).r);
+	GBuffer.CustomDepth = ConvertFromDeviceZ(MobileSceneTextures.CustomDepthTextureArray.Load(int4(PixelPos.xy, ViewId, 0)).r);
 	GBuffer.CustomStencil = MobileSceneTextures.CustomStencilTextureArray.Load(int4(PixelPos.xy, ViewId, 0)) STENCIL_COMPONENT_SWIZZLE;
 #else
-	GBuffer.CustomDepth = ConvertFromDeviceZ(Texture2DSample(MobileSceneTextures.CustomDepthTexture, MobileSceneTextures.CustomDepthTextureSampler, float3(UV, ViewId)).r);
+	GBuffer.CustomDepth = ConvertFromDeviceZ(MobileSceneTextures.CustomDepthTexture.Load(int3(PixelPos.xy, 0))).r;
 	GBuffer.CustomStencil = MobileSceneTextures.CustomStencilTexture.Load(int3(PixelPos.xy, 0)) STENCIL_COMPONENT_SWIZZLE;
 #endif
 
