[TSR] Improve Thin Geometry Detection performance.
Reduce RejectShading Thin Geometry LerpClamp3x3() logic VGPR pressure by reusing existing VGPRs, swap MinMax3x3 by consecutive calls. This increases the occupancy from 40% to 60% on some consoles (saving approximately 0.1ms), back to the occupancy without thin geometry detection.

#Synced-CL 49128105

diff --git a/Engine/Shaders/Private/TemporalSuperResolution/TSRShadingAnalysis.ush b/Engine/Shaders/Private/TemporalSuperResolution/TSRShadingAnalysis.ush
--- a/Engine/Shaders/Private/TemporalSuperResolution/TSRShadingAnalysis.ush
+++ b/Engine/Shaders/Private/TemporalSuperResolution/TSRShadingAnalysis.ush
@@ -502,18 +502,34 @@ void MeasureRejection(
 #if CONFIG_THIN_GEOMETRY_DETECTION  
 	if (bEnableHistoryControl)
 	{
-		tsr_tensor_halfC LerpClampedHistoryInput = LerpClamp3x3(FilteredHistory, FilteredInput, HistoryRelaxationFactor);
-		#if CONFIG_COMPILE_FP16
-			tsr_tensor_halfC FilteredHistoryBoxMin;
-			tsr_tensor_halfC FilteredHistoryBoxMax;
-			MinMax3x3(LerpClampedHistoryInput, /* out */ FilteredHistoryBoxMin, /* out */ FilteredHistoryBoxMax);
-        #else
-			tsr_tensor_halfC FilteredHistoryBoxMin = Min3x3(LerpClampedHistoryInput);
-			tsr_tensor_halfC FilteredHistoryBoxMax = Max3x3(LerpClampedHistoryInput);
-        #endif
-		tsr_tensor_halfC ScalingFactor = tsr_tensor_halfC::Vectorize((HistoryRelaxationFactor));
-		FilteredBoxMin = select(ScalingFactor > tsr_tensor_halfC::Const(tsr_half(1.0f / 127.0f)), FilteredHistoryBoxMin, FilteredBoxMin);
-		FilteredBoxMax = select(ScalingFactor > tsr_tensor_halfC::Const(tsr_half(1.0f / 127.0f)), FilteredHistoryBoxMax, FilteredBoxMax);
+		// VGPR optimization for LerpClamp3x3()
+		HistoryRelaxationFactor = select(HistoryRelaxationFactor > tsr_tensor_half::Const(tsr_half(1.0f / 127.0f)), HistoryRelaxationFactor, tsr_tensor_half::Const(0));
+		tsr_tensor_halfC LerpFactorVector = tsr_tensor_halfC::Vectorize((HistoryRelaxationFactor));
+		
+		PLATFORM_SPECIFIC_ISOLATE
+		{
+			// Register pressure is high. No MinMax3x3()
+			tsr_tensor_halfC ToClampBoundaryLimit;
+				
+			ToClampBoundaryLimit = Min3x3(FilteredHistory);
+			FilteredBoxMin = lerp(FilteredBoxMin, ToClampBoundaryLimit, LerpFactorVector);
+
+			ToClampBoundaryLimit = Max3x3(FilteredHistory);
+			FilteredBoxMax = lerp(FilteredBoxMax, ToClampBoundaryLimit, LerpFactorVector);
+		}
+		
+		tsr_tensor_halfC LerpClampedHistoryInput = fastClamp(FilteredHistory, FilteredBoxMin, FilteredBoxMax);
+
+		{
+			// Register pressure is high. No MinMax3x3()
+			tsr_tensor_halfC ToClampBoundaryLimit;
+				
+			ToClampBoundaryLimit = Min3x3(LerpClampedHistoryInput);
+			FilteredBoxMin = select(LerpFactorVector > tsr_tensor_halfC::Const(0), ToClampBoundaryLimit, FilteredBoxMin);
+
+			ToClampBoundaryLimit = Max3x3(LerpClampedHistoryInput);
+			FilteredBoxMax = select(LerpFactorVector > tsr_tensor_halfC::Const(0), ToClampBoundaryLimit, FilteredBoxMax);
+		}
 	}
 #endif
 
