Clarified rect light threshold usage.

#rnx
#tests pc editor dx12
#Synced-CL 49403680

diff --git a/Engine/Shaders/Private/RectLight.ush b/Engine/Shaders/Private/RectLight.ush
--- a/Engine/Shaders/Private/RectLight.ush
+++ b/Engine/Shaders/Private/RectLight.ush
@@ -52,6 +52,11 @@ float3 SampleRectTexture(FRectTexture RectTexture, float2 RectUV, float Level, b
 #endif
 }
 
+// This threshold is here to avoid float explosion / NaNs at the border of rect light in sharp reflections, as well as rect light ghosts in other part of the image.
+// All that can happen for relatively long and thin rect lights. The trade off is that the rect light edge can then look softer than it should be in smooht reflections for elongated rect lights.
+// We can change the threshold but there will always be a size that will break the look. The only solution is to cut long rect light sources into smaller ones.
+#define RECT_LIGHT_SAFE_THRESHOLD 1.0e-4f
+
 // Optimized Lambert
 float3 RectIrradianceLambert( float3 N, FRect Rect, out float BaseIrradiance, out float NoL )
 {
@@ -117,10 +122,10 @@ float3 RectIrradianceLambert( float3 N, FRect Rect, out float BaseIrradiance, ou
 	// normalize( cross( L0, L1 ) ) = cross( L0, L1 ) * rsqrt( 1 - dot( L0, L1 )^2 )
 	// acos( x ) ~= sqrt(1 - x) * (1.5708 - 0.175 * x)
 
-	float w01 = ( 1.5708 - 0.175 * c01 ) * rsqrt( max(c01 + 1, 1.0e-4f) );	// 1 mad, 1 add, 1 rsqrt
-	float w12 = ( 1.5708 - 0.175 * c12 ) * rsqrt( max(c12 + 1, 1.0e-4f) );	// 1 mad, 1 add, 1 rsqrt
-	float w23 = ( 1.5708 - 0.175 * c23 ) * rsqrt( max(c23 + 1, 1.0e-4f) );	// 1 mad, 1 add, 1 rsqrt
-	float w30 = ( 1.5708 - 0.175 * c30 ) * rsqrt( max(c30 + 1, 1.0e-4f) );	// 1 mad, 1 add, 1 rsqrt
+	float w01 = ( 1.5708 - 0.175 * c01 ) * rsqrt( max(c01 + 1, RECT_LIGHT_SAFE_THRESHOLD) );	// 1 mad, 1 add, 1 rsqrt
+	float w12 = ( 1.5708 - 0.175 * c12 ) * rsqrt( max(c12 + 1, RECT_LIGHT_SAFE_THRESHOLD) );	// 1 mad, 1 add, 1 rsqrt
+	float w23 = ( 1.5708 - 0.175 * c23 ) * rsqrt( max(c23 + 1, RECT_LIGHT_SAFE_THRESHOLD) );	// 1 mad, 1 add, 1 rsqrt
+	float w30 = ( 1.5708 - 0.175 * c30 ) * rsqrt( max(c30 + 1, RECT_LIGHT_SAFE_THRESHOLD) );	// 1 mad, 1 add, 1 rsqrt
 	// 8 alu, 4 rsqrt
 #endif
 
@@ -202,7 +207,7 @@ float IntegrateEdge( float3 L0, float3 L1 )
 	float w01 = ( 0.8543985 + (0.4965155 + 0.0145206 * abs(c01)) * abs(c01) ) /
 				( 3.4175940 + (4.1616724 + abs(c01)) * abs(c01) );
 	
-	w01 = c01 > 0 ? w01 : 0.5 * rsqrt( max(1 - c01 * c01, 1.0e-4f)  ) - w01;
+	w01 = c01 > 0 ? w01 : 0.5 * rsqrt( max(1 - c01 * c01, RECT_LIGHT_SAFE_THRESHOLD)  ) - w01;
 #endif
 
 	return w01;
