Substrate - Fixi for Lumen HWRT reflection when Lumen diffuse GI is OFF.

#rnx
#tests PC editor dx
#Synced-CL 49403680

diff --git a/Engine/Shaders/Private/ReflectionEnvironmentPixelShader.usf b/Engine/Shaders/Private/ReflectionEnvironmentPixelShader.usf
--- a/Engine/Shaders/Private/ReflectionEnvironmentPixelShader.usf
+++ b/Engine/Shaders/Private/ReflectionEnvironmentPixelShader.usf
@@ -548,6 +548,10 @@ void ReflectionEnvironmentSkyLighting(
 	
 		#if REFLECTION_SOURCE_TYPE == REFLECTION_SOURCE_LUMEN_STANDALONE
 			{
+				// Pre exposure the sky light specular evaluated from within SubstrateEnvLightingCommon (CompositeReflectionCapturesAndSkylightTWS) when rendering lumen HWRT reflections while diffuse global illumination is "none".
+				// This matches with the fact that LumenReflectionTexture is also already pre-exposed, see condition below using "REFLECTION_SOURCE_TYPE == REFLECTION_SOURCE_LUMEN_STANDALONE".
+				SpecularLighting *= View.PreExposure; 
+
 				const FLumenReflectionOnlyOutput LumenReflections = GetLumenReflectionOnly(SubstrateEnvLight, SubstrateBSDFContext.BSDF, BSDFThroughput, Texture2DArraySample(LumenReflectionTexture, LumenReflectionTextureSampler, float3(BufferUV, ClosureIndex)));
 
 				SpecularLighting *= min(LumenReflections.LightingLobe0.a, LumenReflections.LightingLobe1.a);
