[TSR] fp16 optimization for thin geometry detection pass with some minor performance improvement (~5us) at 1080p on consoles.

#Synced-CL 49128105

diff --git a/Engine/Shaders/Private/TemporalSuperResolution/TSRDetectThinGeometry.usf b/Engine/Shaders/Private/TemporalSuperResolution/TSRDetectThinGeometry.usf
--- a/Engine/Shaders/Private/TemporalSuperResolution/TSRDetectThinGeometry.usf
+++ b/Engine/Shaders/Private/TemporalSuperResolution/TSRDetectThinGeometry.usf
@@ -42,7 +42,7 @@ Texture2D<tsr_half> ReprojectedHistoryCoverageTexture;
 Texture2D<tsr_half2> DecimateMaskTexture;
 
 RWTexture2DArray<tsr_half> HistoryCoverageOutput;
-RWTexture2DArray<uint> R8Output;
+RWTexture2DArray<tsr_ushort> R8Output;
 
 //------------------------------------------------------- GLOBALS
 static uint2 GGroupId = 0;
@@ -109,8 +109,8 @@ void FetchCommonInput(
 		tsr_short2 ClampedFetchPixelPos = PixelPos.GetElement(ElementIndex);
 		CurrentCoverage.SetElement(ElementIndex, CurrentCoverageTexture[ClampedFetchPixelPos]);
 
-		uint EncodedDilateMask = R8Output[int3(ClampedFetchPixelPos,0)];
-		bHasReprojectionOffset.SetElement(ElementIndex, (EncodedDilateMask & 0x80u) != 0u);
+		tsr_ushort EncodedDilateMask = R8Output[int3(ClampedFetchPixelPos,0)];
+		bHasReprojectionOffset.SetElement(ElementIndex, (EncodedDilateMask & tsr_ushort(0x80)) != tsr_ushort(0));
 #if CONFIG_THIN_GEOMETRY_EDGE_REPROJECTION
 		ReprojectionOffset.SetElement(ElementIndex, DecodeReprojectionOffset(R8Output[int3(ClampedFetchPixelPos, ThinGeometryTextureIndex)]));
 #endif
@@ -255,7 +255,7 @@ TLaneVector2D<ScalarType, VectorSize, LaneStrideX, LaneStrideY> VerticalThinEdge
 
 #endif // !CONFIG_SKY_RELAXATION
 
-tsr_tensor_uint EncodeThinGeometryWeight(tsr_tensor_half ThinGeometryEdge
+tsr_tensor_ushort EncodeThinGeometryWeight(tsr_tensor_half ThinGeometryEdge
 	, tsr_tensor_bool IsClusterHoleRegion,
 	tsr_tensor_half Weight
 )
@@ -266,8 +266,8 @@ tsr_tensor_uint EncodeThinGeometryWeight(tsr_tensor_half ThinGeometryEdge
 	Weight = select(IsClusterHoleRegion, Weight, tsr_tensor_half::Const(0));
 	Weight = select(ThinGeometryEdge > tsr_tensor_half::Const(0), ThinGeometryEdge, Weight);
 	IsClusterHoleRegion = select(ThinGeometryEdge > tsr_tensor_half::Const(0), tsr_tensor_bool::Const(false), IsClusterHoleRegion);
-	tsr_tensor_uint Encoded = bit_shift_left(tsr_tensor_uint::CastFrom<tsr_half>(Weight * tsr_half(127.0)),tsr_tensor_uint::Const(0x1u));
-	Encoded = bit_or(Encoded,select(IsClusterHoleRegion,tsr_tensor_uint::Const(0x1u),tsr_tensor_uint::Const(0)));
+	tsr_tensor_ushort Encoded = bit_shift_left(tsr_tensor_ushort::CastFrom<tsr_half>(Weight * tsr_half(127.0)),tsr_tensor_ushort::Const(tsr_ushort(0x1)));
+	Encoded = bit_or(Encoded,select(IsClusterHoleRegion,tsr_tensor_ushort::Const(tsr_ushort(0x1)),tsr_tensor_ushort::Const(0)));
 	return Encoded;
 }
 
@@ -407,7 +407,7 @@ void DetectThinGeometryCS(
 		ReprojectedHistoryCoverage = CurrentCoverage;
 	}
 
-	tsr_tensor_uint EncodedThinGeometryEdge = EncodeThinGeometryWeight(
+	tsr_tensor_ushort EncodedThinGeometryEdge = EncodeThinGeometryWeight(
 		ThinGeometryEdge
 		,IsClusterHoleRegion
 		,HistoryRelaxationWeight
diff --git a/Engine/Shaders/Private/TemporalSuperResolution/TSRRejectShading.usf b/Engine/Shaders/Private/TemporalSuperResolution/TSRRejectShading.usf
--- a/Engine/Shaders/Private/TemporalSuperResolution/TSRRejectShading.usf
+++ b/Engine/Shaders/Private/TemporalSuperResolution/TSRRejectShading.usf
@@ -129,10 +129,10 @@ void FetchThinGeometryEdge(
 	for (uint ElementIndex = 0; ElementIndex < SIMD_SIZE; ElementIndex++)
 	{
 		tsr_short2 ClampedFetchPixelPos = PixelPos.GetElement(ElementIndex);
-		uint EncodedThinGeometryWeight = ThinGeometryTexture[ClampedFetchPixelPos];
+		tsr_short EncodedThinGeometryWeight = ThinGeometryTexture[ClampedFetchPixelPos];
 		
-		tsr_half Weight = tsr_half(EncodedThinGeometryWeight >> 1) * rcp(tsr_half(127.0));
-		bool UseMultiplyOperation = bool(EncodedThinGeometryWeight & 0x1u);
+		tsr_half Weight = tsr_half(EncodedThinGeometryWeight >> tsr_ushort(1)) * rcp(tsr_half(127.0));
+		bool UseMultiplyOperation = (EncodedThinGeometryWeight & tsr_ushort(0x1))!= tsr_ushort(0x0);
 
 		HistoryRelaxationWeight.SetElement(ElementIndex, Weight);
 		IsClusterHoleRegion.SetElement(ElementIndex, UseMultiplyOperation);
diff --git a/Engine/Shaders/Private/TemporalSuperResolution/TSRWeightRelaxation.usf b/Engine/Shaders/Private/TemporalSuperResolution/TSRWeightRelaxation.usf
--- a/Engine/Shaders/Private/TemporalSuperResolution/TSRWeightRelaxation.usf
+++ b/Engine/Shaders/Private/TemporalSuperResolution/TSRWeightRelaxation.usf
@@ -35,7 +35,7 @@ Texture2D<tsr_half>  InputMoireLumaTexture;
 Texture2D<tsr_halfC> InputTexture;
 Texture2D<tsr_half4> InputSceneTranslucencyTexture;
 
-RWTexture2DArray<uint> R8Output;
+RWTexture2DArray<tsr_ushort> R8Output;
 
 //------------------------------------------------------- GLOBALS
 static uint2 GGroupId = 0;
