[Materials] Fixed a bug with the ActorWorldPosition material node when it's used for calculating the previous frame WPO for motion vectors.

#Synced-CL 49417496

diff --git a/Engine/Shaders/Private/MaterialTemplate.ush b/Engine/Shaders/Private/MaterialTemplate.ush
--- a/Engine/Shaders/Private/MaterialTemplate.ush
+++ b/Engine/Shaders/Private/MaterialTemplate.ush
@@ -1410,14 +1410,14 @@ MaterialFloat2 GetDefaultSceneTextureUV(FMaterialPixelParameters Parameters, con
 
 	FWSVector3 GetPreviousActorWorldPosition(FMaterialVertexParameters Parameters)
 	{
-		float3 InstanceToActor = WSMultiplyDemote(GetActorWorldPosition(Parameters), GetWorldToInstance(Parameters));
-		return WSMultiply(InstanceToActor, Parameters.LWCData.PreviousLocalToWorld);
+		float3 InstanceSpacePos = WSMultiplyDemote(GetActorWorldPosition(Parameters), GetWorldToInstance(Parameters));
+		return WSMultiply(InstanceSpacePos, GetPrevInstanceToWorld(Parameters));
 	}
 
 	FWSVector3 GetPreviousActorWorldPosition(FMaterialPixelParameters Parameters)
 	{
-		float3 InstanceToActor = WSMultiplyDemote(GetActorWorldPosition(Parameters), GetWorldToInstance(Parameters));
-		return WSMultiply(InstanceToActor, Parameters.LWCData.PreviousLocalToWorld);
+		float3 InstanceSpacePos = WSMultiplyDemote(GetActorWorldPosition(Parameters), GetWorldToInstance(Parameters));
+		return WSMultiply(InstanceSpacePos, GetPrevInstanceToWorld(Parameters));
 	}
 
 	float3 GetActorTranslatedWorldPosition(FMaterialVertexParameters Parameters)
