Check that ExposureLowPercent and ExposureHighPercent are equal directly instead of using Min/MaxFractionSum.

#Synced-CL 49470770

diff --git a/Engine/Shaders/Private/PostProcessHistogramCommon.ush b/Engine/Shaders/Private/PostProcessHistogramCommon.ush
--- a/Engine/Shaders/Private/PostProcessHistogramCommon.ush
+++ b/Engine/Shaders/Private/PostProcessHistogramCommon.ush
@@ -132,7 +132,7 @@ float ComputeHistogramSum(Texture2D HistogramTexture)
 // @param MaxFractionSum e.g. ComputeHistogramSum() * 0.9f for 90% percentil
 float ComputeAverageLuminanceWithoutOutlier(Texture2D HistogramTexture, float MinFractionSum, float MaxFractionSum)
 {
-	if (abs(MinFractionSum - MaxFractionSum) < 0.0001f)
+	if (abs(EyeAdaptation_ExposureLowPercent - EyeAdaptation_ExposureHighPercent) < 0.0001f)
 	{
 		return 1.0f; // special case to avoid affecting existing projects (eg: Lyra)
 	}
