Path Tracer: Remove the legacy payload layout in favor of always using the Substrate one

Use a common layout even when substrate is not enabled to reduce the number of codepaths in the path tracing code.

Implement duplicate any-hit shading avoidance in GPULightmass

Initialize SheenLTC texture unconditionally so we can use the substrate lobe even when substrate isn't enabled.

Refactor how transmittance color is encoded for thin glass and how bottom normal is encoded for clearcoat to simplify the logic

#rnx

#Synced-CL 49563263

diff --git a/Engine/Shaders/Private/PathTracing/Material/PathTracingClearCoat.ush b/Engine/Shaders/Private/PathTracing/Material/PathTracingClearCoat.ush
deleted file mode 100644
diff --git a/Engine/Shaders/Private/PathTracing/Material/PathTracingCloth.ush b/Engine/Shaders/Private/PathTracing/Material/PathTracingCloth.ush
deleted file mode 100644
diff --git a/Engine/Shaders/Private/PathTracing/Material/PathTracingDefaultLit.ush b/Engine/Shaders/Private/PathTracing/Material/PathTracingDefaultLit.ush
deleted file mode 100644
diff --git a/Engine/Shaders/Private/PathTracing/Material/PathTracingMaterialSampling.ush b/Engine/Shaders/Private/PathTracing/Material/PathTracingMaterialSampling.ush
--- a/Engine/Shaders/Private/PathTracing/Material/PathTracingMaterialSampling.ush
+++ b/Engine/Shaders/Private/PathTracing/Material/PathTracingMaterialSampling.ush
@@ -11,20 +11,10 @@ PathTracingMaterialSampling.usf: Material sampling functions
 #include "PathTracingLambert.ush"
 #if !SIMPLIFIED_MATERIAL_SHADER
 #include "PathTracingGlossy.ush"
-#if SUBSTRATE_ENABLED
 #include "PathTracingSubstrate.ush"
 #include "PathTracingSubstrateFoliage.ush"
 #include "PathTracingSubstrateSolidGlass.ush"
 #include "PathTracingSubstrateThinGlass.ush"
-#else
-#include "PathTracingSolidGlass.ush"
-#include "PathTracingThinGlass.ush"
-#include "PathTracingDefaultLit.ush"
-#include "PathTracingClearCoat.ush"
-#include "PathTracingTwoSidedFoliage.ush"
-#include "PathTracingCloth.ush"
-#include "PathTracingSubsurfaceProfile.ush"
-#endif // SUBSTRATE_ENABLED
 #include "PathTracingEye.ush"
 #include "PathTracingHair.ush"
 #include "PathTracingMedium.ush"
@@ -40,7 +30,7 @@ float3 EstimateMaterialAlbedo(FPathTracingPayload Payload)
 		case SHADINGMODELID_TWOSIDED_FOLIAGE:
 			Albedo = Payload.BaseColor + Payload.SubsurfaceColor;
 			break;
-#elif PATHTRACING_SUBSTRATE_PAYLOAD
+#else
 		case SHADINGMODELID_SUBSTRATE:
 		{
 			Albedo = Payload.DiffuseColor + Payload.SubsurfaceColor + Payload.SpecularColor + Payload.FuzzAmount * Payload.FuzzColor;
@@ -78,37 +68,6 @@ float3 EstimateMaterialAlbedo(FPathTracingPayload Payload)
 			break;
 		case SHADINGMODELID_HAIR:
 		// SUBSTRATE_TODO: Add more cases here?
-#else
-		case SHADINGMODELID_TWOSIDED_FOLIAGE:
-		case SHADINGMODELID_SUBSURFACE:
-		case SHADINGMODELID_PREINTEGRATED_SKIN:
-		case SHADINGMODELID_SUBSURFACE_PROFILE:
-		case SHADINGMODELID_EYE:
-			Albedo = Payload.DiffuseColor + Payload.SpecularColor + Payload.SubsurfaceColor;
-			break;
-		case SHADINGMODELID_DEFAULT_LIT:
-			Albedo = Payload.DiffuseColor + Payload.SpecularColor;
-			break;
-		case SHADINGMODELID_CLEAR_COAT:
-			Albedo = Payload.DiffuseColor + Payload.SpecularColor + Payload.GetClearCoat() * CLEAR_COAT_F0;
-			break;
-		case SHADINGMODELID_THIN_TRANSLUCENT:
-		{
-			// glass portion is a mixture of pure specular and tinted transmission (when refracting)
-			// or glass portion is just the specular highlight with transmission accounted for by transparency
-			const float F0 = F0RGBToF0(Payload.SpecularColor);
-			Albedo = Payload.HasRefraction() ? lerp(Payload.GetTransmittanceColor(), 1.0, F0) : F0;
-			break;
-		}
-		case SHADINGMODELID_SOLID_GLASS:
-		{
-			Albedo = 1.0; // all energy is either reflected or transmitted
-			break;
-		}
-		case SHADINGMODELID_CLOTH:
-			Albedo = (Payload.DiffuseColor + lerp(Payload.SpecularColor, Payload.GetClothColor(), Payload.GetClothAmount()));
-			break;
-		case SHADINGMODELID_HAIR:
 #endif
 		default:
 			Albedo = Payload.GetBaseColor();
@@ -132,13 +91,10 @@ void AccumulateAlbedoNormal(inout FPathTracingPayload Payload, float3 PathThroug
 		return;
 	}
 	PathThroughput *= AOVWeight;
-#if PATHTRACING_SUBSTRATE_PAYLOAD
 	PathThroughput *= Payload.WeightV * Payload.BSDFOpacity * lerp(1.0, Payload.TransmittanceN, Payload.CoverageAboveAlongN);
-#endif
 	float NormThroughput = Luminance(PathThroughput);
 	// figure out how much the specular lobes should contribute to the total (relative to the current PathRoughness)
 	// if we are above the MinRoughness, use the PathRoughness as are lower bound instead so that multiple bounces don't accumulate too much albedo
-#if PATHTRACING_SUBSTRATE_PAYLOAD
 	switch (Payload.ShadingModelID)
 	{
 		case SHADINGMODELID_SUBSTRATE:
@@ -209,51 +165,6 @@ void AccumulateAlbedoNormal(inout FPathTracingPayload Payload, float3 PathThroug
 			break;
 		}
 	}
-#else
-
-	const float RoughnessWeight = ROUGHNESS_WEIGHT(Payload.Roughness);
-	const float BrdfWeight = Payload.BSDFOpacity;
-
-	switch (Payload.ShadingModelID)
-	{
-	case SHADINGMODELID_TWOSIDED_FOLIAGE:
-	case SHADINGMODELID_SUBSURFACE:
-	case SHADINGMODELID_PREINTEGRATED_SKIN:
-	case SHADINGMODELID_SUBSURFACE_PROFILE:
-	case SHADINGMODELID_EYE:
-		Albedo += BrdfWeight * PathThroughput * (Payload.DiffuseColor + RoughnessWeight * Payload.SpecularColor + Payload.SubsurfaceColor);
-		Normal += BrdfWeight * NormThroughput * Payload.WorldNormal;
-		break;
-	case SHADINGMODELID_DEFAULT_LIT:
-	case SHADINGMODELID_CLEAR_COAT: // TODO: can we be more precise with clearcoat?
-		Albedo += BrdfWeight * PathThroughput * (Payload.DiffuseColor + RoughnessWeight * Payload.SpecularColor);
-		Normal += BrdfWeight * NormThroughput * Payload.WorldNormal;
-		break;
-	case SHADINGMODELID_CLOTH:
-		Albedo += BrdfWeight * PathThroughput * (Payload.DiffuseColor + lerp(RoughnessWeight * Payload.SpecularColor, Payload.GetClothColor(), Payload.GetClothAmount()));
-		Normal += BrdfWeight * NormThroughput * Payload.WorldNormal;
-		break;
-	case SHADINGMODELID_THIN_TRANSLUCENT:
-	{
-		const float F0 = F0RGBToF0(Payload.SpecularColor);
-		const float3 GlassAlbedo = Payload.HasRefraction() ? lerp(Payload.GetTransmittanceColor(), 1.0, F0) : F0;
-		Albedo += BrdfWeight * RoughnessWeight * PathThroughput * GlassAlbedo;
-		Normal += BrdfWeight * RoughnessWeight * NormThroughput * Payload.WorldNormal;
-		break;
-	}
-	case SHADINGMODELID_SOLID_GLASS:
-	{
-		const float GlassAlbedo = 1.0; // solid glass has unit albedo
-		Albedo += BrdfWeight * RoughnessWeight * PathThroughput * GlassAlbedo;
-		Normal += BrdfWeight * RoughnessWeight * NormThroughput * Payload.WorldNormal;
-		break;
-	}
-	default:
-		Albedo += BrdfWeight * PathThroughput * Payload.BaseColor;
-		Normal += BrdfWeight * NormThroughput * Payload.WorldNormal;
-		break;
-	}
-#endif // PATHTRACING_SUBSTRATE_PAYLOAD
 }
 #undef ROUGHNESS_WEIGHT
 
@@ -271,7 +182,7 @@ float GetAverageRoughness(FPathTracingPayload Payload)
 {
 #if SIMPLIFIED_MATERIAL_SHADER
 	return 1.0;
-#elif PATHTRACING_SUBSTRATE_PAYLOAD
+#else
 	float AvgRoughness = Payload.RoughnessData.x;
 	switch (Payload.ShadingModelID)
 	{
@@ -314,27 +225,6 @@ float GetAverageRoughness(FPathTracingPayload Payload)
 		// SUBSTRATE_TODO: do we need special cases for the other slab types?
 	}
 	return AvgRoughness;
-#else
-	float Roughness = Payload.Roughness;
-	switch (Payload.ShadingModelID)
-	{
-		case SHADINGMODELID_CLEAR_COAT:
-			Roughness = max(Roughness, Payload.GetClearCoatRoughness());
-			break;
-		case SHADINGMODELID_SUBSURFACE_PROFILE:
-		{
-			float3 DualRoughnessData = Payload.GetDualRoughnessSpecular();
-			Roughness = max(DualRoughnessData.x, DualRoughnessData.y);
-			break;
-		}
-		case SHADINGMODELID_THIN_TRANSLUCENT:
-		case SHADINGMODELID_SOLID_GLASS:
-			return Roughness; // both cases are purely specular, no diffuse
-		case SHADINGMODELID_MEDIUM:
-			return 1;
-	}
-	float AvgRoughness = lerp(1.0 /* diffuse lobe */, Roughness, Payload.Metallic);
-	return AvgRoughness;
 #endif
 }
 
@@ -358,7 +248,6 @@ void AdjustPayloadAfterUnpack(inout FPathTracingPayload HitPayload, bool bApplyO
 		HitPayload.SpecularColor = HitPayload.SpecularColor * View.SpecularOverrideParameter.w + View.SpecularOverrideParameter.xyz;
 	}
 
-#if PATHTRACING_SUBSTRATE_PAYLOAD
 	switch (HitPayload.ShadingModelID)
 	{
 		case SHADINGMODELID_SUBSTRATE:
@@ -398,83 +287,6 @@ void AdjustPayloadAfterUnpack(inout FPathTracingPayload HitPayload, bool bApplyO
 			break;
 		}
 	}
-#else
-	switch (HitPayload.ShadingModelID)
-	{
-		case SHADINGMODELID_UNLIT:
-		{
-			if (bApplyOverrides)
-			{
-				HitPayload.BaseColor  = HitPayload.BaseColor * View.DiffuseOverrideParameter.w + View.DiffuseOverrideParameter.xyz;
-			}
-			break;
-		}
-		case SHADINGMODELID_SUBSURFACE:
-		case SHADINGMODELID_PREINTEGRATED_SKIN:
-		{
-			HitPayload.SubsurfaceRadius = HitPayload.GetSubsurfaceRadius();
-			// fallthrough to case below to avoid over-bright SSS
-		}
-		case SHADINGMODELID_TWOSIDED_FOLIAGE:
-		{
-			HitPayload.SubsurfaceColor = HitPayload.GetSubsurfaceColor();
-			if (bApplyOverrides)
-			{
-				HitPayload.SubsurfaceColor = HitPayload.SubsurfaceColor * View.DiffuseOverrideParameter.w + View.DiffuseOverrideParameter.xyz;
-			}
-			AdjustPathTracingTwoSidedColorBalance(HitPayload.DiffuseColor, HitPayload.SubsurfaceColor);
-			if (HitPayload.ShadingModelID != SHADINGMODELID_TWOSIDED_FOLIAGE)
-			{
-				AdjustDiffuseSSSContribution(HitPayload.DiffuseColor, HitPayload.SubsurfaceColor, HitPayload.SubsurfaceRadius);
-			}
-			break;
-		}
-		case SHADINGMODELID_SUBSURFACE_PROFILE:
-		{
-			HitPayload.SubsurfaceRadius = HitPayload.GetSubsurfaceRadius();
-			// NOTE: if radius is 0 the sss color will get flipped back into diffuse
-			HitPayload.SubsurfaceColor = HitPayload.DiffuseColor;
-			HitPayload.DiffuseColor = 0;
-			AdjustDiffuseSSSContribution(HitPayload.DiffuseColor, HitPayload.SubsurfaceColor, HitPayload.SubsurfaceRadius);
-			break;
-		}
-		case SHADINGMODELID_EYE:
-		{
-			HitPayload.SubsurfaceRadius = HitPayload.GetSubsurfaceRadius();
-			// eye still needs a diffuse component for the iris center
-			float IrisMask = HitPayload.GetEyeIrisMask();
-			HitPayload.SubsurfaceColor = HitPayload.DiffuseColor - IrisMask * HitPayload.DiffuseColor;
-			HitPayload.DiffuseColor *= IrisMask;
-			AdjustDiffuseSSSContribution(HitPayload.DiffuseColor, HitPayload.SubsurfaceColor, HitPayload.SubsurfaceRadius);
-			break;
-		}
-		case SHADINGMODELID_HAIR:
-		{
-			if (bApplyOverrides)
-			{
-				HitPayload.BaseColor = HitPayload.BaseColor  * View.SpecularOverrideParameter.w  + View.SpecularOverrideParameter.xyz;
-				HitPayload.Specular = HitPayload.Specular * View.SpecularOverrideParameter.w + View.SpecularOverrideParameter.x;
-			}
-			break;
-		}
-		case SHADINGMODELID_CLEAR_COAT:
-		{
-			if (bApplyOverrides)
-			{
-				HitPayload.SetClearCoat(HitPayload.GetClearCoat() * View.SpecularOverrideParameter.w + View.SpecularOverrideParameter.x);
-			}
-			break;
-		}
-		case SHADINGMODELID_CLOTH:
-		{
-			if (bApplyOverrides)
-			{
-				HitPayload.SetClothColor(HitPayload.GetClothColor() * View.SpecularOverrideParameter.w + View.SpecularOverrideParameter.xyz);
-			}
-			break;
-		}
-	}
-#endif // PATHTRACING_SUBSTRATE_PAYLOAD
 }
 
 struct FSSSRandomWalkInfo
@@ -490,7 +302,6 @@ struct FSSSRandomWalkInfo
 FSSSRandomWalkInfo GetMaterialSSSInfo(FPathTracingPayload Payload, float3 V_World)
 {
 	FSSSRandomWalkInfo Result = (FSSSRandomWalkInfo)0;
-#if PATHTRACING_SUBSTRATE_PAYLOAD
 	if (any(Payload.MeanFreePath > 0.0))
 	{
 		Result.Color = Payload.SubsurfaceColor;
@@ -503,13 +314,6 @@ FSSSRandomWalkInfo GetMaterialSSSInfo(FPathTracingPayload Payload, float3 V_Worl
 			Result.Weight = Payload.WeightV * Payload.BSDFOpacity * Data.DiffuseWeight;
 			Result.Prob = Data.LobePdf.x;
 		}
-		else if (Payload.ShadingModelID == SHADINGMODELID_CLEAR_COAT)
-		{
-			FSubstrateClearCoatData Data = (FSubstrateClearCoatData)0;
-			Data.PrepareBSDF(V_World, Payload, true);
-			Result.Weight = Payload.WeightV * Payload.BSDFOpacity * Data.DiffuseWeight;
-			Result.Prob = Data.LobePdf.x;
-		}
 		else if (Payload.ShadingModelID == SHADINGMODELID_EYE)
 		{
 			const float NoV = saturate(dot(Payload.WorldNormal, V_World));
@@ -522,37 +326,6 @@ FSSSRandomWalkInfo GetMaterialSSSInfo(FPathTracingPayload Payload, float3 V_Worl
 			Result.Weight = Payload.BSDFOpacity * DiffWeight;
 		}
 	}
-#else
-	Result.Color = Payload.SubsurfaceColor;
-	Result.Radius = Payload.SubsurfaceRadius;
-	Result.Weight = 0;
-	Result.Prob = 0;
-	Result.G = Payload.GetSubsurfacePhaseFunction();
-	if (Payload.IsSubsurfaceMaterial())
-	{
-		const float NoV = saturate(dot(Payload.WorldNormal, V_World));
-		float3 SpecE = 0;
-		if (Payload.ShadingModelID == SHADINGMODELID_SUBSURFACE_PROFILE)
-		{
-			// dual spec
-			const float3 DualRoughnessData = Payload.GetDualRoughnessSpecular();
-			const float3 SpecE0 = ComputeGGXSpecEnergyTermsRGB(DualRoughnessData.x, NoV, Payload.SpecularColor).E;
-			const float3 SpecE1 = ComputeGGXSpecEnergyTermsRGB(DualRoughnessData.y, NoV, Payload.SpecularColor).E;
-			SpecE = lerp(SpecE0, SpecE1, DualRoughnessData.z);
-		}
-		else
-		{
-			// single spec
-			SpecE = ComputeGGXSpecEnergyTermsRGB(Payload.Roughness, NoV, Payload.SpecularColor).E;
-		}
-
-		const float3 SSSLobeAlbedo = (1 - SpecE) * Payload.SubsurfaceColor;
-		const float3 DiffLobeAlbedo = (1 - SpecE) * Payload.DiffuseColor;
-		const float3 SpecLobeAlbedo = SpecE;
-		Result.Prob = LobeSelectionProb(SSSLobeAlbedo, DiffLobeAlbedo + SpecLobeAlbedo);
-		Result.Weight = Payload.BSDFOpacity * (1 - SpecE);
-	}
-#endif // PATHTRACING_SUBSTRATE_PAYLOAD
 	return Result;
 }
 
@@ -576,13 +349,11 @@ FMaterialSample SampleMaterial(
 {
 #if SIMPLIFIED_MATERIAL_SHADER
 	return Simplified_SampleMaterial(Payload, RandSample);
-	//return Lambert_SampleMaterial(Payload, RandSample);
 #else
 	switch (Payload.ShadingModelID)
 	{
 		case SHADINGMODELID_UNLIT:
 			return Volumetric_SampleMaterial(V_World, Payload, RandSample);
-#if PATHTRACING_SUBSTRATE_PAYLOAD
 		case SHADINGMODELID_SUBSTRATE:
 			return Substrate_SampleMaterial(V_World, Payload, RandSample);
 		case SHADINGMODELID_SOLID_GLASS:
@@ -593,24 +364,6 @@ FMaterialSample SampleMaterial(
 			return SubstrateThinGlass_SampleMaterial(V_World, Payload, RandSample);
 		case SHADINGMODELID_CLEAR_COAT:
 			return SubstrateClearCoat_SampleMaterial(V_World, Payload, RandSample);
-#else
-		case SHADINGMODELID_DEFAULT_LIT:
-		case SHADINGMODELID_SUBSURFACE:
-		case SHADINGMODELID_PREINTEGRATED_SKIN:
-			return DefaultLit_SampleMaterial(V_World, Payload, RandSample);
-		case SHADINGMODELID_THIN_TRANSLUCENT:
-			return RoughThinGlass_SampleMaterial(V_World, Payload, RandSample);
-		case SHADINGMODELID_SOLID_GLASS:
-			return RoughGlass_SampleMaterial(V_World, Payload, RandSample);
-		case SHADINGMODELID_SUBSURFACE_PROFILE:
-			return SubsurfaceProfile_SampleMaterial(V_World, Payload, RandSample);
-		case SHADINGMODELID_CLOTH:
-			return Cloth_SampleMaterial(V_World, Payload, RandSample);
-		case SHADINGMODELID_CLEAR_COAT:
-			return ClearCoat_SampleMaterial(V_World, Payload, RandSample);
-		case SHADINGMODELID_TWOSIDED_FOLIAGE:
-			return TwoSidedFoliage_SampleMaterial(V_World, Payload, RandSample);
-#endif // PATHTRACING_SUBSTRATE_PAYLOAD
 		case SHADINGMODELID_EYE:
 			return Eye_SampleMaterial(V_World, Payload, RandSample);
 		case SHADINGMODELID_HAIR:
@@ -633,13 +386,11 @@ FMaterialEval EvalMaterial(
 {
 #if SIMPLIFIED_MATERIAL_SHADER
 	return Simplified_EvalMaterial(L_World, Payload);
-	//return Lambert_EvalMaterial(L_World, Payload);
 #else
 	switch (Payload.ShadingModelID)
 	{
 		case SHADINGMODELID_UNLIT:
 			return Volumetric_EvalMaterial(V_World, L_World, Payload, DiffuseSpecularScale);
-#if PATHTRACING_SUBSTRATE_PAYLOAD
 		case SHADINGMODELID_SUBSTRATE:
 			return Substrate_EvalMaterial(V_World, L_World, Payload, DiffuseSpecularScale);
 		case SHADINGMODELID_SOLID_GLASS:
@@ -650,24 +401,6 @@ FMaterialEval EvalMaterial(
 			return SubstrateThinGlass_EvalMaterial(V_World, L_World, Payload, DiffuseSpecularScale);
 		case SHADINGMODELID_CLEAR_COAT:
 			return SubstrateClearCoat_EvalMaterial(V_World, L_World, Payload, DiffuseSpecularScale);
-#else
-		case SHADINGMODELID_DEFAULT_LIT:
-		case SHADINGMODELID_SUBSURFACE:
-		case SHADINGMODELID_PREINTEGRATED_SKIN:
-			return DefaultLit_EvalMaterial(V_World, L_World, Payload, DiffuseSpecularScale);
-		case SHADINGMODELID_THIN_TRANSLUCENT:
-			return RoughThinGlass_EvalMaterial(V_World, L_World, Payload, DiffuseSpecularScale);
-		case SHADINGMODELID_SOLID_GLASS:
-			return RoughGlass_EvalMaterial(V_World, L_World, Payload, DiffuseSpecularScale);
-		case SHADINGMODELID_SUBSURFACE_PROFILE:
-			return SubsurfaceProfile_EvalMaterial(V_World, L_World, Payload, DiffuseSpecularScale);
-		case SHADINGMODELID_CLOTH:
-			return Cloth_EvalMaterial(V_World, L_World, Payload, DiffuseSpecularScale);
-		case SHADINGMODELID_CLEAR_COAT:
-			return ClearCoat_EvalMaterial(V_World, L_World, Payload, DiffuseSpecularScale);
-		case SHADINGMODELID_TWOSIDED_FOLIAGE:
-			return TwoSidedFoliage_EvalMaterial(V_World, L_World, Payload, DiffuseSpecularScale);
-#endif // PATHTRACING_SUBSTRATE_PAYLOAD
 		case SHADINGMODELID_EYE:
 			return Eye_EvalMaterial(V_World, L_World, Payload, DiffuseSpecularScale);
 		case SHADINGMODELID_HAIR:
diff --git a/Engine/Shaders/Private/PathTracing/Material/PathTracingSolidGlass.ush b/Engine/Shaders/Private/PathTracing/Material/PathTracingSolidGlass.ush
deleted file mode 100644
diff --git a/Engine/Shaders/Private/PathTracing/Material/PathTracingSubstrate.ush b/Engine/Shaders/Private/PathTracing/Material/PathTracingSubstrate.ush
--- a/Engine/Shaders/Private/PathTracing/Material/PathTracingSubstrate.ush
+++ b/Engine/Shaders/Private/PathTracing/Material/PathTracingSubstrate.ush
@@ -2,10 +2,7 @@
 
 #pragma once
 
-#ifndef SUBSTRATE_ENABLED
-#error "This header should only be included when Substrate is enabled."
-#endif
-
+#include "PathTracingEnergyConservation.ush"
 #include "PathTracingSubstrateCommon.ush"
 #include "PathTracingMaterialCommon.ush"
 #include "PathTracingGlossy.ush"
@@ -271,18 +268,15 @@ struct FSubstrateClearCoatData
 
 	void PrepareBSDF(float3 V_World, FPathTracingPayload Payload, bool bUseSSSColor = false)
 	{
-		float3 BaseWorldNormal = Payload.WorldNormal;
 		float3 CoatWorldNormal = Payload.WorldNormal;
+		CoatWorldNormal = AdjustShadingNormal(CoatWorldNormal, Payload.WorldGeoNormal, V_World);
 #if CLEAR_COAT_BOTTOM_NORMAL
-		if ((Payload.BottomNormalOct16bits & 0xFFFF) != 0)
-		{
-			const float2 BottomNormalOct = UnpackBottomNormalOct(Payload.BottomNormalOct16bits);
-			BaseWorldNormal = SubstrateGetSimpleCoatBottomNormal(BottomNormalOct, Payload.WorldNormal);
-		}
-#endif
-
+		float3 BaseWorldNormal = Payload.GetBottomCoatNormal();
 		BaseWorldNormal = AdjustShadingNormal(BaseWorldNormal, Payload.WorldGeoNormal, V_World);
-		CoatWorldNormal = AdjustShadingNormal(CoatWorldNormal, Payload.WorldGeoNormal, V_World);
+#else
+		float3 BaseWorldNormal = CoatWorldNormal;
+#endif
+		
 
 		BaseBasis = GetGGXBasis(Payload.RoughnessData.x, Payload.Anisotropy, BaseWorldNormal, Payload.WorldTangent, BaseAlpha);
 		BaseAlpha = GetGGXAlpha(Payload.RoughnessData.x, Payload.Anisotropy);
diff --git a/Engine/Shaders/Private/PathTracing/Material/PathTracingSubstrateCommon.ush b/Engine/Shaders/Private/PathTracing/Material/PathTracingSubstrateCommon.ush
--- a/Engine/Shaders/Private/PathTracing/Material/PathTracingSubstrateCommon.ush
+++ b/Engine/Shaders/Private/PathTracing/Material/PathTracingSubstrateCommon.ush
@@ -2,17 +2,11 @@
 
 #pragma once
 
-#ifndef SUBSTRATE_ENABLED
-#error "This header should only be included when Substrate is enabled."
-#endif
-
 #include "../PathTracingCommon.ush"
 #include "../../BRDF.ush"
 
 #define SUBSTRATE_FRESNEL_F82  0
 
-#define SUBSTRATE_PATH_TRACING_GLINTS (SUBSTRATE_COMPLEXSPECIALPATH && PLATFORM_ENABLES_SUBSTRATE_GLINTS)
-
 #if SUBSTRATE_PATH_TRACING_GLINTS
 #include "/Engine/Private/Substrate/Glint/GlintThirdParty.ush"
 #endif
@@ -28,11 +22,13 @@ float3 SubstrateLobeWeight(FPathTracingPayload Payload, float SatNoL)
 float3 SubstrateSpecularTint(FPathTracingPayload Payload, float NoV, float NoL, float VoH, float NoH)
 {
 	float3 Out = 1.f;
+#if SUBSTRATE_PATH_TRACING_SPECULAR_PROFILE
 	BRANCH
 	if (Payload.SpecularProfileId != 0)
 	{
 		Out = EvaluateSpecularProfile(Payload.SpecularProfileId, NoV, NoL, VoH, NoH);
 	}
+#endif
 	return Out;
 }
 
diff --git a/Engine/Shaders/Private/PathTracing/Material/PathTracingSubstrateFoliage.ush b/Engine/Shaders/Private/PathTracing/Material/PathTracingSubstrateFoliage.ush
--- a/Engine/Shaders/Private/PathTracing/Material/PathTracingSubstrateFoliage.ush
+++ b/Engine/Shaders/Private/PathTracing/Material/PathTracingSubstrateFoliage.ush
@@ -2,10 +2,7 @@
 
 #pragma once
 
-#ifndef SUBSTRATE_ENABLED
-#error "This header should only be included when Substrate is enabled."
-#endif
-
+#include "PathTracingEnergyConservation.ush"
 #include "PathTracingSubstrateCommon.ush"
 #include "PathTracingMaterialCommon.ush"
 #include "PathTracingGlossy.ush"
diff --git a/Engine/Shaders/Private/PathTracing/Material/PathTracingSubstrateSolidGlass.ush b/Engine/Shaders/Private/PathTracing/Material/PathTracingSubstrateSolidGlass.ush
--- a/Engine/Shaders/Private/PathTracing/Material/PathTracingSubstrateSolidGlass.ush
+++ b/Engine/Shaders/Private/PathTracing/Material/PathTracingSubstrateSolidGlass.ush
@@ -2,10 +2,7 @@
 
 #pragma once
 
-#ifndef SUBSTRATE_ENABLED
-#error "This header should only be included when Substrate is enabled."
-#endif
-
+#include "PathTracingEnergyConservation.ush"
 #include "PathTracingSubstrateCommon.ush"
 #include "PathTracingMaterialCommon.ush"
 #include "PathTracingGlossy.ush"
diff --git a/Engine/Shaders/Private/PathTracing/Material/PathTracingSubstrateThinGlass.ush b/Engine/Shaders/Private/PathTracing/Material/PathTracingSubstrateThinGlass.ush
--- a/Engine/Shaders/Private/PathTracing/Material/PathTracingSubstrateThinGlass.ush
+++ b/Engine/Shaders/Private/PathTracing/Material/PathTracingSubstrateThinGlass.ush
@@ -2,15 +2,13 @@
 
 #pragma once
 
-#ifndef SUBSTRATE_ENABLED
-#error "This header should only be included when Substrate is enabled."
-#endif
-
+#include "PathTracingEnergyConservation.ush"
 #include "PathTracingSubstrateCommon.ush"
 #include "PathTracingMaterialCommon.ush"
 #include "PathTracingGlossy.ush"
 #include "PathTracingFresnel.ush"
 #include "../../BRDF.ush"
+#include "../../ParticipatingMediaCommon.ush"
 
 
 struct FSubstrateThinGlassSlabData
diff --git a/Engine/Shaders/Private/PathTracing/Material/PathTracingSubsurfaceProfile.ush b/Engine/Shaders/Private/PathTracing/Material/PathTracingSubsurfaceProfile.ush
deleted file mode 100644
diff --git a/Engine/Shaders/Private/PathTracing/Material/PathTracingThinGlass.ush b/Engine/Shaders/Private/PathTracing/Material/PathTracingThinGlass.ush
deleted file mode 100644
diff --git a/Engine/Shaders/Private/PathTracing/Material/PathTracingTwoSidedFoliage.ush b/Engine/Shaders/Private/PathTracing/Material/PathTracingTwoSidedFoliage.ush
deleted file mode 100644
diff --git a/Engine/Shaders/Private/PathTracing/PathTracingCommon.ush b/Engine/Shaders/Private/PathTracing/PathTracingCommon.ush
--- a/Engine/Shaders/Private/PathTracing/PathTracingCommon.ush
+++ b/Engine/Shaders/Private/PathTracing/PathTracingCommon.ush
@@ -50,6 +50,9 @@ HLSL_STATIC_ASSERT(SHADINGMODELID_SOLID_GLASS <= SHADINGMODELID_MASK, "");
 
 #define PATH_TRACING_ABSORPTION_SCALE  0.01        // inverse distance at which we hit the transmission color (Beer's law)
 
+#define SUBSTRATE_PATH_TRACING_SPECULAR_PROFILE	(SUBSTRATE_ENABLED && PLATFORM_ENABLES_SUBSTRATE_SPECULAR_PROFILE)
+#define SUBSTRATE_PATH_TRACING_GLINTS 			(SUBSTRATE_ENABLED && PLATFORM_ENABLES_SUBSTRATE_GLINTS)
+
 
 // Given a ray with origin O in direction D, assuming we have found a valid hit at THit, compute a new TMin value that we can use to retrace
 // the same ray and find more hits. This is used for transparency, mesh decals and SSS tracing.
@@ -194,18 +197,17 @@ float4 PayloadDecodeRoughnessAniso(uint rgba)
 	return float4(r, g, b, a);
 }
 
-uint PayloadEncodeFuzzAndBottomNormal(float FuzzAmount, float FuzzRoughness, uint BottomNormalOct16bits)
+uint PayloadEncodeFuzz(float FuzzAmount, float FuzzRoughness)
 {
-	uint r = uint(floor(saturate(FuzzAmount) * 255.0 + 0.5));
-	uint g = uint(floor(saturate(FuzzRoughness) * 255.0 + 0.5));
-	return (r << 24) | (g << 16) | (BottomNormalOct16bits & 0xFFFF);
+	// TODO: revisit this as we could pack more data here
+	return PackFloat2ToUInt(FuzzAmount, FuzzRoughness);
 }
 
-void PayloadDecodeFuzzAndBottomNormal(uint rgba, inout float FuzzAmount, inout float FuzzRoughness, inout uint BottomNormalOct16bits)
+void PayloadDecodeFuzz(uint Bits, inout float FuzzAmount, inout float FuzzRoughness)
 {
-	FuzzAmount = float((rgba >> 24)) * (1.0 / 255.0);
-	FuzzRoughness = float((rgba >> 16) & 255) * (1.0 / 255.0);
-	BottomNormalOct16bits = rgba & 0xFFFF;
+	float2 Data = UnpackFloat2FromUInt(Bits);
+	FuzzAmount = Data.x;
+	FuzzRoughness = Data.y;
 }
 
 float3 AdjustShadingNormal(float3 ShadingNormal, float3 GeoNormal, float3 V_World)
@@ -230,12 +232,6 @@ float3 AdjustShadingNormal(float3 ShadingNormal, float3 GeoNormal, float3 V_Worl
 #error "Path tracing and GPULightmass payloads are mutually exclusive. They should not be both enabled at once."
 #endif
 
-#if SUBSTRATE_ENABLED && !IS_PAYLOAD_ENABLED(RT_PAYLOAD_TYPE_GPULIGHTMASS)
-#define PATHTRACING_SUBSTRATE_PAYLOAD 1
-#else
-#define PATHTRACING_SUBSTRATE_PAYLOAD 0
-#endif
-
 // This payload structure is what we transport between RGS/CHS/AHS programs
 struct FPackedPathTracingPayload : FMinimalPayload
 {
@@ -243,13 +239,9 @@ struct FPackedPathTracingPayload : FMinimalPayload
 #if IS_PAYLOAD_ENABLED(RT_PAYLOAD_TYPE_GPULIGHTMASS)
 	uint PackedData[7];                                  // 28 bytes (encoded data, depends on shading model and ray type)
 	// 32 bytes total
-#elif PATHTRACING_SUBSTRATE_PAYLOAD
-	// Need a bit more space for Substrate slab data
+#else
 	uint PackedData[18];                                 // 72 bytes
 	// 76 bytes total
-#else
-	uint PackedData[15];								 // 60 bytes (encoded data, depends on shading model and ray type)
-	// 64 bytes total
 #endif
 
 	float UnpackBSDFOpacity() { return f16tof32(PackedData[0] & 0xFFFF); }
@@ -314,68 +306,6 @@ struct FPackedPathTracingPayload : FMinimalPayload
 		PackedData[4] = asuint(PathRoughness);
 	}
 
-#if IS_PAYLOAD_ENABLED(RT_PAYLOAD_TYPE_GPULIGHTMASS)
-	float3 UnpackWorldGeoNormal()    { return PayloadDecodeUnitVector(PackedData[1]); }
-	float3 UnpackWorldNormal()       { return PayloadDecodeUnitVector(PackedData[2]); }
-	float3 UnpackRadiance()          { return PayloadDecodeHDRColor(PackedData[3]); }
-	float3 UnpackBaseColor()         { return PayloadDecodeLDRColor(PackedData[4]); }
-	float3 UnpackTransparencyColor() { return PayloadDecodeLDRColor(PackedData[5]); }
-	float3 UnpackSubsurfaceColor()   { return PayloadDecodeLDRColor(PackedData[6]); }
-
-
-#else
-
-	float3 UnpackWorldGeoNormal() 	 { return PayloadDecodeUnitVector(PackedData[1]); }
-	float3 UnpackWorldSmoothNormal() { return PayloadDecodeUnitVector(PackedData[2]); }
-	float3 UnpackWorldNormal()       { return PayloadDecodeUnitVector(PackedData[3]); }
-	float3 UnpackRadiance()          { return PayloadDecodeHDRColor(PackedData[4]); }
-
-	void PackWorldNormal(float3 N)   { PackedData[3] = PayloadEncodeUnitVector(N); }
-	void PackRadiance(float3 R)      { PackedData[4] = PayloadEncodeHDRColor(R); }
-
-	float3 UnpackWorldTangent()       { return PayloadDecodeTangentAngle(UnpackWorldNormal(), PackedData[5]); }
-	float UnpackIor()                 { return f16tof32(PackedData[5] >> 16); }
-	
-#if PATHTRACING_SUBSTRATE_PAYLOAD
-
-	float3 UnpackDiffuseColor()       { return PayloadDecodeLDRColor(PackedData[6]); }
-	float3 UnpackSpecularColor()      { return PayloadDecodeLDRColor(PackedData[7]); }
-	float3 UnpackSpecularEdgeColor()  { return PayloadDecodeLDRColor(PackedData[8]); }
-	float4 UnpackRoughnessAniso()     { return PayloadDecodeRoughnessAniso(PackedData[9]); }
-	float3 UnpackTransparencyColor()  { return PayloadDecodeLDRColor(PackedData[10]); }
-	float3 UnpackMeanFreePath()       { return float3(f16tof32(PackedData[11]),
-	                                                  f16tof32(PackedData[11] >> 16),
-	                                                  f16tof32(PackedData[12])); }
-	float UnpackPhaseG()              { return f16tof32(PackedData[12] >> 16); }
-	float3 UnpackFuzzColor()          { return PayloadDecodeLDRColor(PackedData[13]); }
-	float3 UnpackWeightV()            { return PayloadDecodeLDRColor(PackedData[15]); }
-	float3 UnpackTransmittanceN()     { return PayloadDecodeLDRColor(PackedData[16]); }
-
-	float UnpackCoverageAboveAlongN() { return abs(f16tof32(PackedData[17] >> 16)); }
-	float UnpackSpecularProfileId()   { return f16tof32(PackedData[17]); }
-
-	void PackDiffuseColor(float3 C)   { PackedData[6] = PayloadEncodeLDRColor(C); }
-	void PackSpecularColor(float3 C)  { PackedData[7] = PayloadEncodeLDRColor(C); }
-	void PackRoughnessAniso(float4 D) { PackedData[9] = PayloadEncodeRoughnessAniso(D.xyz, D.w); }
-
-#else
-	float3 UnpackBaseColor()         { return PayloadDecodeLDRColor(PackedData[6]); }
-	float3 UnpackTransparencyColor() { return PayloadDecodeLDRColor(PackedData[7]); }
-	float UnpackMetallic()           { return f16tof32(PackedData[8]); }
-	float UnpackSpecular()           { return f16tof32(PackedData[8] >> 16); }
-	float UnpackRoughness()          { return f16tof32(PackedData[9]); }
-	float UnpackAnisotropy()         { return f16tof32(PackedData[9] >> 16); }
-	float4 UnpackCustomData0()       { return float4(UnpackFloat2FromUInt(PackedData[10]),
-													 UnpackFloat2FromUInt(PackedData[11])); }
-	float4 UnpackCustomData1()       { return float4(UnpackFloat2FromUInt(PackedData[12]),
-													 UnpackFloat2FromUInt(PackedData[13])); }
-
-	void PackBaseColor(float3 C) { PackedData[6] = PayloadEncodeLDRColor(C); }
-	void PackMetallicSpecular(float M, float S) { PackedData[8] = PackFloat2ToUInt(float2(M, S)); }
-	void PackRoughnessAniso(float R, float A)   { PackedData[9] = PackFloat2ToUInt(float2(R, A)); }
-
-#endif
-
 #if IS_PAYLOAD_ENABLED(RT_PAYLOAD_TYPE_GPULIGHTMASS)
 	// NOTE: GPU Lightmass doesnt't have a shading model that needs volumetric absorption along shadow rays
 #else
@@ -418,6 +348,49 @@ struct FPackedPathTracingPayload : FMinimalPayload
 		HitT = T; // update most recent hit
 	}
 
+
+#if IS_PAYLOAD_ENABLED(RT_PAYLOAD_TYPE_GPULIGHTMASS)
+	float3 UnpackWorldGeoNormal()    { return PayloadDecodeUnitVector(PackedData[1]); }
+	float3 UnpackWorldNormal()       { return PayloadDecodeUnitVector(PackedData[2]); }
+	float3 UnpackRadiance()          { return PayloadDecodeHDRColor(PackedData[3]); }
+	float3 UnpackBaseColor()         { return PayloadDecodeLDRColor(PackedData[4]); }
+	float3 UnpackTransparencyColor() { return PayloadDecodeLDRColor(PackedData[5]); }
+	float3 UnpackSubsurfaceColor()   { return PayloadDecodeLDRColor(PackedData[6]); }
+
+
+#else
+
+	float3 UnpackWorldGeoNormal() 	 { return PayloadDecodeUnitVector(PackedData[1]); }
+	float3 UnpackWorldSmoothNormal() { return PayloadDecodeUnitVector(PackedData[2]); }
+	float3 UnpackWorldNormal()       { return PayloadDecodeUnitVector(PackedData[3]); }
+	float3 UnpackRadiance()          { return PayloadDecodeHDRColor(PackedData[4]); }
+
+	void PackWorldNormal(float3 N)   { PackedData[3] = PayloadEncodeUnitVector(N); }
+	void PackRadiance(float3 R)      { PackedData[4] = PayloadEncodeHDRColor(R); }
+
+	float3 UnpackWorldTangent()       { return PayloadDecodeTangentAngle(UnpackWorldNormal(), PackedData[5]); }
+	float UnpackIor()                 { return f16tof32(PackedData[5] >> 16); }
+	
+	float3 UnpackDiffuseColor()       { return PayloadDecodeLDRColor(PackedData[6]); }
+	float3 UnpackSpecularColor()      { return PayloadDecodeLDRColor(PackedData[7]); }
+	float3 UnpackSpecularEdgeColor()  { return PayloadDecodeLDRColor(PackedData[8]); }
+	float4 UnpackRoughnessAniso()     { return PayloadDecodeRoughnessAniso(PackedData[9]); }
+	float3 UnpackTransparencyColor()  { return PayloadDecodeLDRColor(PackedData[10]); }
+	float3 UnpackMeanFreePath()       { return float3(f16tof32(PackedData[11]),
+	                                                  f16tof32(PackedData[11] >> 16),
+	                                                  f16tof32(PackedData[12])); }
+	float UnpackPhaseG()              { return f16tof32(PackedData[12] >> 16); }
+	float3 UnpackFuzzColor()          { return PayloadDecodeLDRColor(PackedData[13]); }
+	float3 UnpackWeightV()            { return PayloadDecodeLDRColor(PackedData[15]); }
+	float3 UnpackTransmittanceN()     { return PayloadDecodeLDRColor(PackedData[16]); }
+
+	float UnpackCoverageAboveAlongN() { return abs(f16tof32(PackedData[17] >> 16)); }
+	float UnpackSpecularProfileId()   { return f16tof32(PackedData[17]); }
+
+	void PackDiffuseColor(float3 C)   { PackedData[6] = PayloadEncodeLDRColor(C); }
+	void PackSpecularColor(float3 C)  { PackedData[7] = PayloadEncodeLDRColor(C); }
+	void PackRoughnessAniso(float4 D) { PackedData[9] = PayloadEncodeRoughnessAniso(D.xyz, D.w); }
+
 	float4 GetDBufferA()
 	{
 		return float4(
@@ -696,12 +669,14 @@ struct FPathTracingPayload : FMinimalPayload
 	uint   Flags;
 	uint   PrimitiveLightingChannelMask;
 
-#if PATHTRACING_SUBSTRATE_PAYLOAD
+#if IS_PAYLOAD_ENABLED(RT_PAYLOAD_TYPE_GPULIGHTMASS)
+	float3 BaseColor;
+	float3 SubsurfaceColor;
+#else
 	float3 DiffuseColor;
 	float3 SpecularColor;
 	float3 SpecularEdgeColor;
 	float3 RoughnessData; // r0,r1,rblend
-	uint BottomNormalOct16bits;
 	float  Anisotropy;
 	float3 WorldTangent;
 	float3 MeanFreePath;
@@ -713,35 +688,21 @@ struct FPathTracingPayload : FMinimalPayload
 	float3 WeightV;
 	float3 TransmittanceN;
 	float CoverageAboveAlongN;
+#if SUBSTRATE_PATH_TRACING_SPECULAR_PROFILE
 	float SpecularProfileId;
+#endif
+#if SUBSTRATE_PATH_TRACING_GLINTS
 	float GlintValue;
 	float2 GlintUV;
 	float2 GlintUVdx;
 	float2 GlintUVdy;
-
+#endif
 	// not transported through packed payload, but used in raygen
 	float3 SubsurfaceColor;
 
-#else // PATHTRACING_SUBSTRATE_PAYLOAD
-	
-	float3 BaseColor;
-#if IS_PAYLOAD_ENABLED(RT_PAYLOAD_TYPE_GPULIGHTMASS)
-	float3 SubsurfaceColor;
-#else
-	float3 DiffuseColor;
-	float3 SpecularColor;
-	float  Metallic;
-	float  Specular;
-	float  Roughness;
-	float  Anisotropy;
-	float  Ior;
-	float4 CustomData0;
-	float4 CustomData1;
-	float3 WorldTangent;
-	float3 SubsurfaceColor;
-	float3 SubsurfaceRadius;
 #endif // IS_PAYLOAD_ENABLED(RT_PAYLOAD_TYPE_GPULIGHTMASS)
-#endif // PATHTRACING_SUBSTRATE_PAYLOAD
+
+	
 	float3 TranslatedWorldPos;
 
 	void SetFrontFace() { Flags |= PATH_TRACING_PAYLOAD_OUTPUT_FLAG_FRONT_FACE; }
@@ -770,14 +731,12 @@ struct FPathTracingPayload : FMinimalPayload
 	}
 #else // IS_PAYLOAD_ENABLED(RT_PAYLOAD_TYPE_GPULIGHTMASS)
 
-#if PATHTRACING_SUBSTRATE_PAYLOAD
 	// SHADINGMODELID_SUBSTRATE
 	float3 GetMaxLobeWeight()
 	{
 		return WeightV * lerp(1.0f, TransmittanceN, CoverageAboveAlongN); // largest value LobeWeight() could return
 	}
 
-
 	// SHADINGMODELID_LAMBERT and SHADINGMODELID_MEDIUM
 	float3 GetBaseColor() { return DiffuseColor; }
 	void SetBaseColor(float3 BaseColor) { DiffuseColor = BaseColor; }
@@ -827,104 +786,17 @@ struct FPathTracingPayload : FMinimalPayload
 	float GetGlassPhase() { return PhaseG; }
 	void SetGlassPhase(float G) { PhaseG = G; }
 
-
 	// SHADINGMODELID_THIN_TRANSLUCENT
-	float3 GetTransmittanceColor() { return exp(-GetExtinction() * SUBSTRATE_SIMPLEVOLUME_THICKNESS_CM); }
+	float3 GetTransmittanceColor() { return MeanFreePath; }
+	void SetTransmittanceColor(float3 TransmittanceColor) { MeanFreePath = TransmittanceColor; }
 
 	// SHADINGMODELID_TWOSIDED_FOLIAGE
 	float3 GetFoliageTransmissionColor() { return MeanFreePath; }
 	void SetFoliageTransmissionColor(float3 C) { MeanFreePath = C; }
 
-
-#else // PATHTRACING_SUBSTRATE_PAYLOAD
-	// Various ways to interpret CustomData (depending on ShadingModelID)
-	// NOTE: This is not always following the same conventions as GetMaterialCustomData0,1()
-
-	// SHADINGMODELID_CLOTH
-	float3 GetClothColor() { return CustomData0.xyz; }
-	void SetClothColor(float3 ClothColor) { CustomData0.xyz = ClothColor; }
-	float GetClothAmount() { return CustomData0.w; }
-	void SetClothAmount(float ClothAmount) { CustomData0.w = ClothAmount; }
-
 	// SHADINGMODELID_CLEAR_COAT
-	float GetClearCoat() { return CustomData0.x; }
-	float GetClearCoatRoughness() { return CustomData0.y; }
-	void SetClearCoat(float ClearCoat) { CustomData0.x = ClearCoat; }
-	void SetClearCoatRoughness(float ClearCoatRoughness) { CustomData0.y = ClearCoatRoughness; }
-	float3 GetClearCoatBottomNormal() { return CustomData1.xyz; }
-	void SetClearCoatBottomNormal(float3 BottomNormal) { CustomData1.xyz = BottomNormal; }
-
-	// SHADINGMODELID_SUBSURFACE, SHADINGMODELID_SUBSURFACE_PROFILE and SHADINGMODELID_EYE
-	float3 GetSubsurfaceRadius() { return CustomData0.rgb; }
-	void SetSubsurfaceRadius(float3 Radius) { CustomData0.rgb = Radius; }
-	float GetSubsurfacePhaseFunction() { return CustomData0.a; }
-	void SetSubsurfacePhaseFunction(float G) { CustomData0.a = G; }
-
-	// SHADINGMODELID_SUBSURFACE, SHADINGMODELID_TWOSIDED_FOLIAGE (note that SHADINGMODELID_SUBSURFACE_PROFILE and SHADINGMODELID_EYE don't use this as they use DiffuseColor directly)
-	float3 GetSubsurfaceColor() { return CustomData1.xyz; }
-	void SetSubsurfaceColor(float3 SubsurfaceColor) { CustomData1.xyz = SubsurfaceColor; }
-	
-	// SHADINGMODELID_SUBSURFACE_PROFILE
-	float3 GetDualRoughnessSpecular() { return CustomData1.xyz; }
-	void SetDualRoughnessSpecular(float Roughness0, float Roughness1, float LobeMix) { CustomData1.xyz = float3(Roughness0, Roughness1, LobeMix); }
-
-
-	// SHADINGMODELID_EYE
-	float GetEyeRoughness() { return Roughness; }
-	void SetEyeRoughness(float R) { Roughness = R; }
-	float3 GetEyeCausticNormal() { return OctahedronToUnitVector(CustomData1.xy); }
-	void SetEyeCausticNormal(float3 CausticNormal) { CustomData1.xy = UnitVectorToOctahedron(CausticNormal); }
-	float GetEyeIrisMask() { return Anisotropy; }
-	void SetEyeIrisMask(float IrisMask) { Anisotropy = IrisMask; }
-	float3 GetEyeIrisNormal() { return OctahedronToUnitVector(CustomData1.zw); }
-	void SetEyeIrisNormal(float3 IrisNormal) { CustomData1.zw = UnitVectorToOctahedron(IrisNormal); }
-
-	// SHADINGMODELID_HAIR
-	float GetHairLongitudinalRoughness() { return Roughness; }
-	void SetHairLongitudinalRoughness(float R) { Roughness = R; }
-	float GetHairAzimuthalRoughness() { return Metallic; }
-	void SetHairAzimuthalRoughness(float AzimuthalRoughness) { Metallic = AzimuthalRoughness; }
-	float GetHairSpecular() { return Specular; }
-	void SetHairSpecular(float S) { Specular = S; }
-	float2 GetHairPrimitiveUV() { return CustomData0.xy; }
-	void SetHairPrimitiveUV(float2 UV) { CustomData0.xy = UV; }
-
-	// SHADINGMODELID_THIN_TRANSLUCENT
-	float3 GetTransmittanceColor() { return CustomData0.xyz; }
-	void SetTransmittanceColor(float3 TransmittanceColor) { CustomData0.xyz = TransmittanceColor; }
-
-	// SHADINGMODELID_SOLID_GLASS
-	float3 GetExtinction() { return CustomData0.xyz; }
-	void SetExtinction(float3 SigmaT) { CustomData0.xyz = SigmaT; }
-	float3 GetGlassAlbedo() { return BaseColor; }
-	void SetGlassAlbedo(float3 A) { BaseColor = A; }
-	float GetGlassPhase() { return CustomData0.w; }
-	void SetGlassPhase(float G) { CustomData0.w = G; }
-
-
-	// SHADINGMODELID_LAMBERT and SHADINGMODELID_MEDIUM
-	float3 GetBaseColor() { return BaseColor; }
-	void SetBaseColor(float3 C)
-	{
-		BaseColor = C;
-	}
-
-	// SHADINGMODELID_MEDIUM
-	float3 GetHGWeight() { return CustomData0.xyz; }
-	float GetHGPhaseG() { return CustomData0.w; }
-	void SetHG(float3 Weight, float G) { CustomData0 = float4(Weight, G); }
-	float3 GetDualHGWeight() { return CustomData1.xyz; }
-	float3 GetDualHGPhaseData() { return float3(Metallic, Specular, Roughness); }
-	void SetDualHG(float3 Weight, float3 PhaseData)
-	{
-		CustomData1.xyz = Weight;
-		Metallic = PhaseData.x;
-		Specular = PhaseData.y;
-		Roughness = PhaseData.z;
-	}
-	float GetCloudFactor() { return CustomData1.w; }
-	void SetCloudFactor(float Factor) { CustomData1.w = Factor; }
-#endif // PATHTRACING_SUBSTRATE_PAYLOAD
+	float3 GetBottomCoatNormal() { return MeanFreePath; }
+	void SetBottomCoatNormal(float3 BottomNormal) { MeanFreePath = BottomNormal; }
 
 	// Methods used by the integrator
 
@@ -955,16 +827,8 @@ struct FPathTracingPayload : FMinimalPayload
 
 	bool IsSubsurfaceMaterial()
 	{
-#if PATHTRACING_SUBSTRATE_PAYLOAD
 		return ShadingModelID == SHADINGMODELID_SUBSTRATE ||
 		       ShadingModelID == SHADINGMODELID_EYE;
-#else
-		return ShadingModelID == SHADINGMODELID_SUBSURFACE ||
-			   ShadingModelID == SHADINGMODELID_PREINTEGRATED_SKIN ||
-			   ShadingModelID == SHADINGMODELID_SUBSURFACE_PROFILE ||
-			   ShadingModelID == SHADINGMODELID_EYE;
-#endif
-
 	}
 #endif // IS_PAYLOAD_ENABLED(RT_PAYLOAD_TYPE_GPULIGHTMASS)
 };
@@ -1027,7 +891,7 @@ FPackedPathTracingPayload PackPathTracingPayload(FPathTracingPayload Input)
 	Output.PackedData[4] = PayloadEncodeHDRColor(Input.Radiance);
 	Output.PackedData[5] = PayloadEncodeTangentAngle(Output.UnpackWorldNormal(), Input.WorldTangent);
 	Output.PackedData[5]|= f32tof16(Input.Ior) << 16;
-#if PATHTRACING_SUBSTRATE_PAYLOAD
+
 	Output.PackedData[6] = PayloadEncodeLDRColor(Input.DiffuseColor);
 	Output.PackedData[7] = PayloadEncodeLDRColor(Input.SpecularColor);
 	Output.PackedData[8] = PayloadEncodeLDRColor(Input.SpecularEdgeColor);
@@ -1036,11 +900,17 @@ FPackedPathTracingPayload PackPathTracingPayload(FPathTracingPayload Input)
 	Output.PackedData[11] = PackFloat2ToUInt(float2(Input.MeanFreePath.xy));
 	Output.PackedData[12] = PackFloat2ToUInt(float2(Input.MeanFreePath.z, Input.PhaseG));
 	Output.PackedData[13] = PayloadEncodeLDRColor(Input.FuzzColor);
-	Output.PackedData[14] = PayloadEncodeFuzzAndBottomNormal(Input.FuzzAmount, Input.FuzzRoughness, Input.BottomNormalOct16bits);
+	Output.PackedData[14] = PayloadEncodeFuzz(Input.FuzzAmount, Input.FuzzRoughness);
 	Output.PackedData[15] = PayloadEncodeLDRColor(Input.WeightV);
 	Output.PackedData[16] = PayloadEncodeLDRColor(Input.TransmittanceN);
-	Output.PackedData[17] = PackFloat2ToUInt(float2(Input.SpecularProfileId, abs(Input.CoverageAboveAlongN)));
+#if SUBSTRATE_PATH_TRACING_SPECULAR_PROFILE
+	float SpecularProfileId = Input.SpecularProfileId;
+#else
+	float SpecularProfileId = 0;
+#endif
+	Output.PackedData[17] = PackFloat2ToUInt(float2(SpecularProfileId, abs(Input.CoverageAboveAlongN)));
 
+#if SUBSTRATE_PATH_TRACING_GLINTS
 	// In order to not increase the data payload size, we alias SSS/Fuzz parameters when glint are enabled
 	if (Input.GlintValue < 1)
 	{
@@ -1051,19 +921,7 @@ FPackedPathTracingPayload PackPathTracingPayload(FPathTracingPayload Input)
 		Output.PackedData[14] = PackFloat2ToUInt(Input.GlintUVdx);
 		Output.PackedData[17] |= 0x80000000u; // Encode presence of glints in sign bit of CoverageAboveAlongN
 	}
-
-#else // PATHTRACING_SUBSTRATE_PAYLOAD
-
-	Output.PackedData[ 6] = PayloadEncodeLDRColor(Input.BaseColor);
-	Output.PackedData[ 7] = PayloadEncodeLDRColor(Input.TransparencyColor);
-	Output.PackedData[ 8] = PackFloat2ToUInt(float2(Input.Metallic , Input.Specular  ));
-	Output.PackedData[ 9] = PackFloat2ToUInt(float2(Input.Roughness, Input.Anisotropy));
-	Output.PackedData[10] = PackFloat2ToUInt(Input.CustomData0.xy);
-	Output.PackedData[11] = PackFloat2ToUInt(Input.CustomData0.zw);
-	Output.PackedData[12] = PackFloat2ToUInt(Input.CustomData1.xy);
-	Output.PackedData[13] = PackFloat2ToUInt(Input.CustomData1.zw);
-
-#endif // PATHTRACING_SUBSTRATE_PAYLOAD
+#endif
 
 #endif // IS_PAYLOAD_ENABLED(RT_PAYLOAD_TYPE_GPULIGHTMASS)
 	return Output;
@@ -1096,7 +954,6 @@ FPathTracingPayload UnpackPathTracingPayload(FPackedPathTracingPayload Input, FR
 	Output.WorldTangent      = Input.UnpackWorldTangent();
 	Output.Ior               = Input.UnpackIor();
 
-#if PATHTRACING_SUBSTRATE_PAYLOAD
 	Output.DiffuseColor        = Input.UnpackDiffuseColor();
 	Output.SpecularColor       = Input.UnpackSpecularColor();
 	Output.SpecularEdgeColor   = Input.UnpackSpecularEdgeColor();
@@ -1107,12 +964,15 @@ FPathTracingPayload UnpackPathTracingPayload(FPackedPathTracingPayload Input, FR
 	Output.MeanFreePath        = Input.UnpackMeanFreePath();
 	Output.PhaseG              = Input.UnpackPhaseG();
 	Output.FuzzColor           = Input.UnpackFuzzColor();
-	PayloadDecodeFuzzAndBottomNormal(Input.PackedData[14], Output.FuzzAmount, Output.FuzzRoughness, Output.BottomNormalOct16bits);
+	PayloadDecodeFuzz(Input.PackedData[14], Output.FuzzAmount, Output.FuzzRoughness);
 	Output.WeightV             = Input.UnpackWeightV();
 	Output.TransmittanceN      = Input.UnpackTransmittanceN();
 	Output.CoverageAboveAlongN = Input.UnpackCoverageAboveAlongN();
+#if SUBSTRATE_PATH_TRACING_SPECULAR_PROFILE
 	Output.SpecularProfileId   = Input.UnpackSpecularProfileId();
+#endif
 	
+#if SUBSTRATE_PATH_TRACING_GLINTS
 	// TODO: figure out how to make glint unpacking closer to functions above
 	// In order to not increase the data payload size, we alias SSS/Fuzz parameters when glint are enabled
 	const bool bHasGlints = Input.PackedData[17] >> 31u;
@@ -1135,22 +995,7 @@ FPathTracingPayload UnpackPathTracingPayload(FPackedPathTracingPayload Input, FR
 		// When glints are not used, default to full density so we get the regular specular highlight
 		Output.GlintValue = 1.0f;
 	}
-
-#else
-	// !PATHTRACING_SUBSTRATE_PAYLOAD
-	Output.BaseColor         = Input.UnpackBaseColor();
-	Output.TransparencyColor = Input.UnpackTransparencyColor();
-	Output.Metallic          = Input.UnpackMetallic();
-	Output.Specular          = Input.UnpackSpecular();
-	Output.Roughness         = Input.UnpackRoughness();
-	Output.Anisotropy        = Input.UnpackAnisotropy();
-	Output.CustomData0       = Input.UnpackCustomData0();
-	Output.CustomData1       = Input.UnpackCustomData1();
-
-	// 
-	Output.DiffuseColor  = Output.BaseColor - Output.BaseColor * Output.Metallic;
-	Output.SpecularColor = ComputeF0(Output.Specular, Output.BaseColor, Output.Metallic);
-#endif // PATHTRACING_SUBSTRATE_PAYLOAD
+#endif // SUBSTRATE_PATH_TRACING_GLINTS
 #endif // IS_PAYLOAD_ENABLED(RT_PAYLOAD_TYPE_GPULIGHTMASS)
 	return Output;
 }
diff --git a/Engine/Shaders/Private/PathTracing/PathTracingCore.ush b/Engine/Shaders/Private/PathTracing/PathTracingCore.ush
--- a/Engine/Shaders/Private/PathTracing/PathTracingCore.ush
+++ b/Engine/Shaders/Private/PathTracing/PathTracingCore.ush
@@ -368,7 +368,6 @@ void ApplyDecal(FDecalShaderPayload DecalPayload, inout FPackedPathTracingPayloa
 		return;
 	}
 
-#if PATHTRACING_SUBSTRATE_PAYLOAD
 	float4 RoughnessData = PackedPayload.UnpackRoughnessAniso();
 	float Roughness = RoughnessData.x;
 	float3 DiffuseColor = PackedPayload.UnpackDiffuseColor();
@@ -376,12 +375,6 @@ void ApplyDecal(FDecalShaderPayload DecalPayload, inout FPackedPathTracingPayloa
 	float Metallic = F0RGBToMetallic(SpecularColor);
 	float Specular = F0RGBToDielectricSpecular(SpecularColor);
 	float3 BaseColor = lerp(DiffuseColor, SpecularColor, Metallic);
-#else
-	float3 BaseColor = PackedPayload.UnpackBaseColor();
-	float Metallic = PackedPayload.UnpackMetallic();
-	float Specular = PackedPayload.UnpackSpecular();
-	float Roughness = PackedPayload.UnpackRoughness();
-#endif // PATHTRACING_SUBSTRATE_PAYLOAD
 
 	if (bHasColorResponse)
 	{
@@ -408,7 +401,6 @@ void ApplyDecal(FDecalShaderPayload DecalPayload, inout FPackedPathTracingPayloa
 		}
 	}
 
-#if PATHTRACING_SUBSTRATE_PAYLOAD
 	if (bHasColorResponse || bHasRoughnessResponse)
 	{
 		// the encode-decode is potentially lossy, so only do it if we actually need to
@@ -416,17 +408,6 @@ void ApplyDecal(FDecalShaderPayload DecalPayload, inout FPackedPathTracingPayloa
 		PackedPayload.PackSpecularColor(lerp(DielectricSpecularToF0(Specular), BaseColor, Metallic));
 		PackedPayload.PackRoughnessAniso(float4(Roughness, RoughnessData.yzw));
 	}
-#else
-	if (bHasColorResponse)
-	{
-		PackedPayload.PackBaseColor(BaseColor);
-	}
-	if (bHasRoughnessResponse)
-	{
-		PackedPayload.PackMetallicSpecular(Metallic, Specular);
-		PackedPayload.PackRoughnessAniso(Roughness, PackedPayload.UnpackAnisotropy());
-	}
-#endif
 }
 
 void EvaluateDecals(FPathState PathState, float HitT, inout FDecalShaderPayload DecalPayload)
@@ -1966,7 +1947,7 @@ bool PathTracingKernel(inout FPathState PathState, int Bounce)
 		Print(Ctx, TEXT("Pixel Coord : "), FontWhite); Print(Ctx, PathState.PixelCoord, FontYellow); Newline(Ctx);
 
 		// Example of debugging :
-		#if SUBSTRATE_ENABLED
+		#if SUBSTRATE_PATH_TRACING_GLINTS
 		PrintLineN(Ctx, Payload.GlintValue);
 		PrintLineN(Ctx, Payload.GlintUV);
 		PrintLineN(Ctx, Payload.GlintUVdx);
diff --git a/Engine/Shaders/Private/PathTracing/PathTracingDebug.usf b/Engine/Shaders/Private/PathTracing/PathTracingDebug.usf
--- a/Engine/Shaders/Private/PathTracing/PathTracingDebug.usf
+++ b/Engine/Shaders/Private/PathTracing/PathTracingDebug.usf
@@ -375,13 +375,10 @@ FPackedPathTracingPayload TraceTransparentRay(FRayDesc Ray, int ScatterType, int
 				}
 
 				// Accumulate G-Buffer like data for the denoiser
-#if PATHTRACING_SUBSTRATE_PAYLOAD
 				float3 RoughnessData = HitPayload.RoughnessData;
 				float Roughness = lerp(RoughnessData.x, RoughnessData.y, RoughnessData.z); // mix both lobes to get an average value
 				DenoiserThroughput *= HitPayload.WeightV * HitPayload.BSDFOpacity * lerp(1.0, HitPayload.TransmittanceN, HitPayload.CoverageAboveAlongN);
-#else
-				float Roughness = HitPayload.Roughness;
-#endif
+
 				float NormThroughput = Luminance(DenoiserThroughput);
 				float3 DiffuseColor = HitPayload.DiffuseColor;
 				float3 SpecularColor = HitPayload.SpecularColor;
@@ -552,11 +549,7 @@ FPathTracingResult PathTracingDebugCore(FRayDesc Ray, uint2 LaunchIndex, int Sam
 
 #else
 				// Hard-coded diffuse lobe for benchmarking
-#if SUBSTRATE_ENABLED
 				float3 MaterialEvalWeight = PackedPayload.UnpackDiffuseColor();
-#else
-				float3 MaterialEvalWeight = PackedPayload.UnpackBaseColor();
-#endif
 				float  MaterialEvalPdf = saturate(dot(WorldNormal, LightSample.Direction)) / PI;
 				float3 LightContrib = PathThroughput * LightSample.RadianceOverPdf * MaterialEvalWeight * MaterialEvalPdf;
 #endif
@@ -1000,8 +993,6 @@ FPathTracingResult PathTracingDebugCore(FRayDesc Ray, uint2 LaunchIndex)
 				case PATH_TRACER_DEBUG_VIZ_WORLD_NORMAL: 				Color = WorldNormal * 0.5 + 0.5; break;
 				case PATH_TRACER_DEBUG_VIZ_WORLD_SMOOTH_NORMAL:			Color = HitPayload.WorldSmoothNormal * 0.5 + 0.5; break;
 				case PATH_TRACER_DEBUG_VIZ_WORLD_GEO_NORMAL:			Color = HitPayload.WorldGeoNormal * 0.5 + 0.5; break;
-#if SUBSTRATE_ENABLED
-				// TODO: implement more visualizations for this case
 				case PATH_TRACER_DEBUG_VIZ_BASE_COLOR:					Color = lerp(HitPayload.DiffuseColor, HitPayload.SpecularColor, F0RGBToMetallic(HitPayload.SpecularColor)); break;
 				case PATH_TRACER_DEBUG_VIZ_DIFFUSE_COLOR:				Color = HitPayload.DiffuseColor; break;
 				case PATH_TRACER_DEBUG_VIZ_SPECULAR_COLOR:				Color = HitPayload.SpecularColor; break;
@@ -1009,17 +1000,6 @@ FPathTracingResult PathTracingDebugCore(FRayDesc Ray, uint2 LaunchIndex)
 				case PATH_TRACER_DEBUG_VIZ_SPECULAR:					Color = F0RGBToDielectricSpecular(HitPayload.SpecularColor); break;
 				case PATH_TRACER_DEBUG_VIZ_ROUGHNESS:					Color = HitPayload.RoughnessData; break;
 				case PATH_TRACER_DEBUG_VIZ_ANISOTROPY:					Color = AnisoDebugColor(HitPayload.Anisotropy); break;
-#else
-				case PATH_TRACER_DEBUG_VIZ_BASE_COLOR:					Color = HitPayload.BaseColor; break;
-				case PATH_TRACER_DEBUG_VIZ_DIFFUSE_COLOR:				Color = HitPayload.BaseColor * (1 - HitPayload.Metallic); break;
-				case PATH_TRACER_DEBUG_VIZ_SPECULAR_COLOR:				Color = lerp(0.08 * HitPayload.Specular, HitPayload.BaseColor, HitPayload.Metallic); break;
-				case PATH_TRACER_DEBUG_VIZ_METALLIC:					Color = HitPayload.Metallic; break;
-				case PATH_TRACER_DEBUG_VIZ_SPECULAR:					Color = HitPayload.Specular; break;
-				case PATH_TRACER_DEBUG_VIZ_ROUGHNESS:					Color = HitPayload.Roughness; break;
-				case PATH_TRACER_DEBUG_VIZ_ANISOTROPY:					Color = AnisoDebugColor(HitPayload.Anisotropy); break;
-				case PATH_TRACER_DEBUG_VIZ_CUSTOM_DATA0:				Color = HitPayload.CustomData0.rgb * View.PreExposure; ExtraAlpha = HitPayload.CustomData0.a; break;
-				case PATH_TRACER_DEBUG_VIZ_CUSTOM_DATA1:				Color = HitPayload.CustomData1.rgb * View.PreExposure; ExtraAlpha = HitPayload.CustomData1.a; break;
-#endif
 				case PATH_TRACER_DEBUG_VIZ_TRANSPARENCY_COLOR:			Color = HitPayload.TransparencyColor; break;
 				case PATH_TRACER_DEBUG_VIZ_IOR:							Color = HitPayload.Ior; break;
 				case PATH_TRACER_DEBUG_VIZ_SHADING_MODEL:				Color = GetShadingModelColor(HitPayload.ShadingModelID); break;
@@ -1084,7 +1064,6 @@ FPathTracingResult PathTracingDebugCore(FRayDesc Ray, uint2 LaunchIndex)
 					break;
 				}
 				case PATH_TRACER_DEBUG_VIZ_BSDF_OPACITY:    Color = HitPayload.BSDFOpacity; break;
-#if SUBSTRATE_ENABLED
 				case PATH_TRACER_DEBUG_VIZ_FUZZ_COLOR:     	Color = (HitPayload.ShadingModelID == SHADINGMODELID_SUBSTRATE) ? HitPayload.FuzzColor     : 0.0; break;
 				case PATH_TRACER_DEBUG_VIZ_FUZZ_ROUGHNESS:	Color = (HitPayload.ShadingModelID == SHADINGMODELID_SUBSTRATE) ? HitPayload.FuzzRoughness : 0.0; break;
 				case PATH_TRACER_DEBUG_VIZ_FUZZ_AMOUNT:		Color = (HitPayload.ShadingModelID == SHADINGMODELID_SUBSTRATE) ? HitPayload.FuzzAmount    : 0.0; break;
@@ -1092,12 +1071,6 @@ FPathTracingResult PathTracingDebugCore(FRayDesc Ray, uint2 LaunchIndex)
 				case PATH_TRACER_DEBUG_VIZ_SUBSTRATE_TRANSMITTANCE_N: 	Color = HitPayload.TransmittanceN; 		break;
 				case PATH_TRACER_DEBUG_VIZ_SUBSTRATE_COVERAGE_ABOVE_N: 	Color = HitPayload.CoverageAboveAlongN; break;
 				case PATH_TRACER_DEBUG_VIZ_SUBSTRATE_F90:				Color = HitPayload.SpecularEdgeColor;   break;
-
-#else
-				case PATH_TRACER_DEBUG_VIZ_FUZZ_COLOR:     	Color = (HitPayload.ShadingModelID == SHADINGMODELID_CLOTH) ? HitPayload.GetClothColor()  : 0.0; break;
-				case PATH_TRACER_DEBUG_VIZ_FUZZ_ROUGHNESS: 	Color = (HitPayload.ShadingModelID == SHADINGMODELID_CLOTH) ? HitPayload.Roughness        : 0.0; break;
-				case PATH_TRACER_DEBUG_VIZ_FUZZ_AMOUNT:    	Color = (HitPayload.ShadingModelID == SHADINGMODELID_CLOTH) ? HitPayload.GetClothAmount() : 0.0; break;
-#endif
 				case PATH_TRACER_DEBUG_VIZ_EYE_IRIS_MASK:      Color = (HitPayload.ShadingModelID == SHADINGMODELID_EYE) ?             HitPayload.GetEyeIrisMask()      : 0.0; break;
 				case PATH_TRACER_DEBUG_VIZ_EYE_IRIS_NORMAL:    Color = (HitPayload.ShadingModelID == SHADINGMODELID_EYE) ? 0.5 + 0.5 * HitPayload.GetEyeIrisNormal()    : 0.0; break;
 				case PATH_TRACER_DEBUG_VIZ_EYE_CAUSTIC_NORMAL: Color = (HitPayload.ShadingModelID == SHADINGMODELID_EYE) ? 0.5 + 0.5 * HitPayload.GetEyeCausticNormal() : 0.0; break;
diff --git a/Engine/Shaders/Private/PathTracing/PathTracingGPULightmassMaterialHitShader.usf b/Engine/Shaders/Private/PathTracing/PathTracingGPULightmassMaterialHitShader.usf
--- a/Engine/Shaders/Private/PathTracing/PathTracingGPULightmassMaterialHitShader.usf
+++ b/Engine/Shaders/Private/PathTracing/PathTracingGPULightmassMaterialHitShader.usf
@@ -342,7 +342,7 @@ RAY_TRACING_ENTRY_ANY_HIT(GPULightmassMaterialAHS,
 		// not a shadow ray -- don't need to run the AHS logic
 		return;
 	}
-	if (PackedPayload.HitT == RayTCurrent())
+	if (PackedPayload.CheckDuplicateAHS(RayTCurrent()))
 	{
 		// We just processed this exact hit, don't double-count it
 		IgnoreHit();
diff --git a/Engine/Shaders/Private/PathTracing/PathTracingMaterialHitShader.usf b/Engine/Shaders/Private/PathTracing/PathTracingMaterialHitShader.usf
--- a/Engine/Shaders/Private/PathTracing/PathTracingMaterialHitShader.usf
+++ b/Engine/Shaders/Private/PathTracing/PathTracingMaterialHitShader.usf
@@ -44,6 +44,7 @@
 #include "/Engine/Private/DeferredShadingCommon.ush"
 #include "/Engine/Private/SubsurfaceProfileCommon.ush"
 #include "/Engine/Private/BurleyNormalizedSSSCommon.ush"
+#include "/Engine/Private/ParticipatingMediaCommon.ush"
 #include "/Engine/Private/PathTracing/PathTracingCommon.ush"
 #include "/Engine/Private/PathTracing/Material/PathTracingFresnel.ush"
 
@@ -548,12 +549,11 @@ RAY_TRACING_ENTRY_CLOSEST_HIT(PathTracingMaterialCHS,
 							Payload.SpecularColor         = F0;
 							Payload.SpecularEdgeColor     = F90;
 							Payload.RoughnessData         = RoughnessData;
-							Payload.BottomNormalOct16bits = 0;
 							if (bUseClearCoat)
 							{
 								Payload.ShadingModelID = SHADINGMODELID_CLEAR_COAT;
 							#if CLEAR_COAT_BOTTOM_NORMAL
-								Payload.BottomNormalOct16bits = SLAB_HAZINESS(BSDF) >> 16;
+								Payload.SetBottomCoatNormal(SubstrateGetSimpleCoatBottomNormal(UnpackHaziness(SLAB_HAZINESS(BSDF)).BottomNormalOct, N));
 							#endif
 							}
 							Payload.Anisotropy = BSDF_GETHASANISOTROPY(BSDF) ? SLAB_ANISOTROPY(BSDF) : 0.0;
@@ -602,23 +602,27 @@ RAY_TRACING_ENTRY_CLOSEST_HIT(PathTracingMaterialCHS,
 							}
 
 							// Handle special-case specular features
-							Payload.SpecularProfileId = 0;
-							Payload.GlintValue = 1.f;
-
-							#if SUBSTRATE_COMPLEXPATH
-							if (BSDF_GETHASSPECPROFILE(BSDF))
-							{
-								Payload.SpecularProfileId = SLAB_SPECPROFILEID(BSDF);
-							}
+							#if SUBSTRATE_PATH_TRACING_SPECULAR_PROFILE
+								Payload.SpecularProfileId = 0;
+								#if SUBSTRATE_COMPLEXSPECIALPATH
+								if (BSDF_GETHASSPECPROFILE(BSDF))
+								{
+									Payload.SpecularProfileId = SLAB_SPECPROFILEID(BSDF);
+								}
+								#endif
 							#endif
-							#if SUBSTRATE_COMPLEXSPECIALPATH
-							if (BSDF_GETHASGLINT(BSDF))
-							{
-								Payload.GlintValue = SLAB_GLINT_VALUE(BSDF);
-								Payload.GlintUV    = SLAB_GLINT_UV(BSDF);
-								Payload.GlintUVdx  = SLAB_GLINT_UVDDX(BSDF);
-								Payload.GlintUVdy  = SLAB_GLINT_UVDDY(BSDF);
-							}
+
+							#if SUBSTRATE_PATH_TRACING_GLINTS
+								Payload.GlintValue = 1.f;
+								#if SUBSTRATE_COMPLEXSPECIALPATH
+								if (BSDF_GETHASGLINT(BSDF))
+								{
+									Payload.GlintValue = SLAB_GLINT_VALUE(BSDF);
+									Payload.GlintUV    = SLAB_GLINT_UV(BSDF);
+									Payload.GlintUVdx  = SLAB_GLINT_UVDDX(BSDF);
+									Payload.GlintUVdy  = SLAB_GLINT_UVDDY(BSDF);
+								}
+								#endif
 							#endif
 						}
 					#if PATH_TRACER_USE_SUBSTRATE_GLASS_LOBE
@@ -661,12 +665,21 @@ RAY_TRACING_ENTRY_CLOSEST_HIT(PathTracingMaterialCHS,
 							Payload.FuzzAmount    = FuzzAmount;
 							Payload.FuzzColor     = FuzzColor;
 							Payload.FuzzRoughness = FuzzRoughness;
+							#if MATERIAL_ISTHINSURFACE
+							// Convert to transmittance color now since we won't be using the raw extinction coeffient
+							Payload.SetTransmittanceColor(exp(-MeanFreePath * SUBSTRATE_SIMPLEVOLUME_THICKNESS_CM));
+							#else
 							Payload.MeanFreePath  = MeanFreePath;
+							#endif
 							Payload.PhaseG        = PhaseG;
 							Payload.Ior           = Ior;
 							// SUBSTRATE_TODO: should these be supported in the Glass case?
+						#if SUBSTRATE_PATH_TRACING_SPECULAR_PROFILE
 							Payload.SpecularProfileId = 0;
+						#endif
+						#if SUBSTRATE_PATH_TRACING_GLINTS
 							Payload.GlintValue = 1.f;
+						#endif
 
 							// Material is configured as a thin surface, and refraction is disabled
 							#if MATERIAL_ISTHINSURFACE && PATH_TRACER_REFRACTION_ENABLED == 0
@@ -711,8 +724,12 @@ RAY_TRACING_ENTRY_CLOSEST_HIT(PathTracingMaterialCHS,
 					Payload.SetEyeIrisMask(IrisMask);
 					Payload.SetEyeIrisNormal(normalize(EYE_IRISNORMAL(BSDF)));
 					Payload.ShadingModelID = SHADINGMODELID_EYE;
+				#if SUBSTRATE_PATH_TRACING_SPECULAR_PROFILE
 					Payload.SpecularProfileId = 0;
+				#endif
+				#if SUBSTRATE_PATH_TRACING_GLINTS
 					Payload.GlintValue = 1.f;
+				#endif
 					break;
 				}
 				case SUBSTRATE_BSDF_TYPE_HAIR:
@@ -732,8 +749,12 @@ RAY_TRACING_ENTRY_CLOSEST_HIT(PathTracingMaterialCHS,
 					Payload.SetHairPrimitiveUV(0.5);
 					#endif
 					Payload.ShadingModelID = SHADINGMODELID_HAIR;
+					#if SUBSTRATE_PATH_TRACING_SPECULAR_PROFILE
 					Payload.SpecularProfileId = 0;
+					#endif
+					#if SUBSTRATE_PATH_TRACING_GLINTS
 					Payload.GlintValue = 1.f;
+					#endif
 					break;
 				}
 				case SUBSTRATE_BSDF_TYPE_SINGLELAYERWATER:
@@ -768,8 +789,12 @@ RAY_TRACING_ENTRY_CLOSEST_HIT(PathTracingMaterialCHS,
 					Payload.FuzzAmount = 0;
 					Payload.FuzzRoughness = 0;
 					Payload.FuzzColor = 0;
+					#if SUBSTRATE_PATH_TRACING_SPECULAR_PROFILE
 					Payload.SpecularProfileId = 0;
+					#endif
+					#if SUBSTRATE_PATH_TRACING_GLINTS
 					Payload.GlintValue = 1.f;
+					#endif
 					break;
 				}
 				case SUBSTRATE_BSDF_TYPE_TOON:
@@ -796,9 +821,12 @@ RAY_TRACING_ENTRY_CLOSEST_HIT(PathTracingMaterialCHS,
 					Payload.FuzzAmount			= 0.0f;
 					Payload.FuzzColor			= 0.0f;
 					Payload.FuzzRoughness		= 0.0f;
+					#if SUBSTRATE_PATH_TRACING_SPECULAR_PROFILE
 					Payload.SpecularProfileId	= 0.0f;
+					#endif
+					#if SUBSTRATE_PATH_TRACING_GLINTS
 					Payload.GlintValue			= 1.0f;
-
+					#endif
 					break;
 				}
 			}  // switch on slab tyype
@@ -835,7 +863,6 @@ RAY_TRACING_ENTRY_CLOSEST_HIT(PathTracingMaterialCHS,
 	Ior = GetRefractionIor(PixelMaterialInputs);
 #endif
 
-#if USE_LEGACY_CODEPATH_FOR_SUBSTRATE
 	Payload.DiffuseColor = ComputeDiffuseAlbedo(BaseColor, Metallic);
 	Payload.SpecularColor = ComputeF0(Specular, BaseColor, Metallic);
 	Payload.RoughnessData = float3(Roughness, Roughness, 0);
@@ -844,14 +871,12 @@ RAY_TRACING_ENTRY_CLOSEST_HIT(PathTracingMaterialCHS,
 	Payload.WeightV = 1;
 	Payload.CoverageAboveAlongN = 0;
 	Payload.TransmittanceN = 1;
+	#if SUBSTRATE_PATH_TRACING_GLINTS
 	Payload.GlintValue = 1;
+	#endif
+	#if SUBSTRATE_PATH_TRACING_SPECULAR_PROFILE
 	Payload.SpecularProfileId = 0;
-#else
-	Payload.BaseColor = BaseColor;
-	Payload.Specular = Specular;
-	Payload.Metallic = Metallic;
-	Payload.Roughness = Roughness;
-#endif
+	#endif
 
 #if MATERIAL_USES_ANISOTROPY
 	Payload.WorldTangent = CalculateAnisotropyTangent(MaterialParameters, PixelMaterialInputs);
@@ -876,29 +901,17 @@ RAY_TRACING_ENTRY_CLOSEST_HIT(PathTracingMaterialCHS,
 		Payload.Anisotropy = -1;
 	#endif
 
-	#if USE_LEGACY_CODEPATH_FOR_SUBSTRATE
 		Payload.DiffuseColor = BaseColor;
 		Payload.SpecularColor = 0;
 		Payload.RoughnessData = 1;
-	#else
-		Payload.Specular = 0;
-		Payload.Metallic = 0;
-		Payload.Roughness = 1;
-	#endif
 	}
 	else
 #elif MATERIAL_SHADINGMODEL_UNLIT
 	if (Payload.ShadingModelID == SHADINGMODELID_UNLIT)
 	{
 		// Regular unlit material, make sure the volumetric component is zero'ed out
-	#if USE_LEGACY_CODEPATH_FOR_SUBSTRATE
 		Payload.DiffuseColor = 0;
 		Payload.SpecularColor = 0;
-	#else
-		Payload.BaseColor = 0;
-		Payload.Specular = 0;
-		Payload.Metallic = 0;
-	#endif
 	}
 	else
 #endif // MATERIAL_SHADINGMODEL_UNLIT
@@ -912,31 +925,17 @@ RAY_TRACING_ENTRY_CLOSEST_HIT(PathTracingMaterialCHS,
 
 	#if CLEAR_COAT_BOTTOM_NORMAL
 		#if NUM_MATERIAL_OUTPUTS_CLEARCOATBOTTOMNORMAL > 0
+			float3 BottomNormal = ClearCoatBottomNormal0(MaterialParameters);
 			#if MATERIAL_TANGENTSPACENORMAL
-				float3 BottomNormal = normalize(TransformTangentVectorToWorld(MaterialParameters.TangentToWorld, ClearCoatBottomNormal0(MaterialParameters)));
-			#else
-				float3 BottomNormal = ClearCoatBottomNormal0(MaterialParameters);
+		    BottomNormal = normalize(TransformTangentVectorToWorld(MaterialParameters.TangentToWorld, BottomNormal));
 			#endif
+			Payload.SetBottomCoatNormal(BottomNormal);
 		#else
-			float3 BottomNormal = Payload.WorldNormal;
+			Payload.SetBottomCoatNormal(Payload.WorldNormal);
 		#endif
-	#else
-			float3 BottomNormal = Payload.WorldNormal;
 	#endif
-
-#if USE_LEGACY_CODEPATH_FOR_SUBSTRATE
 		Payload.RoughnessData.y = ClearCoatRoughness; // adjusted later
 		Payload.RoughnessData.z = ClearCoat;
-
-		float2 oct2 = UnitVectorToOctahedron(Payload.WorldNormal);
-		float2 oct1 = UnitVectorToOctahedron(BottomNormal);
-		float2 BottomNormalOct = ((oct1 - oct2) * 0.25) + (128.0 / 255.0);
-		Payload.BottomNormalOct16bits = PackR8(BottomNormalOct.x) | (PackR8(BottomNormalOct.y) << 8);
-#else
-		Payload.SetClearCoat(ClearCoat);
-		Payload.SetClearCoatRoughness(AdjustMaterialRoughness(ClearCoatRoughness, PackedPayload.GetPathRoughness()));
-		Payload.SetClearCoatBottomNormal(BottomNormal);
-#endif
 	}
 	else
 #endif // MATERIAL_SHADINGMODEL_CLEAR_COAT
@@ -946,7 +945,6 @@ RAY_TRACING_ENTRY_CLOSEST_HIT(PathTracingMaterialCHS,
 	{
 		const float3 ClothColor = GetMaterialSubsurfaceData(PixelMaterialInputs).rgb;
 		const float Fuzz = saturate(GetMaterialCustomData0(PixelMaterialInputs));
-	#if USE_LEGACY_CODEPATH_FOR_SUBSTRATE
 		// the default substrate model includes a fuzz layer, so we can do it this way
 		// Additionally, tune the look like in SubstrateConvertLegacyMaterialDynamic
 		Payload.SpecularColor *= 1 - Fuzz;
@@ -954,10 +952,6 @@ RAY_TRACING_ENTRY_CLOSEST_HIT(PathTracingMaterialCHS,
 		Payload.FuzzAmount = Fuzz;
 		Payload.FuzzRoughness = 0.5 * Roughness + 0.5; // (adjust range because the legacy model has a very different response)
 		Payload.ShadingModelID = SHADINGMODELID_SUBSTRATE;
-	#else
-		Payload.SetClothColor(ClothColor);
-		Payload.SetClothAmount(Fuzz);
-	#endif
 	}
 	else
 #endif // MATERIAL_SHADINGMODEL_CLOTH
@@ -966,11 +960,7 @@ RAY_TRACING_ENTRY_CLOSEST_HIT(PathTracingMaterialCHS,
 	if (Payload.ShadingModelID == SHADINGMODELID_TWOSIDED_FOLIAGE)
 	{
 		const float3 SubsurfaceColor = GetMaterialSubsurfaceData(PixelMaterialInputs).rgb;
-	#if USE_LEGACY_CODEPATH_FOR_SUBSTRATE
 		Payload.SetFoliageTransmissionColor(SubsurfaceColor);
-	#else
-		Payload.SetSubsurfaceColor(SubsurfaceColor);
-	#endif
 	}
 	else
 #endif // MATERIAL_SHADINGMODEL_TWOSIDED_FOLIAGE
@@ -997,19 +987,12 @@ RAY_TRACING_ENTRY_CLOSEST_HIT(PathTracingMaterialCHS,
 		float PhaseG = 0.0;
 	#endif
 
-	#if USE_LEGACY_CODEPATH_FOR_SUBSTRATE
 		AdjustPathTracingTwoSidedColorBalance(Payload.DiffuseColor, SubsurfaceColor);
 		// default substrate slab includes SSS
 		Payload.ShadingModelID = SHADINGMODELID_SUBSTRATE;
 		Payload.DiffuseColor += SubsurfaceColor; // no distinction between diffuse portion and SSS portion (could try using stochastic blending here instead?)
 		Payload.MeanFreePath = SSSRadius;
 		Payload.PhaseG = PhaseG;
-	#else
-		Payload.ShadingModelID = SHADINGMODELID_SUBSURFACE;
-		Payload.SetSubsurfaceColor(SubsurfaceColor);
-		Payload.SetSubsurfaceRadius(SSSRadius);
-		Payload.SetSubsurfacePhaseFunction(PhaseG);
-	#endif
 	}
 	else
 #endif // MATERIAL_SHADINGMODEL_SUBSURFACE || MATERIAL_SHADINGMODEL_PREINTEGRATED_SKIN
@@ -1053,17 +1036,11 @@ RAY_TRACING_ENTRY_CLOSEST_HIT(PathTracingMaterialCHS,
 		float PhaseG = 0;
 	#endif
 
-	#if USE_LEGACY_CODEPATH_FOR_SUBSTRATE
 		// substrate slab includes SSS and dual roughness already
 		Payload.ShadingModelID = SHADINGMODELID_SUBSTRATE;
 		Payload.MeanFreePath = SSSRadius;
 		Payload.PhaseG = PhaseG;
 		Payload.RoughnessData = float3(Lobe0Roughness, Lobe1Roughness, LobeMix);
-	#else
-		Payload.SetSubsurfaceRadius(SSSRadius);
-		Payload.SetSubsurfacePhaseFunction(PhaseG);
-		Payload.SetDualRoughnessSpecular(Lobe0Roughness, Lobe1Roughness, LobeMix);
-	#endif
 	}
 	else
 #endif // MATERIAL_SHADINGMODEL_SUBSURFACE_PROFILE
@@ -1084,11 +1061,7 @@ RAY_TRACING_ENTRY_CLOSEST_HIT(PathTracingMaterialCHS,
 		const float3 WaterAlbedo = select(WaterExtinction > 0.0, ScatteringCoeff / WaterExtinction, float3(0.0, 0.0, 0.0));
 
 		Payload.ShadingModelID = SHADINGMODELID_SOLID_GLASS; // replace shading model to avoid extra cases in raygen
-	#if USE_LEGACY_CODEPATH_FOR_SUBSTRATE
-		Payload.MeanFreePath = rcp(max(1e-5, WaterExtinction));
-	#else
-		Payload.SetExtinction(WaterExtinction);
-	#endif
+		Payload.MeanFreePath = rcp(WaterExtinction);
 	    Payload.SetGlassAlbedo(WaterAlbedo);
 		Payload.SetGlassPhase(PhaseG);
 		Payload.Ior = DielectricF0ToIor(DielectricSpecularToF0(Specular));
@@ -1109,22 +1082,14 @@ RAY_TRACING_ENTRY_CLOSEST_HIT(PathTracingMaterialCHS,
 		{
 			// Pick glass lobe - shade as glass (leave SHADINGMODELID_THIN_TRANSLUCENT on)
 			Payload.Ior = Ior;
-		#if USE_LEGACY_CODEPATH_FOR_SUBSTRATE
-			Payload.MeanFreePath = TransmittanceToMeanFreePath(TransmittanceColor, SUBSTRATE_SIMPLEVOLUME_THICKNESS_CM * CENTIMETER_TO_METER) * METER_TO_CENTIMETER;
-			Payload.DiffuseColor = 0;
-		#else
 			Payload.SetTransmittanceColor(TransmittanceColor);
-		#endif
+			Payload.DiffuseColor = 0;
 		}
 		else
 		{
 			// Pick solid lobe -- toggle back to default_lit response
 			// Ior stays 0.0 since this mode does not include any refraction
-		#if USE_LEGACY_CODEPATH_FOR_SUBSTRATE
 			Payload.ShadingModelID = SHADINGMODELID_SUBSTRATE;
-		#else
-			Payload.ShadingModelID = SHADINGMODELID_DEFAULT_LIT;
-		#endif
 		}
 		
 		float SurfaceCoverage = GetThinTranslucentMaterialOutput1(MaterialParameters);
@@ -1179,11 +1144,7 @@ RAY_TRACING_ENTRY_CLOSEST_HIT(PathTracingMaterialCHS,
 		const float IrisMask     = saturate(GetMaterialCustomData0(PixelMaterialInputs));
 		const float IrisDistance = saturate(GetMaterialCustomData1(PixelMaterialInputs));
 		Payload.SetEyeIrisMask(IrisMask);
-	#if USE_LEGACY_CODEPATH_FOR_SUBSTRATE
 		Payload.SpecularColor = ComputeF0(Specular, BaseColor, 0);
-	#else
-		Payload.Metallic = 0.0;
-	#endif
 
 	#if NUM_MATERIAL_OUTPUTS_GETTANGENTOUTPUT > 0
 		// Blend in the negative intersection normal to create some concavity
@@ -1201,11 +1162,7 @@ RAY_TRACING_ENTRY_CLOSEST_HIT(PathTracingMaterialCHS,
 
 	#if IRIS_NORMAL // on
 		// Match BasePassPixelShader logic, Specular part of g-buffer is re-use, so hard code the term
-		#if USE_LEGACY_CODEPATH_FOR_SUBSTRATE
-			Payload.SpecularColor = DielectricSpecularToF0(0.25);
-		#else
-			Payload.Specular = 0.25;
-		#endif
+		Payload.SpecularColor = DielectricSpecularToF0(0.25);
 
 		#if NUM_MATERIAL_OUTPUTS_CLEARCOATBOTTOMNORMAL > 0
 			float3 IrisNormal = normalize( ClearCoatBottomNormal0(MaterialParameters) );
@@ -1241,13 +1198,8 @@ RAY_TRACING_ENTRY_CLOSEST_HIT(PathTracingMaterialCHS,
 		float PhaseG = 0;
 	#endif
 
-	#if USE_LEGACY_CODEPATH_FOR_SUBSTRATE
 		Payload.MeanFreePath = SSSRadius;
 		Payload.PhaseG = PhaseG;
-	#else
-		Payload.SetSubsurfaceRadius(SSSRadius);
-		Payload.SetSubsurfacePhaseFunction(PhaseG);
-	#endif
 	}
 	else
 #endif // MATERIAL_SHADINGMODEL_EYE
@@ -1255,9 +1207,7 @@ RAY_TRACING_ENTRY_CLOSEST_HIT(PathTracingMaterialCHS,
 #if MATERIAL_SHADINGMODEL_DEFAULT_LIT
 	if (Payload.ShadingModelID == SHADINGMODELID_DEFAULT_LIT)
 	{
-		#if USE_LEGACY_CODEPATH_FOR_SUBSTRATE
 		Payload.ShadingModelID = SHADINGMODELID_SUBSTRATE;
-		#endif
 
 	#if PATH_TRACER_REFRACTION_ENABLED
 		// only allow refraction for default lit materials since we need space for the absorption amount
@@ -1283,12 +1233,7 @@ RAY_TRACING_ENTRY_CLOSEST_HIT(PathTracingMaterialCHS,
 				float Extinction = 0;
 			#endif
 
-			#if USE_LEGACY_CODEPATH_FOR_SUBSTRATE
 				Payload.MeanFreePath = rcp(Extinction);
-			#else
-				Payload.SetExtinction(Extinction);
-			#endif
-
 				Payload.SetGlassAlbedo(0.0);
 			}
 			Payload.BSDFOpacity = 1.0; // cancels out with the probability of selecting the lobe
@@ -1302,12 +1247,8 @@ RAY_TRACING_ENTRY_CLOSEST_HIT(PathTracingMaterialCHS,
 	}
 
 	// adjust after everything else (because SSS profile case needs to decode the dual spec info from the unmodified roughness)
-#if USE_LEGACY_CODEPATH_FOR_SUBSTRATE
 	Payload.RoughnessData.x = AdjustMaterialRoughness(Payload.RoughnessData.x, PackedPayload.GetPathRoughness());
 	Payload.RoughnessData.y = AdjustMaterialRoughness(Payload.RoughnessData.y, PackedPayload.GetPathRoughness());
-#else
-	Payload.Roughness = AdjustMaterialRoughness(Payload.Roughness, PackedPayload.GetPathRoughness());
-#endif
 
 #endif // SUBSTRATE_ENABLED && !USE_LEGACY_CODEPATH_FOR_SUBSTRATE
 
diff --git a/Engine/Shaders/Private/PathTracing/Volume/PathTracingVolumeCommon.ush b/Engine/Shaders/Private/PathTracing/Volume/PathTracingVolumeCommon.ush
--- a/Engine/Shaders/Private/PathTracing/Volume/PathTracingVolumeCommon.ush
+++ b/Engine/Shaders/Private/PathTracing/Volume/PathTracingVolumeCommon.ush
@@ -246,9 +246,9 @@ FPathTracingPayload CreateMediumHitPayload(float HitT, float3 TranslatedWorldPos
 	Result.SetHG(HGWeight, ShadedResult.PhaseG);
 	Result.SetDualHG(DualHGWeight, ShadedResult.DualPhaseData);
 	Result.SetCloudFactor(1.0);
-#if SUBSTRATE_ENABLED
+	#if SUBSTRATE_PATH_TRACING_GLINTS
 	Result.GlintValue = 1; // avoid encode/decode cycle zero'ing out some parameters
-#endif
+	#endif
 	return Result;
 }
 
diff --git a/Engine/Source/Runtime/Engine/Private/UnrealEngine.cpp b/Engine/Source/Runtime/Engine/Private/UnrealEngine.cpp
--- a/Engine/Source/Runtime/Engine/Private/UnrealEngine.cpp
+++ b/Engine/Source/Runtime/Engine/Private/UnrealEngine.cpp
@@ -3616,10 +3616,8 @@ void UEngine::ConditionallyLoadPreIntegratedSkinBRDFTexture()
 
 void UEngine::LoadLTCTextures()
 {
-	const bool bSubstrateEnabled = Substrate::IsSubstrateEnabled();
-
 	const bool bGGX = true;
-	const bool bSheen = bSubstrateEnabled;
+	const bool bSheen = true;
 	if (bGGX && GGXLTCAmpTexture == nullptr && GGXLTCAmpTextureName.IsValid())
 	{
 		LoadEngineTexture(GGXLTCAmpTexture, *GGXLTCAmpTextureName.ToString());
diff --git a/Engine/Source/Runtime/Renderer/Private/PathTracing.cpp b/Engine/Source/Runtime/Renderer/Private/PathTracing.cpp
--- a/Engine/Source/Runtime/Renderer/Private/PathTracing.cpp
+++ b/Engine/Source/Runtime/Renderer/Private/PathTracing.cpp
@@ -1279,13 +1279,7 @@ END_GLOBAL_SHADER_PARAMETER_STRUCT()
 
 IMPLEMENT_GLOBAL_SHADER_PARAMETER_STRUCT(FPathTracingCloudParameterGlobals, "PathTracingCloudParameters");
 
-static uint32 GetPathtracingMaterialPayloadSize()
-{
-	// Substrate uses a slightly bigger payload as the basic slab contains more information
-	return Substrate::IsSubstrateEnabled() ? 76u : 64u;
-}
-
-IMPLEMENT_RT_PAYLOAD_TYPE_FUNCTION(ERayTracingPayloadType::PathTracingMaterial, GetPathtracingMaterialPayloadSize);
+IMPLEMENT_RT_PAYLOAD_TYPE(ERayTracingPayloadType::PathTracingMaterial, 76);
 IMPLEMENT_RT_PAYLOAD_TYPE(ERayTracingPayloadType::GPULightmass, 32);
 
 class FPathTracingRG : public FGlobalShader
diff --git a/Engine/Source/Runtime/Renderer/Private/SystemTextures.cpp b/Engine/Source/Runtime/Renderer/Private/SystemTextures.cpp
--- a/Engine/Source/Runtime/Renderer/Private/SystemTextures.cpp
+++ b/Engine/Source/Runtime/Renderer/Private/SystemTextures.cpp
@@ -991,8 +991,8 @@ void FSystemTextures::InitializeEngineDependentTextures(FRHICommandListImmediate
 		}
 	}
 
-	// Initialize textures only once (SheenLTC is only used with Substrate)
-	bEngineDependentTexturesInitialized = GGXLTCMat && GGXLTCAmp && (!Substrate::IsSubstrateEnabled() || SheenLTC);
+	// Initialize textures only once
+	bEngineDependentTexturesInitialized = GGXLTCMat && GGXLTCAmp && SheenLTC;
 }
 
 void FSystemTextures::ReleaseRHI()
