Fix incorrect hair shading due to incorrect tracking of complex transmittance bit with Substrate.

#Synced-CL 49616947

diff --git a/Engine/Shaders/Private/DeferredShadingCommon.ush b/Engine/Shaders/Private/DeferredShadingCommon.ush
--- a/Engine/Shaders/Private/DeferredShadingCommon.ush
+++ b/Engine/Shaders/Private/DeferredShadingCommon.ush
@@ -1130,6 +1130,11 @@ float ConvertSubsurfaceProfileToFloat(uint InProfileIdUint)
 	return (0xFF & InProfileIdUint) * (1.f / 255.f);
 }
 
+bool ShouldUseHairComplexTransmittance(FGBufferData GBuffer)
+{
+	return GBuffer.CustomData.a > 0;
+}
+
 #if SHADING_PATH_DEFERRED
 
 #if FEATURE_LEVEL >= FEATURE_LEVEL_SM5
diff --git a/Engine/Shaders/Private/HairStrands/HairStrandsComposition.usf b/Engine/Shaders/Private/HairStrands/HairStrandsComposition.usf
--- a/Engine/Shaders/Private/HairStrands/HairStrandsComposition.usf
+++ b/Engine/Shaders/Private/HairStrands/HairStrandsComposition.usf
@@ -448,7 +448,7 @@ void MainPS(
 
 	#if SUBSTRATE_ENABLED && SUBSTRATE_GBUFFER_FORMAT==1
 		const bool bWriteHairData = bWriteDummyData == 0;
-		FSubstrateBSDF HairBSDF = GetSubstrateHairBSDF(Samples_BaseColor, 0.f /*Scatter*/, Samples_Specular, Samples_Roughness, Samples_Backlit, 0 /*Emissive*/, 0/*SharedLocalBasisIndex*/).InlinedBSDF;
+		FSubstrateBSDF HairBSDF = GetSubstrateHairBSDF(Samples_BaseColor, 0.f /*Scatter*/, Samples_Specular, Samples_Roughness, Samples_Backlit, 0 /*Emissive*/, true /*bComplexTransmittance*/, 0/*SharedLocalBasisIndex*/).InlinedBSDF;
 
 		FSubstrateAddressing SubstrateAddressing = GetSubstratePixelDataByteOffset(PixelCoord.xy, uint2(ResolvedView.BufferSizeAndInvSize.xy), MaxBytesPerPixel);
 		FSubstratePixelHeader SubstratePixelHeader = InitialiseSubstratePixelHeader();
diff --git a/Engine/Shaders/Private/HairStrands/HairStrandsDeepTransmittanceDualScattering.ush b/Engine/Shaders/Private/HairStrands/HairStrandsDeepTransmittanceDualScattering.ush
--- a/Engine/Shaders/Private/HairStrands/HairStrandsDeepTransmittanceDualScattering.ush
+++ b/Engine/Shaders/Private/HairStrands/HairStrandsDeepTransmittanceDualScattering.ush
@@ -188,8 +188,3 @@ FHairTransmittanceData EvaluateDualScattering(inout FGBufferData GBuffer, float3
 	checkSlow(GBuffer.ShadingModelID == SHADINGMODELID_HAIR);
 	return EvaluateDualScattering(GBuffer.BaseColor, GBuffer.WorldNormal, GBuffer.Roughness, V, L);
 }
-
-bool ShouldUseHairComplexTransmittance(FGBufferData GBuffer)
-{
-	return GBuffer.CustomData.a > 0;
-}
diff --git a/Engine/Shaders/Private/Substrate/Substrate.ush b/Engine/Shaders/Private/Substrate/Substrate.ush
--- a/Engine/Shaders/Private/Substrate/Substrate.ush
+++ b/Engine/Shaders/Private/Substrate/Substrate.ush
@@ -1876,17 +1876,17 @@ FSubstrateData GetSubstrateUnlitBSDF(float3 Emissive, float3 TransmittanceColor,
 	return SubstrateData;
 }
 
-FSubstrateData GetSubstrateHairBSDF(float3 BaseColor, float Scatter, float Specular, float Roughness, float Backlit, float3 Emissive, uint SharedLocalBasisIndex)
+FSubstrateData GetSubstrateHairBSDF(float3 BaseColor, float Scatter, float Specular, float Roughness, float Backlit, float3 Emissive, bool bComplexTransmittance, uint SharedLocalBasisIndex)
 {
 	FSubstrateData SubstrateData = GetInitialisedSubstrateData();
 
 	// Enable complex transmittance only for hair using cards/strands vertex factories
-	#ifndef USE_HAIR_COMPLEX_TRANSMITTANCE
-	  #if defined(HAIR_CARD_MESH_FACTORY) || defined(HAIR_STRAND_MESH_FACTORY)
-		#define USE_HAIR_COMPLEX_TRANSMITTANCE 1
-	  #else
-		#define USE_HAIR_COMPLEX_TRANSMITTANCE 0
+	#ifdef USE_HAIR_COMPLEX_TRANSMITTANCE
+	  #if USE_HAIR_COMPLEX_TRANSMITTANCE == 0
+		bComplexTransmittance = false;
 	  #endif
+	#elif defined(HAIR_CARD_MESH_FACTORY) || defined(HAIR_STRAND_MESH_FACTORY)
+		bComplexTransmittance = true;
 	#endif
 
 	BSDF_SETTYPE					(SubstrateData.InlinedBSDF, SUBSTRATE_BSDF_TYPE_HAIR);
@@ -1897,7 +1897,7 @@ FSubstrateData GetSubstrateHairBSDF(float3 BaseColor, float Scatter, float Specu
 	HAIR_ROUGHNESS					(SubstrateData.InlinedBSDF) = Roughness;
 	HAIR_SPECULAR					(SubstrateData.InlinedBSDF) = Specular;
 	HAIR_BACKLIT					(SubstrateData.InlinedBSDF) = Backlit;
-	HAIR_COMPLEXTRANSMITTANCE		(SubstrateData.InlinedBSDF) = USE_HAIR_COMPLEX_TRANSMITTANCE;
+	HAIR_COMPLEXTRANSMITTANCE		(SubstrateData.InlinedBSDF) = bComplexTransmittance;
 
 #if SUBSTRATE_INLINE_SHADING
 	SubstrateData.InlinedBSDF.Coverage = 1.0f;
diff --git a/Engine/Shaders/Private/Substrate/SubstrateLegacyConversion.ush b/Engine/Shaders/Private/Substrate/SubstrateLegacyConversion.ush
--- a/Engine/Shaders/Private/Substrate/SubstrateLegacyConversion.ush
+++ b/Engine/Shaders/Private/Substrate/SubstrateLegacyConversion.ush
@@ -209,6 +209,7 @@ FSubstrateData SubstrateConvertLegacyMaterialDynamic(
 			Roughness,	   // Roughness
 			ClearCoat,	   // Backlit
 			Emissive,	   // EmissiveColor
+			false,		   // bComplexTransmittance
 			SharedLocalBasisIndex);
 	}
 	#endif // MATERIAL_SHADINGMODEL_HAIR
@@ -600,6 +601,7 @@ FSubstrateData SubstrateConvertLegacyMaterialStatic(
 				Roughness,	   // Roughness
 				ClearCoat,	   // Backlit
 				Emissive,	   // EmissiveColor
+				false,	       // bComplexTransmittance
 				SharedLocalBasisIndex);
 	}
 	else if (ShadingModel == SHADINGMODELID_EYE)
diff --git a/Engine/Shaders/Private/Substrate/SubstrateRead.ush b/Engine/Shaders/Private/Substrate/SubstrateRead.ush
--- a/Engine/Shaders/Private/Substrate/SubstrateRead.ush
+++ b/Engine/Shaders/Private/Substrate/SubstrateRead.ush
@@ -90,7 +90,8 @@ FSubstrateGBufferBSDF SubstrateReadGBufferBSDF(in FScreenSpaceData ScreenSpaceDa
 		{
 			const float Scatter = GBuffer.Metallic;
 			const float BackLit = GBuffer.CustomData.b;
-			BSDF = GetSubstrateHairBSDF(GBuffer.BaseColor, Scatter, GBuffer.Specular, GBuffer.Roughness, BackLit, Emissive, SharedLocalBasisIndex).InlinedBSDF;
+			const bool bComplexTransmittance = ShouldUseHairComplexTransmittance(GBuffer);
+			BSDF = GetSubstrateHairBSDF(GBuffer.BaseColor, Scatter, GBuffer.Specular, GBuffer.Roughness, BackLit, Emissive, bComplexTransmittance, SharedLocalBasisIndex).InlinedBSDF;
 
 		}
 		else if (GBuffer.ShadingModelID == SHADINGMODELID_EYE)
diff --git a/Engine/Source/Runtime/Core/Private/UObject/FortniteMainBranchObjectVersion.cpp b/Engine/Source/Runtime/Core/Private/UObject/FortniteMainBranchObjectVersion.cpp
--- a/Engine/Source/Runtime/Core/Private/UObject/FortniteMainBranchObjectVersion.cpp
+++ b/Engine/Source/Runtime/Core/Private/UObject/FortniteMainBranchObjectVersion.cpp
@@ -13,7 +13,7 @@ TMap<FGuid, FGuid> FFortniteMainBranchObjectVersion::GetSystemGuids()
 	SystemGuids.Add(DevGuids.MATERIALSHADERMAP_DERIVEDDATA_VER, FGuid("1A77CBCDFBDB4C12BF8563EDBB7D8888"));
 	SystemGuids.Add(DevGuids.NIAGARASHADERMAP_DERIVEDDATA_VER, FGuid("CC229D919A214E59B79D01AC2E5BF981"));
 	SystemGuids.Add(DevGuids.Niagara_LatestScriptCompileVersion, FGuid("FB15C2870A4E4AE8B5F58566F0E3825E"));
-	SystemGuids.Add(DevGuids.MaterialTranslationDDCVersion, FGuid("88DCA9D21A884B7FBB944CAB57C34EC3"));
+	SystemGuids.Add(DevGuids.MaterialTranslationDDCVersion, FGuid("EA008827093C4A2BBFAC9F4EA3D92BE9"));
 
 	// These GUIDs are in AppendRestrictedFortniteMainStreamObjectVersionGuids as they should rarely be changed.
 	//SystemGuids.Add(DevGuids.NANITE_DERIVEDDATA_VER, FGuid("ACF98184978742A592537DA223A1E6BE"));
diff --git a/Engine/Source/Runtime/Engine/Private/Materials/HLSLMaterialTranslator.cpp b/Engine/Source/Runtime/Engine/Private/Materials/HLSLMaterialTranslator.cpp
--- a/Engine/Source/Runtime/Engine/Private/Materials/HLSLMaterialTranslator.cpp
+++ b/Engine/Source/Runtime/Engine/Private/Materials/HLSLMaterialTranslator.cpp
@@ -15366,7 +15366,7 @@ int32 FHLSLMaterialTranslator::SubstrateHairBSDF(int32 BaseColor, int32 Scatter,
 	if (PromoteToOperator)
 	{
 		return AddCodeChunk(
-			MCT_Substrate, TEXT("Parameters.%s.PromoteParameterBlendedBSDFToOperator(GetSubstrateHairBSDF(%s, %s, %s, %s, %s, %s, %s), %u, %u, %u, %u) /* Tangent:%s */"),
+			MCT_Substrate, TEXT("Parameters.%s.PromoteParameterBlendedBSDFToOperator(GetSubstrateHairBSDF(%s, %s, %s, %s, %s, %s, false/*bComplexTransmittance*/, %s), %u, %u, %u, %u) /* Tangent:%s */"),
 			*GetParametersSubstrateTreeName(CurrentSubstrateCompilationContext),
 			*SubstrateGetCastParameterCode(BaseColor, MCT_Float3),
 			*SubstrateGetCastParameterCode(Scatter, MCT_Float),
@@ -15384,7 +15384,7 @@ int32 FHLSLMaterialTranslator::SubstrateHairBSDF(int32 BaseColor, int32 Scatter,
 	}
 
 	return AddCodeChunk(
-		MCT_Substrate, TEXT("GetSubstrateHairBSDF(%s, %s, %s, %s, %s, %s, %s)"),
+		MCT_Substrate, TEXT("GetSubstrateHairBSDF(%s, %s, %s, %s, %s, %s, false/*bComplexTransmittance*/, %s)"),
 		*SubstrateGetCastParameterCode(BaseColor, MCT_Float3),
 		*SubstrateGetCastParameterCode(Scatter, MCT_Float),
 		*SubstrateGetCastParameterCode(Specular, MCT_Float),
