Change rough diffuse BSDF to use EON model instead of Chan.

EON's model is more accurate than Chan's model.  Replays with rough diffuse enabled didn't show any major regression (+0.02ms for the whole frame).

#Synced-CL 49616947

diff --git a/Engine/Shaders/Private/BRDF.ush b/Engine/Shaders/Private/BRDF.ush
--- a/Engine/Shaders/Private/BRDF.ush
+++ b/Engine/Shaders/Private/BRDF.ush
@@ -232,7 +232,7 @@ float3 Diffuse_EON( float3 DiffuseColor, float Roughness, float NoV, float NoL,
 	return Rho * (f_ss + Rho * f_ms) * (1.0 / PI);
 }
 
-#define ROUGH_DIFFUSE_BRDF_VERSION 2
+#define ROUGH_DIFFUSE_BRDF_VERSION 3
 // This models a rough surface that has a GGX NDF where each microfacet has a lambertian response. Various models have been proposed
 // to try and approximate this behavior.
 float3 Diffuse_GGX_Rough( float3 DiffuseColor, float Roughness, float NoV, float NoL, float VoH, float NoH, float RetroReflectivityWeight )
