Fix incorrect Subsurface look with area lights placed beneath the surface.

#Synced-CL 49616947

diff --git a/Engine/Shaders/Private/Substrate/SubstrateEvaluation.ush b/Engine/Shaders/Private/Substrate/SubstrateEvaluation.ush
--- a/Engine/Shaders/Private/Substrate/SubstrateEvaluation.ush
+++ b/Engine/Shaders/Private/Substrate/SubstrateEvaluation.ush
@@ -895,14 +895,14 @@ FSubstrateEvaluateResult SubstrateEvaluateBSDFCommon(FSubstrateBSDFContext BSDFC
 					// Legacy Subsurface
 					 
 					// To get an effect when you see through the material (hard coded pow constant)
-					const half InScatter = pow(saturate(dot(BSDFContext.L, -BSDFContext.V)), 12) * lerp(3, .1f, Opacity);
+					const half InScatter = pow(saturate(dot(AreaLight.DiffuseL, -BSDFContext.V)), 12) * lerp(3, .1f, Opacity);
 					// Wrap around lighting, 
 					// * /(PI*2) to be energy consistent (hack do get some view dependnt and light dependent effect)
 					// * Opacity of 0 gives no normal dependent lighting, Opacity of 1 gives strong normal contribution	
 					// * Simplified version (with w=.5,  n=1.5): 
 					//     half WrappedDiffuse = 2 * pow(saturate((dot(N, L) + w) / (1.0f + w)), n) * (n + 1) / (2 * (1 + w));
 					//     NormalContribution = WrappedDiffuse * Opacity + 1 - Opacity;
-					const half WrappedDiffuse = pow(saturate(BSDFContext.Context.NoL * (1.f / 1.5f) + (0.5f / 1.5f)), 1.5f) * (2.5f / 1.5f);
+					const half WrappedDiffuse = pow(saturate(dot(BSDFContext.N, AreaLight.DiffuseL) * (1.f / 1.5f) + (0.5f / 1.5f)), 1.5f) * (2.5f / 1.5f);
 					const half NormalContribution = lerp(1.f, WrappedDiffuse, Opacity);
 					const half BackScatter = /* GBuffer.GBufferAO */ NormalContribution / (PI * 2);
 
