Substrate - fixed bad define name.

#tests PC editor and preflight
#rnx

#Synced-CL 49627059

diff --git a/Engine/Plugins/FX/Niagara/Shaders/Private/NiagaraDataInterfaceCollisionQuery.ush b/Engine/Plugins/FX/Niagara/Shaders/Private/NiagaraDataInterfaceCollisionQuery.ush
--- a/Engine/Plugins/FX/Niagara/Shaders/Private/NiagaraDataInterfaceCollisionQuery.ush
+++ b/Engine/Plugins/FX/Niagara/Shaders/Private/NiagaraDataInterfaceCollisionQuery.ush
@@ -53,7 +53,7 @@ void NDICollisionQuery_QuerySceneDepthGPU(in float3 In_SamplePos, in float3 In_L
 		FLWCVector3 WorldPosition = WorldPositionFromSceneDepth(ScreenPosition.xy, SceneDepth);
 
 #if NDICOLLISIONQUERY_USE_GBUFFER_NORMAL
-#if SUBTRATE_GBUFFER_FORMAT==1
+#if SUBSTRATE_GBUFFER_FORMAT==1
 		uint2 ScreenPos = uint2(ScreenUV * View.BufferSizeAndInvSize.xy);
 		float3 WorldNormal = SubstratePublic_GetWorldNormal(ScreenPos);
 #else
diff --git a/Engine/Plugins/FX/Niagara/Shaders/Private/NiagaraDataInterfaceGBufferTemplate.ush b/Engine/Plugins/FX/Niagara/Shaders/Private/NiagaraDataInterfaceGBufferTemplate.ush
--- a/Engine/Plugins/FX/Niagara/Shaders/Private/NiagaraDataInterfaceGBufferTemplate.ush
+++ b/Engine/Plugins/FX/Niagara/Shaders/Private/NiagaraDataInterfaceGBufferTemplate.ush
@@ -6,7 +6,7 @@
 #define NDIGBUFFER_USE_GBUFFER_VELOCITY		(false && NDIGBUFFER_HAS_GBUFFER)
 #define NDIGBUFFER_IS_MOBILE				(FEATURE_LEVEL <= FEATURE_LEVEL_ES3_1)
 
-#define NIAGARA_SAMPLE_SUBSTRATE (SUBTRATE_GBUFFER_FORMAT==1 && !NDIGBUFFER_IS_MOBILE)
+#define NIAGARA_SAMPLE_SUBSTRATE (SUBSTRATE_GBUFFER_FORMAT==1 && !NDIGBUFFER_IS_MOBILE)
 
 #if NIAGARA_SAMPLE_SUBSTRATE
 #define NEEDS_SCENE_TEXTURES 1
diff --git a/Engine/Shaders/Private/Substrate/SubstrateExport.ush b/Engine/Shaders/Private/Substrate/SubstrateExport.ush
--- a/Engine/Shaders/Private/Substrate/SubstrateExport.ush
+++ b/Engine/Shaders/Private/Substrate/SubstrateExport.ush
@@ -996,7 +996,7 @@ void WriteHair(
 
 // Hidden conversion from root node for opaque material guarantee that there is a BSDF output with Coverage=1 and Transmittance=0 that is always visible.
 // SUBSTRATE_TODO a simple more general way to make this work would be to make sure there is no CoverageWeight node in the graph instead of MATERIAL_IS_SUBSTRATEHIDDENCONVERSION? In this case that optimisation would be valid too.
-#define SUBTRATE_OPTIMISED_LEGACY_OPAQUE_EXPORT (SUBSTRATE_OPAQUE_DEFERRED && MATERIAL_IS_SUBSTRATEHIDDENCONVERSION)
+#define SUBSTRATE_OPTIMISED_LEGACY_OPAQUE_EXPORT (SUBSTRATE_OPAQUE_DEFERRED && MATERIAL_IS_SUBSTRATEHIDDENCONVERSION)
 
 struct FExportResult
 {
@@ -1040,11 +1040,11 @@ FExportResult SubstrateMaterialExportOut(
 	const int BSDFIdx = 0;
 	#define CurrentBSDF SubstratePixelHeader.SubstrateTree.BSDFs[BSDFIdx]
 
-	#if !SUBTRATE_OPTIMISED_LEGACY_OPAQUE_EXPORT
+	#if !SUBSTRATE_OPTIMISED_LEGACY_OPAQUE_EXPORT
 	if (SubstratePixelHeader.SubstrateTree.BSDFCount > 0)
 	#endif
 	{
-		#if SUBTRATE_OPTIMISED_LEGACY_OPAQUE_EXPORT
+		#if SUBSTRATE_OPTIMISED_LEGACY_OPAQUE_EXPORT
 		// In this case, there is always a BSDF and it is always 100% visible.
 		// Sanitize BSDF before it is exported.
 		CurrentBSDF.SubstrateSanitizeBSDF();
@@ -1057,7 +1057,7 @@ FExportResult SubstrateMaterialExportOut(
 		SubstratePixelHeader.SubstrateUpdateTree(SubstrateData, V, Settings, Out.Coverage, Out.TransmittancePreCoverage);
 		#endif
 
-		#if !SUBTRATE_OPTIMISED_LEGACY_OPAQUE_EXPORT
+		#if !SUBSTRATE_OPTIMISED_LEGACY_OPAQUE_EXPORT
 		if (SubstrateIsBSDFVisible(CurrentBSDF))
 		#endif
 		{
