Substrate - fixes for two broken issues for tiled SSR within stencil tested path:
- fixed wrong logic conversion from 15385871
- fixed the stencil test logic which was not matching expectations.

#tests PC editor + renderdoc capture
#rnx
#Synced-CL 49627059

diff --git a/Engine/Shaders/Private/SSRT/SSRTReflections.usf b/Engine/Shaders/Private/SSRT/SSRTReflections.usf
--- a/Engine/Shaders/Private/SSRT/SSRTReflections.usf
+++ b/Engine/Shaders/Private/SSRT/SSRTReflections.usf
@@ -72,24 +72,24 @@ void ScreenSpaceReflectionsStencilPS(
 	const FSubstrateTopLayerData TopLayerData = SubstrateUnpackTopLayerData(Substrate.TopLayerTexture.Load(uint3(PixelCoord, 0)));
 
 	const float Roughness = TopLayerData.Roughness;
-	const bool bNoMaterial = !IsSubstrateMaterial(TopLayerData);
+	const bool bValidMaterial = IsSubstrateMaterial(TopLayerData);
 
 #else // SUBSTRATE_GBUFFER_FORMAT==1
 
 	FGBufferData GBuffer = GetGBufferDataFromSceneTextures(UV);
 	const float Roughness = GetRoughness(GBuffer);
-	const bool bNoMaterial = GBuffer.ShadingModelID == 0;
+	const bool bValidMaterial = GBuffer.ShadingModelID != 0;
 
 #endif // SUBSTRATE_GBUFFER_FORMAT==1
 
 	const float RoughnessFade = GetRoughnessFade(Roughness);
 
 #if !defined(TILE_COMPUTE_SSR)
-	if (RoughnessFade > 0.0 && bNoMaterial)
+	if (RoughnessFade > 0.0 && bValidMaterial)
 	{
 		// we are going to compute SSR for this pixel, so we discard this
 		// pixel shader invocation to not overwrite the stencil buffer and
-		// tehrefore execute ScreenSpaceReflectionsPS() for this pixel. 
+		// therefore execute ScreenSpaceReflectionsPS() for this pixel. 
 		discard;
 	}
 #endif
@@ -128,14 +128,14 @@ void ScreenSpaceReflections(
 
 	float3 N = 0;
 	float Roughness = 0.0;
-	bool bNoMaterial = true;
+	bool bValidMaterial = false;
 
 	const FSubstrateTopLayerData TopLayerData = SubstrateUnpackTopLayerData(Substrate.TopLayerTexture.Load(uint3(PixelCoord, 0)));
 
 	// For now, SSR is only used to trace the top material layer.
 	N = TopLayerData.WorldNormal;
 	Roughness = TopLayerData.Roughness;
-	bNoMaterial = ShouldReflectOnlyWater ? !IsSingleLayerWaterMaterial(TopLayerData) : !IsSubstrateMaterial(TopLayerData);
+	bValidMaterial = ShouldReflectOnlyWater ? IsSingleLayerWaterMaterial(TopLayerData) : IsSubstrateMaterial(TopLayerData);
 
 	float DeviceZ = SampleDeviceZFromSceneTextures(UV);
 	float SceneDepth = ConvertFromDeviceZ(DeviceZ);
@@ -166,14 +166,14 @@ void ScreenSpaceReflections(
 #endif
 	
 	float Roughness = GetRoughness(GBuffer);
-	const bool bNoMaterial = GBuffer.ShadingModelID == 0;
+	const bool bValidMaterial = GBuffer.ShadingModelID != 0;
 
 #endif // SUBSTRATE_GBUFFER_FORMAT==1
 
 	float RoughnessFade = GetRoughnessFade(Roughness);
 
 	// Early out. Useless if using the stencil prepass.
-	BRANCH if( RoughnessFade <= 0.0 || bNoMaterial)
+	BRANCH if( RoughnessFade <= 0.0 || !bValidMaterial)
 	{
 		return;
 	}
diff --git a/Engine/Source/Runtime/Renderer/Private/ScreenSpaceRayTracing.cpp b/Engine/Source/Runtime/Renderer/Private/ScreenSpaceRayTracing.cpp
--- a/Engine/Source/Runtime/Renderer/Private/ScreenSpaceRayTracing.cpp
+++ b/Engine/Source/Runtime/Renderer/Private/ScreenSpaceRayTracing.cpp
@@ -1230,7 +1230,7 @@ void RenderScreenSpaceReflections(
 				if (SSRStencilPrePass)
 				{
 					// Clobers the stencil to pixel that should not compute SSR
-					GraphicsPSOInit.DepthStencilState = TStaticDepthStencilState<false, CF_Always, true, CF_Equal, SO_Keep, SO_Keep, SO_Keep>::GetRHI();
+					GraphicsPSOInit.DepthStencilState = TStaticDepthStencilState<false, CF_Always, true, CF_NotEqual, SO_Keep, SO_Keep, SO_Keep>::GetRHI();
 				}
 
 				SetGraphicsPipelineState(RHICmdList, GraphicsPSOInit, 0x80);
@@ -1269,7 +1269,7 @@ void RenderScreenSpaceReflections(
 			if (SSRStencilPrePass)
 			{
 				// Clobers the stencil to pixel that should not compute SSR
-				GraphicsPSOInit.DepthStencilState = TStaticDepthStencilState<false, CF_Always, true, CF_Equal, SO_Keep, SO_Keep, SO_Keep>::GetRHI();
+				GraphicsPSOInit.DepthStencilState = TStaticDepthStencilState<false, CF_Always, true, CF_NotEqual, SO_Keep, SO_Keep, SO_Keep>::GetRHI();
 			}
 			GraphicsPSOInit.PrimitiveType = GRHISupportsRectTopology ? PT_RectList : PT_TriangleList;
 			GraphicsPSOInit.BoundShaderState.VertexDeclarationRHI = GEmptyVertexDeclaration.VertexDeclarationRHI;
@@ -1321,7 +1321,7 @@ void RenderScreenSpaceReflections(
 				if (SSRStencilPrePass)
 				{
 					// Clobers the stencil to pixel that should not compute SSR
-					GraphicsPSOInit.DepthStencilState = TStaticDepthStencilState<false, CF_Always, true, CF_Equal, SO_Keep, SO_Keep, SO_Keep>::GetRHI();
+					GraphicsPSOInit.DepthStencilState = TStaticDepthStencilState<false, CF_Always, true, CF_NotEqual, SO_Keep, SO_Keep, SO_Keep>::GetRHI();
 				}
 				GraphicsPSOInit.PrimitiveType = GRHISupportsRectTopology ? PT_RectList : PT_TriangleList;
 				GraphicsPSOInit.BoundShaderState.VertexDeclarationRHI = GEmptyVertexDeclaration.VertexDeclarationRHI;
