Fix virtual texture sampling issue when mixing 16 and 32 bit page tables.
The interpretation of page table bit depth was being stored in the physical texture uniform.
This is wrong because in some (rare) cases the physical texture might be shared in a material between virtual texture stacks with different page table bit depths.
The fix is to pass the page table bit depth in the page table uniform.
#rnx
#Synced-CL 49645631

diff --git a/Engine/Shaders/Private/VirtualTextureCommon.ush b/Engine/Shaders/Private/VirtualTextureCommon.ush
--- a/Engine/Shaders/Private/VirtualTextureCommon.ush
+++ b/Engine/Shaders/Private/VirtualTextureCommon.ush
@@ -189,6 +189,7 @@ struct VTPageTableUniform
 	uint MipBiasGroup; // 2
 	uint ShiftedPageTableID; // 4
 	uint AdaptiveLevelBias; // 4
+	bool bPageTable32Bit;
 
 	float2 SizeInPages;
 	float2 UVScale;
@@ -208,6 +209,7 @@ VTPageTableUniform VTPageTableUniform_Unpack(uint4 PackedPageTableUniform0, uint
 	result.MipBiasGroup = (PackedPageTableUniform1.y >> 28) & 0x3;
 	result.MaxLevel = PackedPageTableUniform1.z & 0xf;
 	result.AdaptiveLevelBias = (PackedPageTableUniform1.z >> 4) & 0xf;
+	result.bPageTable32Bit = (PackedPageTableUniform1.z & 0x100) != 0;
 	result.ShiftedPageTableID = PackedPageTableUniform1.w;
 	return result;
 }
@@ -223,8 +225,9 @@ VTPageTableUniform VTPageTableUniform_Unpack(uint2 PackedPageTableUniform)
 	Result.YOffsetInPages = (PackedPageTableUniform.x >> 12) & 0xfff;
 	Result.vPageTableMipBias = (PackedPageTableUniform.x >> 24) & 0xf;
 	Result.MipBiasGroup = 0;
-	Result.MaxLevel = (PackedPageTableUniform.y >> 24) & 0x3f;
+	Result.MaxLevel = (PackedPageTableUniform.y >> 24) & 0xf;
 	Result.AdaptiveLevelBias = 0;
+	Result.bPageTable32Bit = (PackedPageTableUniform.y >> 28) & 0x1;
 	Result.ShiftedPageTableID = PackedPageTableUniform.x & 0xf0000000;
 	return Result;
 }
@@ -245,6 +248,7 @@ struct VTPageTableResult
 	uint4 PageTableValue[4];
 	uint PackedRequest;
 	float MipLevelFrac;
+	bool bPageTable32Bit;
 	uint PendingMips;
 };
 
@@ -356,6 +360,7 @@ void TextureLoadVirtualPageTableInternal(
 	float2 UV, int vLevel)
 {
 	OutResult.UV = UV * PageTableUniform.SizeInPages;
+	OutResult.bPageTable32Bit = PageTableUniform.bPageTable32Bit;
 
 	const uint vLevelClamped = clamp(vLevel, 0, int(PageTableUniform.MaxLevel));
 	uint vPageX = (uint(OutResult.UV.x) + PageTableUniform.XOffsetInPages) >> vLevelClamped;
@@ -405,6 +410,7 @@ void TextureLoadVirtualPageTableInternal(
 	float2 UV, int vLevel)
 {
 	OutResult.UV = UV * PageTableUniform.SizeInPages;
+	OutResult.bPageTable32Bit = PageTableUniform.bPageTable32Bit;
 
 	const uint vLevelClamped = clamp(vLevel, 0, int(PageTableUniform.MaxLevel));
 	uint vPageX = (uint(OutResult.UV.x) + PageTableUniform.XOffsetInPages) >> vLevelClamped;
@@ -778,7 +784,6 @@ struct VTUniform
 	float pPageSize;
 	float vPageSize;
 	float vPageBorderSize;
-	bool bPageTableExtraBits;
 	float4 FallbackValue;
 };
 
@@ -786,10 +791,9 @@ struct VTUniform
 VTUniform VTUniform_Unpack(uint4 PackedUniform)
 {
 	VTUniform result;
-	result.pPageSize = abs(asfloat(PackedUniform.w));
+	result.pPageSize = asfloat(PackedUniform.w);
 	result.vPageSize = asfloat(PackedUniform.y);
 	result.vPageBorderSize = asfloat(PackedUniform.z);
-	result.bPageTableExtraBits = asfloat(PackedUniform.w) > 0;
 	result.FallbackValue.b = float((PackedUniform.x >> 0) & 0xFF) * (1.0f / 255.0f);
 	result.FallbackValue.g = float((PackedUniform.x >> 8) & 0xFF) * (1.0f / 255.0f);
 	result.FallbackValue.r = float((PackedUniform.x >> 16) & 0xFF) * (1.0f / 255.0f);
@@ -814,8 +818,8 @@ float2 VTComputePhysicalUVs(in out VTPageTableResult PageTableResult, uint Layer
 	const float UVScale = float(4096u >> vLevel) * (1.0f / 4096.0f);
 
 	// This will compile to runtime branch, but should in theory be conditional moves selecting 1 of 2 different bitfield extracting instructions
-	const uint pPageX = Uniform.bPageTableExtraBits ? (PackedPageTableValue >> 4) & 0xff : (PackedPageTableValue >> 4) & 0x3f;
-	const uint pPageY = Uniform.bPageTableExtraBits ? (PackedPageTableValue >> 12) & 0xff : (PackedPageTableValue >> 10) & 0x3f;
+	const uint pPageX = PageTableResult.bPageTable32Bit ? (PackedPageTableValue >> 4) & 0xff : (PackedPageTableValue >> 4) & 0x3f;
+	const uint pPageY = PageTableResult.bPageTable32Bit ? (PackedPageTableValue >> 12) & 0xff : (PackedPageTableValue >> 10) & 0x3f;
 
 	const float2 vPageFrac = frac(PageTableResult.UV * UVScale);
 	const float2 pUV = float2(pPageX, pPageY) * Uniform.pPageSize + (vPageFrac * Uniform.vPageSize + Uniform.vPageBorderSize);
diff --git a/Engine/Source/Runtime/Engine/Private/VT/MeshPaintVirtualTexture.cpp b/Engine/Source/Runtime/Engine/Private/VT/MeshPaintVirtualTexture.cpp
--- a/Engine/Source/Runtime/Engine/Private/VT/MeshPaintVirtualTexture.cpp
+++ b/Engine/Source/Runtime/Engine/Private/VT/MeshPaintVirtualTexture.cpp
@@ -151,9 +151,7 @@ namespace MeshPaintVirtualTexture
 		OutParams.PackedUniform.Y = FMath::AsUInt((float)vPageSize * RcpPhysicalTextureSize);
 		OutParams.PackedUniform.Z = FMath::AsUInt((float)PageBorderSize * RcpPhysicalTextureSize);
 		
-		const bool bPageTableExtraBits = InAllocatedVT->GetPageTableFormat() == EVTPageTableFormat::UInt32;
-		const float PackedSignBit = bPageTableExtraBits ? 1.f : -1.f;
-		OutParams.PackedUniform.W = FMath::AsUInt((float)pPageSize * RcpPhysicalTextureSize * PackedSignBit);
+		OutParams.PackedUniform.W = FMath::AsUInt((float)pPageSize * RcpPhysicalTextureSize);
 	}
 
 	/** A global set to track allocated mesh paint virtual textures. */
@@ -270,11 +268,12 @@ namespace MeshPaintVirtualTexture
 		const uint32 vPageTableMipBias = FMath::FloorLog2(AllocatedVT->GetVirtualTileSize());
 		const uint32 SpaceID = AllocatedVT->GetSpaceID();
 		const uint32 MaxLevel = AllocatedVT->GetMaxLevel();
+		const uint32 PageTable32Bit = AllocatedVT->GetPageTableFormat() == EVTPageTableFormat::UInt32 ? 1 : 0;
 
 		// Descriptor layout should match shader unpacking in VTPageTableUniform_Unpack().
 		// We optionally pack a coordinate index which is used for debug views.
 		Descriptor.X = vPageX | (vPageY << 12) | (vPageTableMipBias << 24) | (SpaceID << 28);
-		Descriptor.Y = WidthInPages | (HeightInPages << 12) | (MaxLevel << 24) | (InOptionalCoordinateIndex << 30);
+		Descriptor.Y = WidthInPages | (HeightInPages << 12) | (MaxLevel << 24) | PageTable32Bit << 28 | (InOptionalCoordinateIndex << 30);
 		
 		return Descriptor;
 	}
diff --git a/Engine/Source/Runtime/Renderer/Private/VT/AllocatedVirtualTexture.cpp b/Engine/Source/Runtime/Renderer/Private/VT/AllocatedVirtualTexture.cpp
--- a/Engine/Source/Runtime/Renderer/Private/VT/AllocatedVirtualTexture.cpp
+++ b/Engine/Source/Runtime/Renderer/Private/VT/AllocatedVirtualTexture.cpp
@@ -490,6 +490,7 @@ void FAllocatedVirtualTexture::GetPackedPageTableUniform(FUintVector4* OutUnifor
 	const uint32 vPageTableMipBias = FMath::FloorLog2(vPageSize);
 	const uint32 MipBiasGroup = GetResidencyMipMapBiasGroup();
 	const uint32 AdaptiveLevelBias = Description.AdaptiveLevelBias;
+	const uint32 PageTable32Bit = (GetPageTableFormat() == EVTPageTableFormat::UInt32) ? 1 : 0;
 
 	// Here MaxAnisotropy only controls the VT mip level selection
 	// We don't need to limit this value based on border size, and we can add this factor in even if HW anisotropic filtering is disabled
@@ -512,7 +513,7 @@ void FAllocatedVirtualTexture::GetPackedPageTableUniform(FUintVector4* OutUnifor
 
 	OutUniform[1].X = BitcastFloatToUInt32((float)MaxAnisotropyLog2);
 	OutUniform[1].Y = vPageX | (vPageY << 12) | (vPageTableMipBias << 24) | (MipBiasGroup << 28);
-	OutUniform[1].Z = MaxLevel | (AdaptiveLevelBias << 4);
+	OutUniform[1].Z = MaxLevel | (AdaptiveLevelBias << 4) | (PageTable32Bit << 8);
 	OutUniform[1].W = (SpaceID << 28);
 }
 
@@ -529,10 +530,7 @@ void FAllocatedVirtualTexture::GetPackedUniform(FUintVector4* OutUniform, uint32
 		OutUniform->X = FallbackColorPerTextureLayer[LayerIndex];
 		OutUniform->Y = BitcastFloatToUInt32((float)vPageSize * RcpPhysicalTextureSize);
 		OutUniform->Z = BitcastFloatToUInt32((float)PageBorderSize * RcpPhysicalTextureSize);
-		// Pack page table format bool as sign bit on page size.
-		const bool bPageTableExtraBits = GetPageTableFormat() == EVTPageTableFormat::UInt32;
-		const float PackedSignBit = bPageTableExtraBits ? 1.f : -1.f;
-		OutUniform->W = BitcastFloatToUInt32((float)pPageSize * RcpPhysicalTextureSize * PackedSignBit);
+		OutUniform->W = BitcastFloatToUInt32((float)pPageSize * RcpPhysicalTextureSize);
 	}
 	else
 	{
