Path Tracing: Fix shader warnings from removed payload member

#rnx
#Synced-CL 49608129

diff --git a/Engine/Shaders/Private/PathTracing/PathTracingMaterialHitShader.usf b/Engine/Shaders/Private/PathTracing/PathTracingMaterialHitShader.usf
--- a/Engine/Shaders/Private/PathTracing/PathTracingMaterialHitShader.usf
+++ b/Engine/Shaders/Private/PathTracing/PathTracingMaterialHitShader.usf
@@ -660,7 +660,6 @@ RAY_TRACING_ENTRY_CLOSEST_HIT(PathTracingMaterialCHS,
 							Payload.SpecularColor         = F0;
 							Payload.SpecularEdgeColor     = F90;
 							Payload.RoughnessData         = RoughnessData;
-							Payload.BottomNormalOct16bits = 0;
 							Payload.Anisotropy = BSDF_GETHASANISOTROPY(BSDF) ? SLAB_ANISOTROPY(BSDF) : 0.0;
 							Payload.FuzzAmount    = FuzzAmount;
 							Payload.FuzzColor     = FuzzColor;
@@ -814,7 +813,6 @@ RAY_TRACING_ENTRY_CLOSEST_HIT(PathTracingMaterialCHS,
 					Payload.SpecularColor         = 0.0f;
 					Payload.SpecularEdgeColor     = 0.0f;
 					Payload.RoughnessData         = 0.0f;
-					Payload.BottomNormalOct16bits = 0;
 
 					// Handle special-case specular features
 					Payload.Anisotropy			= 0.0f;
