Substrate - Toon - Fixed missing specular color application from evaluation.

#tests pc editor
#rnx
#Synced-CL 49646537

diff --git a/Engine/Shaders/Private/Substrate/SubstrateEvaluation.ush b/Engine/Shaders/Private/Substrate/SubstrateEvaluation.ush
--- a/Engine/Shaders/Private/Substrate/SubstrateEvaluation.ush
+++ b/Engine/Shaders/Private/Substrate/SubstrateEvaluation.ush
@@ -1333,7 +1333,7 @@ FSubstrateEvaluateResult SubstrateEvaluateBSDFCommon(FSubstrateBSDFContext BSDFC
 			// TOON_TODO experiment with GGX alike lob that is normalize according to peak lobe?
 			// TOON_TODO Area light, the roughness expand trick here does not make sense to use. MRP?
 
-			const float3 SpecularRamp = EvaluateToonProfileSpecularRamp(ToonProfileId, Specular);
+			const float3 SpecularRamp = SpecularColor * EvaluateToonProfileSpecularRamp(ToonProfileId, Specular);
 			Sample.SpecularPathValue = SpecularRamp;
 			Sample.IntegratedSpecularValue += ShadowTerms.SurfaceShadow * AreaLight.Falloff * AreaLight.FalloffColor * Sample.SpecularPathValue;
 
@@ -2030,6 +2030,8 @@ FSubstrateEnvLightResult SubstrateEvaluateForEnvLight(FSubstrateBSDFContext BSDF
 
 			SubstrateEnvLightResult.DiffuseWeight			= DiffuseColor;
 			SubstrateEnvLightResult.DiffuseColor			= DiffuseColor;
+
+			// TOON_TODO Specular support.
 		}
 		break;
 
