Path Tracing: Remove trace opaque first permutation from path tracer

Remove the experimental option to turn off "trace opaque first". This cuts the number of permutations in half when experimental options are enabled.

Also, enable translucent tracing by default for path tracing debug mode.

#Synced-CL 49691844

diff --git a/Engine/Shaders/Private/PathTracing/PathTracingCore.ush b/Engine/Shaders/Private/PathTracing/PathTracingCore.ush
--- a/Engine/Shaders/Private/PathTracing/PathTracingCore.ush
+++ b/Engine/Shaders/Private/PathTracing/PathTracingCore.ush
@@ -496,7 +496,7 @@ FPackedPathTracingPayload TraceTransparentRay(inout FPathState PathState, int Bo
 		PathState.PathThroughput = 0;
 		return Payload;
 	}
-#if PATH_TRACER_TRACE_OPAQUE_FIRST
+
 	FPackedPathTracingPayload OpaquePayload = InitPathTracingPayload(PathState.GetFirstScatterType(), ApproximateCaustics ? PathState.GetPathRoughness() : 0.0);
 	OpaquePayload.SetDBufferA(float4(0, 0, 0, 1));
 	OpaquePayload.SetDBufferB(float4(0, 0, 0, 1));
@@ -580,12 +580,6 @@ FPackedPathTracingPayload TraceTransparentRay(inout FPathState PathState, int Bo
 	{
 		PathState.Ray.TMax = asfloat(asuint(OpaquePayload.HitT) - 1);
 	}
-	
-#else // PATH_TRACER_TRACE_OPAQUE_FIRST
-
-	InstanceInclusionMask |= bUseCameraRayType ? PATHTRACER_MASK_CAMERA_TRANSLUCENT : PATHTRACER_MASK_INDIRECT_TRANSLUCENT;
-
-#endif
 
 	for (;;)
 	{
@@ -594,7 +588,6 @@ FPackedPathTracingPayload TraceTransparentRay(inout FPathState PathState, int Bo
 		PackedPayload.SetDBufferB(float4(0, 0, 0, 1));
 		PackedPayload.SetDBufferC(float4(0, 0, 0, 1));
 		PackedPayload.SetStochasticSlabRand(RISRandSample.z);
-#if PATH_TRACER_TRACE_OPAQUE_FIRST
 		if (OpaquePayload.IsHit())
 		{
 			const float3 ViewZ = View.ViewToTranslatedWorld[2].xyz;
@@ -602,7 +595,6 @@ FPackedPathTracingPayload TraceTransparentRay(inout FPathState PathState, int Bo
 			PackedPayload.SetFlag(PATH_TRACING_PAYLOAD_INPUT_FLAG_HAS_SCENE_DEPTH);
 			PackedPayload.SetSceneDepth(dot(TranslatedWorldPos, ViewZ));
 		}
-#endif
 
 #if PATH_TRACER_USE_SER 
 		{
@@ -636,13 +628,12 @@ FPackedPathTracingPayload TraceTransparentRay(inout FPathState PathState, int Bo
 			PackedPayload);
 #endif
 
-#if PATH_TRACER_TRACE_OPAQUE_FIRST
 		if (PackedPayload.IsMiss())
 		{
 			// If we miss the translucent geometry, the next hit is the opaque one
 			PackedPayload = OpaquePayload;
 		}
-#endif
+
 		FVolumeIntersectionList VolumeIsectList = VolumeIntersect(PathState.Ray.Origin, PathState.Ray.Direction, PathState.Ray.TMin, PathState.Ray.TMax, PathState.GetPathRoughness(), VolumeFlags);
 		if (VolumeIsectList.HitBlocker() &&
 			(PathState.GetFirstScatterType() == PATHTRACER_SCATTER_VOLUME ||  (VolumeFlags & PATH_TRACER_VOLUME_SHOW_PLANET_GROUND) != 0) &&
@@ -959,47 +950,6 @@ FPackedPathTracingPayload TraceTransparentRay(inout FPathState PathState, int Bo
 			break;
 		}
 
-#if PATH_TRACER_TRACE_OPAQUE_FIRST == 0
-		if (PackedPayload.IsDecalReceiver() && PathState.GetPathRoughness() < max(MeshDecalRoughnessCutoff, DecalRoughnessCutoff))
-		{
-			FDecalShaderPayload DecalPayload = InitDecalShaderPayload(PathState.Ray.Origin + PackedPayload.HitT * PathState.Ray.Direction);
-			EvaluateDecals(PathState, PackedPayload.HitT, DecalPayload);
-#if USE_DBUFFER
-			if (PackedPayload.UsesDBufferLookup() && HasDecalData(DecalPayload))
-			{
-				if (EnableDBuffer)
-				{
-					// Retrace ray with DBuffer data
-					PackedPayload = InitPathTracingPayload(PathState.GetFirstScatterType(), ApproximateCaustics ? PathState.GetPathRoughness() : 0.0);
-
-					PackedPayload.SetDBufferA(DecalPayload.GetBaseColor());
-					PackedPayload.SetDBufferB(DecalPayload.GetWorldNormal());
-					PackedPayload.SetDBufferC(DecalPayload.GetMetallicSpecularRoughness());
-					PackedPayload.SetStochasticSlabRand(RISRandSample.z); // Must use same random value to get a consistent result to the first call
-					TraceRay(
-						TLAS,
-						RayFlags,
-						InstanceInclusionMask,
-						RAY_TRACING_SHADER_SLOT_MATERIAL,
-						RAY_TRACING_NUM_SHADER_SLOTS,
-						MissShaderIndex,
-						PathState.Ray.GetNativeDesc(),
-						PackedPayload);
-				}
-				else
-				{
-					// The user asked for a dbuffer lookup, but we have disabled the double-trace.
-					// Set all the decal response flags so we at least pickup some kind of response
-					// even if it won't be exactly what the user had in mind.
-					PackedPayload.SetFlag(PATH_TRACING_PAYLOAD_OUTPUT_FLAG_DECAL_RESPONSE_COLOR | PATH_TRACING_PAYLOAD_OUTPUT_FLAG_DECAL_RESPONSE_NORMAL | PATH_TRACING_PAYLOAD_OUTPUT_FLAG_DECAL_RESPONSE_ROUGHNESS);
-				}
-			}
-#endif
-			// Now that we know the real decal data, apply it
-			ApplyDecal(DecalPayload, PackedPayload, PathState.Ray.Direction, PathState.GetPathRoughness());
-		}
-#endif
-
 		if (!PackedPayload.IsHoldout() && bIsCameraRay)
 		{
 			PathState.Alpha += Luminance(PathState.PathThroughput * (1.0 - PackedPayload.UnpackTransparencyColor()));
diff --git a/Engine/Source/Runtime/Renderer/Private/PathTracing.cpp b/Engine/Source/Runtime/Renderer/Private/PathTracing.cpp
--- a/Engine/Source/Runtime/Renderer/Private/PathTracing.cpp
+++ b/Engine/Source/Runtime/Renderer/Private/PathTracing.cpp
@@ -142,7 +142,7 @@ TAutoConsoleVariable<bool> CVarPathTracingVisualizeEnableEmissive(
 
 TAutoConsoleVariable<bool> CVarPathTracingVisualizeIncludeTranslucent(
 	TEXT("r.PathTracing.Visualize.IncludeTranslucent"),
-	false,
+	true,
 	TEXT("Indicates if translucent materials should be included in selected visualization mode."),
 	ECVF_RenderThreadSafe
 );
@@ -177,14 +177,6 @@ TAutoConsoleVariable<int32> CVarPathTracingCompactionDepth(
 	ECVF_RenderThreadSafe
 );
 
-TAutoConsoleVariable<bool> CVarPathTracingTraceOpaqueFirst(
-	TEXT("r.PathTracing.TraceOpaqueFirst"),
-	true,
-	TEXT("Trace opaque geometry before translucent geometry. This allows the path tracer to setup a correct depth for DepthFade based effects.\n")
-	TEXT("Requires r.PathTracing.Experimental=true to modify.\n"),
-	ECVF_RenderThreadSafe
-);
-
 TAutoConsoleVariable<bool> CVarPathTracingShaderExecutionReordering(
 	TEXT("r.PathTracing.ShaderExecutionReordering"),
 	true,
@@ -1292,9 +1284,8 @@ class FPathTracingRG : public FGlobalShader
 	class FCloudShader : SHADER_PERMUTATION_BOOL("PATH_TRACER_USE_CLOUD_SHADER");
 	class FSubstrateComplexSpecialMaterial : SHADER_PERMUTATION_BOOL("PATH_TRACER_USE_SUBSTRATE_SPECIAL_COMPLEX_MATERIAL");
 	class FUseSER : SHADER_PERMUTATION_BOOL("PATH_TRACER_USE_SER");
-	class FTraceOpaqueFirst : SHADER_PERMUTATION_BOOL("PATH_TRACER_TRACE_OPAQUE_FIRST");
 	class FNeedTMinWorkaround : SHADER_PERMUTATION_BOOL("NEED_TMIN_WORKAROUND");
-	using FPermutationDomain = TShaderPermutationDomain<FCompactionType, FAdaptiveSampling, FCloudShader, FSubstrateComplexSpecialMaterial, FUseSER, FTraceOpaqueFirst, FNeedTMinWorkaround>;
+	using FPermutationDomain = TShaderPermutationDomain<FCompactionType, FAdaptiveSampling, FCloudShader, FSubstrateComplexSpecialMaterial, FUseSER, FNeedTMinWorkaround>;
 
 	static bool ShouldCompilePermutation(const FGlobalShaderPermutationParameters& Parameters)
 	{
@@ -1307,10 +1298,6 @@ class FPathTracingRG : public FGlobalShader
 			{
 				return false;
 			}
-			if (!PermutationVector.Get<FTraceOpaqueFirst>())
-			{
-				return false;
-			}
 			if (!PermutationVector.Get<FCompactionType>())
 			{
 				return false;
@@ -3304,7 +3291,6 @@ static FPathTracingRG::FPermutationDomain GetPathTracingRGPermutation(const FVie
 	const bool bUseCloudShader = uint32(View.PathTracingVolumetricCloudCallableShaderIndex) < Scene.RayTracingSBT.NumCallableShaderSlots;
 	const bool bHasComplexSpecialRenderPath = Substrate::IsSubstrateEnabled() && UsesSubstrateTileType(View.GetSceneTextures().SubstrateSceneData->UsesTileTypeMask, ESubstrateTileType::EComplexSpecial);
 	const bool bUseSER = GRHIGlobals.SupportsShaderExecutionReordering && CVarPathTracingShaderExecutionReordering.GetValueOnRenderThread();
-	const bool bUseTraceOpaqueFirst = !bUseExperimental || CVarPathTracingTraceOpaqueFirst.GetValueOnRenderThread();
 
 	FPathTracingRG::FPermutationDomain Out;
 	Out.Set<FPathTracingRG::FCompactionType>(bUseCompaction);
@@ -3312,7 +3298,6 @@ static FPathTracingRG::FPermutationDomain GetPathTracingRGPermutation(const FVie
 	Out.Set<FPathTracingRG::FCloudShader>(bUseCloudShader);
 	Out.Set<FPathTracingRG::FSubstrateComplexSpecialMaterial>(bHasComplexSpecialRenderPath);
 	Out.Set<FPathTracingRG::FUseSER>(bUseSER);
-	Out.Set<FPathTracingRG::FTraceOpaqueFirst>(bUseTraceOpaqueFirst);
 	Out.Set<FPathTracingRG::FNeedTMinWorkaround>(GPUNeedsTMinWorkaround());
 	return Out;
 }
