CompositeReflectionCapturesAndSkylightTWS only applies Reflection Captures for projects that have r.AllowStaticLighting=1
* Saves 139 instructions from Translucent Surface materials in projects with r.AllowStaticLighting=0

#Synced-CL 49613554

diff --git a/Engine/Shaders/Private/ReflectionEnvironmentComposite.ush b/Engine/Shaders/Private/ReflectionEnvironmentComposite.ush
--- a/Engine/Shaders/Private/ReflectionEnvironmentComposite.ush
+++ b/Engine/Shaders/Private/ReflectionEnvironmentComposite.ush
@@ -26,6 +26,7 @@ float3 CompositeReflectionCapturesAndSkylightTWS(
 	float4 ImageBasedReflections = float4(0, 0, 0, CompositeAlpha);
 	float2 CompositedAverageBrightness = float2(0.0f, 1.0f);
 
+#if ALLOW_STATIC_LIGHTING
 #if REFLECTION_COMPOSITE_USE_BLENDED_REFLECTION_CAPTURES
 	// Accumulate reflections from captures affecting this tile, applying largest captures first so that the smallest ones display on top
 	LOOP
@@ -149,6 +150,8 @@ float3 CompositeReflectionCapturesAndSkylightTWS(
 	ImageBasedReflections.rgb *= View.PrecomputedIndirectSpecularColorScale;
 	CompositedAverageBrightness.x *= Luminance( View.PrecomputedIndirectSpecularColorScale );
 
+#endif
+
 #if ENABLE_SKY_LIGHT
 
 	BRANCH  
