Fix physics field rendering offset + crash with the show field option

#Synced-CL 49728154

diff --git a/Engine/Shaders/Private/PhysicsFieldVisualizer.usf b/Engine/Shaders/Private/PhysicsFieldVisualizer.usf
--- a/Engine/Shaders/Private/PhysicsFieldVisualizer.usf
+++ b/Engine/Shaders/Private/PhysicsFieldVisualizer.usf
@@ -50,8 +50,9 @@ void FillRayInfos(in uint2 ScreenCoord, inout float3 RayBegin, inout float3 RayE
 	const float2 UV = PixelCoord / float2(OutputResolution);
 	const float ClosestDepth = ConvertFromDeviceZ(0.0000001f);//ConvertFromDeviceZ(SceneDepthTexture.Load(uint3(PixelCoord, 0)));
 
-	RayBegin = PrimaryView.TranslatedWorldCameraOrigin;
-	RayEnd = ReconstructTranslatedWorldPositionFromDepth(UV, ClosestDepth);
+	const float3 ViewOffset = DFDemote(PrimaryView.PreViewTranslation);
+	RayBegin = PrimaryView.TranslatedWorldCameraOrigin - ViewOffset;
+	RayEnd = ReconstructTranslatedWorldPositionFromDepth(UV, ClosestDepth) - ViewOffset;
 	RayDirection = (RayEnd - RayBegin);
 }
 
@@ -346,7 +347,8 @@ void MainCS(uint3 DispatchThreadId : SV_DispatchThreadID,
 		{
 			for(int BoundsIndex = BoundsOffset, BoundsEnd = BoundsOffset + BoundsSize; BoundsIndex < BoundsEnd; ++BoundsIndex)
 			{
-				AddAABBWS(BoundsMin[BoundsIndex].xyz, BoundsMax[BoundsIndex].xyz, float4(AABBColor.Color.x, AABBColor.Color.y, AABBColor.Color.z, 1));
+				AddAABBWS(BoundsMin[BoundsIndex].xyz, BoundsMax[BoundsIndex].xyz, 
+					float4(AABBColor.Color.x, AABBColor.Color.y, AABBColor.Color.z, 1));
 			}
 #if PERMUTATION_FIELD == 2
 			PrintIntegerField();
diff --git a/Engine/Source/Runtime/Engine/Private/PhysicsField/PhysicsFieldComponent.cpp b/Engine/Source/Runtime/Engine/Private/PhysicsField/PhysicsFieldComponent.cpp
--- a/Engine/Source/Runtime/Engine/Private/PhysicsField/PhysicsFieldComponent.cpp
+++ b/Engine/Source/Runtime/Engine/Private/PhysicsField/PhysicsFieldComponent.cpp
@@ -971,11 +971,6 @@ void UPhysicsFieldComponent::TickComponent(float DeltaTime, enum ELevelTick Tick
 	// Persistent ones will be destrtoyed on request by the field component itslef through the RemovePersistentCommand method
 	TransientCommands[(uint8)(EFieldCommandBuffer::CPUWriteBuffer)].Reset();
 
-	if (GetWorld() && GetWorld()->Scene)
-	{
-		GetWorld()->Scene->ShowPhysicsField();
-	}
-
 	MarkRenderDynamicDataDirty();
 }
 
diff --git a/Engine/Source/Runtime/Renderer/Private/RendererScene.cpp b/Engine/Source/Runtime/Renderer/Private/RendererScene.cpp
--- a/Engine/Source/Runtime/Renderer/Private/RendererScene.cpp
+++ b/Engine/Source/Runtime/Renderer/Private/RendererScene.cpp
@@ -2515,7 +2515,7 @@ void FScene::SetPhysicsField(FPhysicsFieldSceneProxy* PhysicsFieldSceneProxy)
 
 void FScene::ShowPhysicsField()
 {
-	// Set the shader print/debug values from game thread if
+	// Set the shader print/debug values if
 	// physics field visualisation has been enabled
 	if (PhysicsField && PhysicsField->FieldResource && PhysicsField->FieldResource->FieldInfos.bShowFields)
 	{
@@ -2538,12 +2538,13 @@ void FScene::ResetPhysicsField()
 
 void FScene::UpdatePhysicsField(FRDGBuilder& GraphBuilder, FViewInfo& View)
 {
-	if (PhysicsField)
+	if (PhysicsField && PhysicsField->FieldResource)
 	{
 		PhysicsField->FieldResource->FieldInfos.ViewOrigin = View.ViewMatrices.GetViewOrigin();
 		if (View.Family )
 		{
 			PhysicsField->FieldResource->FieldInfos.bShowFields = View.Family->EngineShowFlags.PhysicsField;
+			ShowPhysicsField();
 		}
 	}
 }
