Path Tracing: Fix transmission math on two-sided foliage

There were some inconsistencies introduced in the shading normal refactor.

#rnx
#Synced-CL 49766554

diff --git a/Engine/Shaders/Private/PathTracing/Material/PathTracingSubstrateFoliage.ush b/Engine/Shaders/Private/PathTracing/Material/PathTracingSubstrateFoliage.ush
--- a/Engine/Shaders/Private/PathTracing/Material/PathTracingSubstrateFoliage.ush
+++ b/Engine/Shaders/Private/PathTracing/Material/PathTracingSubstrateFoliage.ush
@@ -94,7 +94,7 @@ FMaterialSample SubstrateFoliage_SampleMaterial(
 		RandSample.x = RescaleRandomNumber(RandSample.x, bSampledTransmit ? 0.0 : TransmitCdf, bSampledTransmit ? TransmitCdf : Data.LobeCdf.x);
 		// diffuse lobe
 		Result.Direction = CosineSampleHemisphereWorld(RandSample.xy, Payload.WorldNormal * (bSampledTransmit ? -1 : 1)).xyz;
-		L = mul(Result.Direction, Data.Basis);
+		L = mul(Data.Basis, Result.Direction);
 		H = normalize(L + V);
 	}
 	else if (bSampledSpecular)
@@ -112,7 +112,7 @@ FMaterialSample SubstrateFoliage_SampleMaterial(
 			// invalid output direction, exit early
 			return NullMaterialSample();
 		}
-		Result.Direction = mul(Data.Basis, L);
+		Result.Direction = mul(L, Data.Basis);
 	}
 	else
 	{
@@ -120,11 +120,11 @@ FMaterialSample SubstrateFoliage_SampleMaterial(
 		RandSample.x = RescaleRandomNumber(RandSample.x, Data.LobeCdf.z, 1.0);
 		L = Data.Fuzz.Sample(RandSample.xy);
 		H = normalize(L + V);
-		Result.Direction = mul(Data.Basis, L);
+		Result.Direction = mul(L, Data.Basis);
 	}
 
 	// transform to world space
-	const float3 L_World = normalize(mul(L, Data.Basis));
+	const float3 L_World = Result.Direction;
 	const float NoV = saturate(V.z);
 	const float NoL = abs(L.z);
 	const float VoH = saturate(dot(V, H));
@@ -143,6 +143,7 @@ FMaterialSample SubstrateFoliage_SampleMaterial(
 	if (bSampledTransmit)
 	{
 		Result.AddLobeWithMIS(Data.DiffuseWeight * Payload.GetFoliageTransmissionColor(), DiffPdf, Data.LobePdf.x * Data.TransmissionProb);
+		Result.PositionBiasSign = -1;
 	}
 	else if (bSampledDiffuse)
 	{
