Fix incorrect velocity Scene Texture Data indexing in post processing materials when viewport size ~= buffer size

#tests Test on UDN user provided example and match showflag.visualizemotionblur 2
#Synced-CL 49128105

diff --git a/Engine/Shaders/Private/MaterialTemplate.ush b/Engine/Shaders/Private/MaterialTemplate.ush
--- a/Engine/Shaders/Private/MaterialTemplate.ush
+++ b/Engine/Shaders/Private/MaterialTemplate.ush
@@ -3174,7 +3174,7 @@ MaterialFloat4 MobileSceneTextureLookup(inout FMaterialPixelParameters Parameter
 #if SHADING_PATH_DEFERRED
 
 #if POST_PROCESS_MATERIAL
-/** Samples the screen-space velocity for the specified UV coordinates. */
+/** Samples the screen-space velocity for the specified UV coordinates in Buffer Space. */
 float2 PostProcessVelocityLookup(float Depth, float2 UV)
 {
 	float4 EncodedVelocity = Texture2DSampleLevel(SceneTexturesStruct.GBufferVelocityTexture, SceneTexturesStruct_GBufferVelocityTextureSampler, UV, 0);
@@ -3186,7 +3186,8 @@ float2 PostProcessVelocityLookup(float Depth, float2 UV)
 	}
 	else
 	{
-		float2 ScreenPos = ViewportUVToScreenPos(UV);
+		float2 ViewportUV = BufferUVToViewportUV(UV);
+		float2 ScreenPos = ViewportUVToScreenPos(ViewportUV);
 		float4 ThisClip = float4( ScreenPos, Depth, 1 );
 		float4 PrevClip = mul( ThisClip, View.ClipToPrevClip );
 		float2 PrevScreen = PrevClip.xy / PrevClip.w;
