[Substrate, Ray traced translucency] Do not force holdout hits to opaque (i.e. TransmittancePreCoverage=0 and Coverage=1) to better match holdout blending behavior in raster. For example, a holdout material with 0.5 opacity (or transmittance color) appears 50% opaque in scene color and alpha when rastered but it appears solid black when rendered using RT translucency

#tests tested in editor
#Synced-CL 49641444

diff --git a/Engine/Shaders/Private/RayTracing/RayTracingMaterialHitShaders.usf b/Engine/Shaders/Private/RayTracing/RayTracingMaterialHitShaders.usf
--- a/Engine/Shaders/Private/RayTracing/RayTracingMaterialHitShaders.usf
+++ b/Engine/Shaders/Private/RayTracing/RayTracingMaterialHitShaders.usf
@@ -758,8 +758,6 @@ RAY_TRACING_ENTRY_CLOSEST_HIT(MaterialCHS,
 		Payload.SubstrateData = (FSubstrateRaytracingPayload)0;
 		Payload.ShadingModelID = SHADINGMODELID_UNLIT;
 		Payload.Radiance = 0.0;
-		Payload.Opacity = 1.0;
-		Payload.TransmittancePreCoverage = 0.0;
 #else
 		Payload.ShadingModelID = SHADINGMODELID_UNLIT;
 		Payload.Radiance = 0.0;
