Add Ubershader macro for condition block in shader.

#Synced-CL 49800659

diff --git a/Engine/Shaders/Private/Common.ush b/Engine/Shaders/Private/Common.ush
--- a/Engine/Shaders/Private/Common.ush
+++ b/Engine/Shaders/Private/Common.ush
@@ -2519,3 +2519,15 @@ float4 GatherDepth(Texture2D Texture, float2 UV)
 #endif
 
 #endif // SUPPORTS_INDEPENDENT_SAMPLERS
+
+// ---------------------------------------------------- Ubershader
+
+#ifndef UBERSHADER
+#define UBERSHADER 0
+#endif
+
+#if UBERSHADER
+#define IF_UBERSHADER(X) if(X)
+#else
+#define IF_UBERSHADER(X)
+#endif
\ No newline at end of file
diff --git a/Engine/Shaders/Private/DeferredLightPixelShaders.usf b/Engine/Shaders/Private/DeferredLightPixelShaders.usf
--- a/Engine/Shaders/Private/DeferredLightPixelShaders.usf
+++ b/Engine/Shaders/Private/DeferredLightPixelShaders.usf
@@ -116,7 +116,12 @@ uint bLightingChannelEnabled;
 uint GetLightingChannelMask(float2 UV)
 {
 	uint2 IntegerUV = UV * View.BufferSizeAndInvSize.xy;
-	return bLightingChannelEnabled ? LightingChannelsTexture.Load(uint3(IntegerUV, 0)).x : 0xFFu;
+	uint Out = 0xFFu;
+	IF_UBERSHADER(bLightingChannelEnabled)
+	{
+		Out = LightingChannelsTexture.Load(uint3(IntegerUV, 0)).x;
+	}
+	return Out;
 }
 
 float GetExposure()
@@ -179,14 +184,14 @@ float3 GetAtmosphereAndCloudAttenuation(float3 StableTranslatedLargeWorldPositio
 	float3 Out = 1;
 
 	#if USE_ATMOSPHERE_TRANSMITTANCE
-	if (bAtmosphereTransmittanceEnabled)
+	IF_UBERSHADER(bAtmosphereTransmittanceEnabled)
 	{
 		Out *= GetAtmosphericLightTransmittance(StableTranslatedLargeWorldPosition, ScreenUV, DeferredLightUniforms.Direction.xyz);
 	}
 	#endif
 
 	#if USE_CLOUD_TRANSMITTANCE
-	if (bCloudShadowEnabled)
+	IF_UBERSHADER(bCloudShadowEnabled)
 	{
 		float OutOpticalDepth = 0.0f;
 		Out *= lerp(1.0f, GetCloudVolumetricShadow(StableTranslatedLargeWorldPosition, CloudShadowmapTranslatedWorldToLightClipMatrix, CloudShadowmapFarDepthKm, CloudShadowmapTexture, CloudShadowmapSampler, OutOpticalDepth), CloudShadowmapStrength);
diff --git a/Engine/Source/Runtime/Renderer/Private/LightRendering.cpp b/Engine/Source/Runtime/Renderer/Private/LightRendering.cpp
--- a/Engine/Source/Runtime/Renderer/Private/LightRendering.cpp
+++ b/Engine/Source/Runtime/Renderer/Private/LightRendering.cpp
@@ -915,7 +915,7 @@ class FDeferredLightPS : public FGlobalShader
 	class FCloudTransmittance 	: SHADER_PERMUTATION_BOOL("USE_CLOUD_TRANSMITTANCE");
 	class FAnistropicMaterials 	: SHADER_PERMUTATION_BOOL("SUPPORTS_ANISOTROPIC_MATERIALS");
 	class FVirtualShadowMapMask : SHADER_PERMUTATION_BOOL("USE_VIRTUAL_SHADOW_MAP_MASK");
-	class FUbershader 			: SHADER_PERMUTATION_BOOL("USE_UBERSHADER");
+	class FUbershader 			: SHADER_PERMUTATION_BOOL("UBERSHADER");
 
 	using FPermutationDomain = TShaderPermutationDomain<
 		FSourceShapeDim,
