Fix PT hair strands tangent direction.

The code was done by [REDACTED_PERSON]

#Synced-CL 49800659

diff --git a/Engine/Shaders/Private/HairStrands/HairStrandsMaterialCommon.ush b/Engine/Shaders/Private/HairStrands/HairStrandsMaterialCommon.ush
--- a/Engine/Shaders/Private/HairStrands/HairStrandsMaterialCommon.ush
+++ b/Engine/Shaders/Private/HairStrands/HairStrandsMaterialCommon.ush
@@ -236,15 +236,9 @@ FEvaluateOutput Evaluate(
 		
 	#if HAIRSTRANDS_HAS_NORMAL_CONNECTED
 	  #if MATERIAL_TANGENTSPACENORMAL 
-		// Convert the tangent space input tangent 
-		// * from a frame with Z/(0,0,1) being up and Y/(0,1,0) being the main tangent direction (called green)
-		// * to   a frame with Y/(0,1,0) being up and Z/(0,0,1) being the main tangent direction (called blue)
 		// Tangent is in world space. First transform it into local space prior rotation)
 		float3 LocalTangent = mul(Tangent, transpose(MaterialParameters.TangentToWorld));
-		const float3 TangentSpaceGreen = LocalTangent;
-		const float3 TangentSpaceBlue = float3(-TangentSpaceGreen.x, TangentSpaceGreen.z, TangentSpaceGreen.y);
-		OutSample.Tangent = mul(TangentSpaceBlue, MaterialParameters.TangentToWorld);
-		OutSample.Tangent = normalize(OutSample.Tangent);
+		OutSample.Tangent = normalize(mul(RotateLocalHairTangent(LocalTangent), MaterialParameters.TangentToWorld));
 	  #else // Tangent in World space (likely provided by the HairAttributes node)
 		OutSample.Tangent = Tangent;
 	  #endif
@@ -262,13 +256,7 @@ FEvaluateOutput Evaluate(
 
 	#if HAIRSTRANDS_HAS_NORMAL_CONNECTED
 	  #if MATERIAL_TANGENTSPACENORMAL 
-		// Convert the tangent space input tangent 
-		// * from a frame with Z/(0,0,1) being up and Y/(0,1,0) being the main tangent direction (called green)
-		// * to   a frame with Y/(0,1,0) being up and Z/(0,0,1) being the main tangent direction (called blue)
-		const float3 TangentSpaceGreen = Tangent;
-		const float3 TangentSpaceBlue = float3(-TangentSpaceGreen.x, TangentSpaceGreen.z, TangentSpaceGreen.y);
-		OutSample.Tangent = mul(TangentSpaceBlue, MaterialParameters.TangentToWorld);
-		OutSample.Tangent = normalize(OutSample.Tangent);
+		OutSample.Tangent = normalize(mul(RotateLocalHairTangent(Tangent), MaterialParameters.TangentToWorld));
 	  #else // Tangent in World space (likely provided by the HairAttributes node)
 		OutSample.Tangent = Tangent;
 	  #endif
diff --git a/Engine/Shaders/Private/HairStrands/HairStrandsVertexFactory.ush b/Engine/Shaders/Private/HairStrands/HairStrandsVertexFactory.ush
--- a/Engine/Shaders/Private/HairStrands/HairStrandsVertexFactory.ush
+++ b/Engine/Shaders/Private/HairStrands/HairStrandsVertexFactory.ush
@@ -147,6 +147,29 @@ FVertexInfo GetVertexInfo(FVertexFactoryInput Input)
 	return VertexInfo;
 }
 
+float3 RotateLocalHairTangent(float3 In)
+{
+	// Convert the tangent space input tangent 
+	// * from a frame with Z/(0,0,1) being up and Y/(0,1,0) being the main tangent direction (called green)
+	// * to   a frame with Y/(0,1,0) being up and Z/(0,0,1) being the main tangent direction (called blue)
+	return float3(-In.x, In.z, In.y);
+}
+
+// Unlike AssembleTangentToWorld, this helper first adjusts the TangentToWorld0 direction
+// using TangentToWorld2 to ensure they are perpendicular before constructing the matrix.
+half3x3 AssembleTangentToWorldForHairStrand(half3 TangentToWorld0, half4 TangentToWorld2)
+{
+	// Derive the third basis vector off of the other two.
+	// Flip based on the determinant sign
+	const half3 TangentToWorld1 = normalize(cross(TangentToWorld2.xyz,TangentToWorld0) * TangentToWorld2.w);
+
+	// Correct TangentToWorld0 direct in case it is not perpendicular.
+	const half3 CorrectedTangentToWorld0 = normalize(cross(TangentToWorld1, TangentToWorld2.xyz));
+
+	// Transform from tangent space to world space
+	return half3x3(CorrectedTangentToWorld0, TangentToWorld1, TangentToWorld2.xyz);
+}
+
 #if RAYHITGROUPSHADER
 FVertexFactoryInput LoadVertexFactoryInputForHGS(uint TriangleIndex, int VertexIndex)
 {
@@ -286,7 +309,7 @@ FMaterialPixelParameters GetMaterialPixelParameters(FVertexFactoryInterpolantsVS
 	half4 TangentToWorld2 = GetTangentToWorld2(Interpolants);
 	Result.UnMirrored = TangentToWorld2.w;
 
-	Result.TangentToWorld = AssembleTangentToWorld( TangentToWorld0, TangentToWorld2 );
+	Result.TangentToWorld = AssembleTangentToWorldForHairStrand( TangentToWorld0, TangentToWorld2 );
 #if USE_WORLDVERTEXNORMAL_CENTER_INTERPOLATION
 	Result.WorldVertexNormal_Center = Interpolants.TangentToWorld2_Center.xyz;
 #endif
@@ -644,7 +667,7 @@ FMaterialPixelParameters GetMaterialPixelParameters(float3 RayOrigin, float3 Ray
 	float3 TangentToWorld2 = normalize(lerp(TangentToWorld_V0[2], TangentToWorld_V1[2], UCoord));
 	float TangentSign = lerp(TangentSign_V0, TangentSign_V1, UCoord) * GetPrimitive_DeterminantSign_FromFlags(SceneData.Primitive.Flags);
 
-	Result.TangentToWorld = AssembleTangentToWorld(TangentToWorld0, float4(TangentToWorld2, TangentSign));
+	Result.TangentToWorld = AssembleTangentToWorldForHairStrand(TangentToWorld0, float4(TangentToWorld2, TangentSign));
 
 	// compute a smooth ribbon normal based on the tangent vector
 	float3 T = Result.TangentToWorld[2];
diff --git a/Engine/Shaders/Private/MaterialTemplate.ush b/Engine/Shaders/Private/MaterialTemplate.ush
--- a/Engine/Shaders/Private/MaterialTemplate.ush
+++ b/Engine/Shaders/Private/MaterialTemplate.ush
@@ -4268,6 +4268,10 @@ float GetRoughnessFromNormalCurvature(FMaterialPixelParameters InMaterialParamet
 #endif
 }
 
+#if HAIR_STRAND_MESH_FACTORY
+float3 RotateLocalHairTangent(float3 In);
+#endif
+
 SHADER_PUSH_WARNINGS_STATE
 SHADER_DISABLE_WARNINGS
 
@@ -4314,6 +4318,10 @@ SHADER_POP_WARNINGS_STATE
 
 #if MATERIAL_TANGENTSPACENORMAL
 
+#if RAYHITGROUPSHADER && HAIR_STRAND_MESH_FACTORY && HAIRSTRANDS_HAS_NORMAL_CONNECTED
+	MaterialNormal = RotateLocalHairTangent(MaterialNormal);
+#endif
+
 #if FEATURE_LEVEL >= FEATURE_LEVEL_SM4
 	// Mobile will rely on only the final normalize for performance
 	MaterialNormal = normalize(MaterialNormal);
@@ -4382,6 +4390,10 @@ SHADER_POP_WARNINGS_STATE
 
 #if MATERIAL_TANGENTSPACENORMAL
 
+#if RAYHITGROUPSHADER && HAIR_STRAND_MESH_FACTORY && HAIRSTRANDS_HAS_NORMAL_CONNECTED
+	MaterialNormal = RotateLocalHairTangent(MaterialNormal);
+#endif
+
 #if FEATURE_LEVEL >= FEATURE_LEVEL_SM4
 	// Mobile will rely on only the final normalize for performance
 	MaterialNormal = normalize(MaterialNormal);
diff --git a/Engine/Shaders/Private/PathTracing/Material/PathTracingHair.ush b/Engine/Shaders/Private/PathTracing/Material/PathTracingHair.ush
--- a/Engine/Shaders/Private/PathTracing/Material/PathTracingHair.ush
+++ b/Engine/Shaders/Private/PathTracing/Material/PathTracingHair.ush
@@ -10,6 +10,12 @@
 
 // Shared routines
 
+// wrapper function to get hair strand direction, it is encoded in `WorldSmoothNormal` in offline path tracer
+float3 GetHairTangent(in const FPathTracingPayload Payload)
+{
+	return Payload.WorldSmoothNormal;
+}
+
 float3x3 GetHairBasis(float3 T_World, float3 V_World)
 {
 	// a convenient basis for hair shading calculations:
@@ -50,7 +56,7 @@ FMaterialSample Hair_SampleMaterial(
 	FPathTracingPayload Payload,
 	float3 RandSample)
 {
-	float3x3 HairBasis = GetHairBasis(Payload.WorldTangent, V_World);
+	const float3x3 HairBasis = GetHairBasis(GetHairTangent(Payload), V_World);
 
 	// sample a random position across the width of the fiber
 	const float h = 2 * RandSample.x - 1;
@@ -76,7 +82,7 @@ FMaterialEval Hair_EvalMaterial(
 	float2 DiffuseSpecularScale
 )
 {
-	const float3x3 HairBasis = GetHairBasis(Payload.WorldTangent, V_World);
+	const float3x3 HairBasis = GetHairBasis(GetHairTangent(Payload), V_World);
 
 	const float3 L = mul(HairBasis, L_World);
 
@@ -291,7 +297,7 @@ FMaterialSample Hair_SampleMaterial(
 	FPathTracingPayload Payload,
 	float3 RandSample)
 {
-	const float3x3 HairBasis = GetHairBasis(Payload.WorldTangent, V_World);
+	const float3x3 HairBasis = GetHairBasis(GetHairTangent(Payload), V_World);
 
 	const float3 V = mul(HairBasis, V_World); // TODO: simplify this since HairBasis is already sort of aligned with V ?
 	const float h = 2 * Payload.GetHairPrimitiveUV().y - 1; // remap back to [-1,1]
@@ -369,7 +375,7 @@ FMaterialEval Hair_EvalMaterial(
 	float2 DiffuseSpecularScale
 )
 {
-	const float3x3 HairBasis = GetHairBasis(Payload.WorldTangent, V_World);
+	const float3x3 HairBasis = GetHairBasis(GetHairTangent(Payload), V_World);
 
 	const float3 V = mul(HairBasis, V_World);
 	const float3 L = mul(HairBasis, L_World);
diff --git a/Engine/Shaders/Private/PathTracing/PathTracingDebug.usf b/Engine/Shaders/Private/PathTracing/PathTracingDebug.usf
--- a/Engine/Shaders/Private/PathTracing/PathTracingDebug.usf
+++ b/Engine/Shaders/Private/PathTracing/PathTracingDebug.usf
@@ -857,7 +857,7 @@ FPathTracingResult PathTracingDebugCore(FRayDesc Ray, uint2 LaunchIndex)
 	FPathTracingResult Result = (FPathTracingResult) 0;
 	const uint RayFlags = DebugMode == PATH_TRACER_DEBUG_VIZ_HITKIND ? 0 : RAY_FLAG_CULL_BACK_FACING_TRIANGLES;
 	const uint MissShaderIndex = 0;
-	uint InstanceInclusionMask = PATHTRACER_MASK_CAMERA;
+	uint InstanceInclusionMask = PATHTRACER_MASK_CAMERA | PATHTRACER_MASK_HAIR_CAMERA;
 	int TransparentHitCount = 0;
 
 	const float RandSample = RandomSequenceCreate(uint3(LaunchIndex, View.StateFrameIndex), 0, 1).Get1D();
diff --git a/Engine/Shaders/Private/PathTracing/PathTracingMaterialHitShader.usf b/Engine/Shaders/Private/PathTracing/PathTracingMaterialHitShader.usf
--- a/Engine/Shaders/Private/PathTracing/PathTracingMaterialHitShader.usf
+++ b/Engine/Shaders/Private/PathTracing/PathTracingMaterialHitShader.usf
@@ -884,11 +884,10 @@ RAY_TRACING_ENTRY_CLOSEST_HIT(PathTracingMaterialCHS,
 #endif	
 
 #if HAIR_STRAND_MESH_FACTORY
-	Payload.WorldSmoothNormal = Payload.WorldNormal = Payload.WorldGeoNormal;
+	Payload.WorldNormal = Payload.WorldGeoNormal;
 	Payload.WorldTangent = MaterialParameters.TangentToWorld[2];
-	// Take into account per-pixel normal shading variation
-	// Note: For hair shading the tangent is actually stored as the normal input
-	Payload.WorldTangent = normalize(MaterialParameters.WorldNormal);
+	// Repurpose WorldSmoothNormal to encode hair strand direction
+	Payload.WorldSmoothNormal = normalize(MaterialParameters.WorldNormal);
 #endif
 
 #if PATH_TRACER_USE_VOLUMETRIC_SHADING_MODEL
diff --git a/Engine/Source/Runtime/Renderer/Private/PathTracing.cpp b/Engine/Source/Runtime/Renderer/Private/PathTracing.cpp
--- a/Engine/Source/Runtime/Renderer/Private/PathTracing.cpp
+++ b/Engine/Source/Runtime/Renderer/Private/PathTracing.cpp
@@ -2435,6 +2435,7 @@ public:
 		OutEnvironment.SetDefine(TEXT("SCENE_TEXTURES_DISABLED"), 1);
 		OutEnvironment.SetDefine(TEXT("SIMPLIFIED_MATERIAL_SHADER"), IsGPULightmass);
 		OutEnvironment.SetDefine(TEXT("SUBSTRATE_USE_FULLYSIMPLIFIED_MATERIAL"), IsGPULightmass || SimplifySubstrate);
+		OutEnvironment.SetDefine(TEXT("HAIRSTRANDS_HAS_NORMAL_CONNECTED"), Parameters.MaterialParameters.bHasNormalConnected ? 1 : 0);
 		OutEnvironment.CompilerFlags.Add(CFLAG_ForceDXC);
 		OutEnvironment.CompilerFlags.Add(CFLAG_HLSL2021);
 		FMeshMaterialShader::ModifyCompilationEnvironment(Parameters, OutEnvironment);
diff --git a/Engine/Source/Runtime/Renderer/Private/RayTracing/RayTracingMaterialHitShaders.cpp b/Engine/Source/Runtime/Renderer/Private/RayTracing/RayTracingMaterialHitShaders.cpp
--- a/Engine/Source/Runtime/Renderer/Private/RayTracing/RayTracingMaterialHitShaders.cpp
+++ b/Engine/Source/Runtime/Renderer/Private/RayTracing/RayTracingMaterialHitShaders.cpp
@@ -187,6 +187,7 @@ public:
 		OutEnvironment.SetDefine(TEXT("USE_MATERIAL_INTERSECTION_SHADER"), UseIntersectionShader ? 1 : 0);
 		OutEnvironment.SetDefine(TEXT("USE_RAYTRACED_TEXTURE_RAYCONE_LOD"), UseRayConeTextureLod ? 1 : 0);
 		OutEnvironment.SetDefine(TEXT("SCENE_TEXTURES_DISABLED"), 1);
+		OutEnvironment.SetDefine(TEXT("HAIRSTRANDS_HAS_NORMAL_CONNECTED"), Parameters.MaterialParameters.bHasNormalConnected ? 1 : 0);
 		LightMapPolicyType::ModifyCompilationEnvironment(Parameters, OutEnvironment);
 		FMeshMaterialShader::ModifyCompilationEnvironment(Parameters, OutEnvironment);
 		const bool VirtualTextureLightmaps = UseVirtualTextureLightmap(Parameters.Platform);
