Removed remnant from landscape XY offset feature

#Synced-CL 49783950

diff --git a/Engine/Shaders/Private/LandscapeVertexFactory.ush b/Engine/Shaders/Private/LandscapeVertexFactory.ush
--- a/Engine/Shaders/Private/LandscapeVertexFactory.ush
+++ b/Engine/Shaders/Private/LandscapeVertexFactory.ush
@@ -79,22 +79,22 @@ uint GetRayTracingSectionVertCountForLOD(uint ComponentIndex)
 	return ((LandscapeParameters.RayTracingSectionSizeVerts >> Lod) + 1);
 }
 
+
 /**
- * x = Heightmap texture LOD difference from current LOD to highest LOD
- * y = todo: XYOffset texture LOD difference from current LOD to highest LOD
+ * Return Heightmap texture LOD difference from current LOD to highest LOD
  */
-float2 GetLodBias(uint ComponentIndex)
+float GetLodBias(uint ComponentIndex)
 {
 #if FIXED_GRID
 	// Fixed grid doesn't use LodBias
 	return 0;
 #else
-	float LodBias = LandscapeContinuousLODParameters.SectionLODBias[ComponentIndex];
-	return float2(LodBias, 0.f);
+	return LandscapeContinuousLODParameters.SectionLODBias[ComponentIndex];
 #endif
 }
 
 
+
 float GetNeighborSectionLodFromOffset(int2 NeighborOffset, float CenterLOD)
 {
 #if FIXED_GRID
@@ -311,7 +311,7 @@ struct FVertexFactoryIntermediates
 	float3 WorldNormal;
 	uint ComponentIndex;
 	float4 LodValues;
-	float2 LodBias;
+	float LodBias;
 	/** Cached primitive and instance data */
 	FSceneDataIntermediates SceneData;
 };
@@ -668,11 +668,11 @@ FVertexFactoryIntermediates GetVertexFactoryIntermediates(FVertexFactoryInput In
 
 	// Get the height and normal XY for current and next LOD out of the textures
 	float2 SampleCoords = InputPositionLODAdjusted * LandscapeParameters.HeightmapUVScaleBias.xy + LandscapeParameters.HeightmapUVScaleBias.zw + 0.5*LandscapeParameters.HeightmapUVScaleBias.xy + Intermediates.InputPosition.zw * LandscapeParameters.SubsectionOffsetParams.xy;
-	float4 SampleValue = Texture2DSampleLevel(LandscapeParameters.HeightmapTexture, LandscapeParameters.HeightmapTextureSampler, SampleCoords, LodValue-Intermediates.LodBias.x);
+	float4 SampleValue = Texture2DSampleLevel(LandscapeParameters.HeightmapTexture, LandscapeParameters.HeightmapTextureSampler, SampleCoords, LodValue-Intermediates.LodBias);
 	float Height = DecodePackedHeight(SampleValue.xy);
 
 	float2 SampleCoordsNextLOD = InputPositionNextLOD * LandscapeParameters.HeightmapUVScaleBias.xy + LandscapeParameters.HeightmapUVScaleBias.zw + 0.5*LandscapeParameters.HeightmapUVScaleBias.xy + Intermediates.InputPosition.zw * LandscapeParameters.SubsectionOffsetParams.xy;
-	float4 SampleValueNextLOD = Texture2DSampleLevel(LandscapeParameters.HeightmapTexture, LandscapeParameters.HeightmapTextureSampler, SampleCoordsNextLOD, LodValue+1-Intermediates.LodBias.x);
+	float4 SampleValueNextLOD = Texture2DSampleLevel(LandscapeParameters.HeightmapTexture, LandscapeParameters.HeightmapTextureSampler, SampleCoordsNextLOD, LodValue+1-Intermediates.LodBias);
 	float HeightNextLOD = DecodePackedHeight(SampleValueNextLOD.xy);
 
 	Intermediates.LocalPosition = lerp( float3(InputPositionLODAdjusted, Height), float3(InputPositionNextLOD, HeightNextLOD), MorphAlpha );
