Fix bug in GetGBufferData(...) for mobile shaders.
- PixelPos was float instead of float2.

#Synced-CL 49833820

diff --git a/Engine/Shaders/Private/DeferredShadingCommon.ush b/Engine/Shaders/Private/DeferredShadingCommon.ush
--- a/Engine/Shaders/Private/DeferredShadingCommon.ush
+++ b/Engine/Shaders/Private/DeferredShadingCommon.ush
@@ -877,7 +877,7 @@ FGBufferData MobileFetchAndDecodeGBuffer(in float2 UV, in float2 PixelPos)
 // Mirror desktop GetGBufferData() 
 FGBufferData GetGBufferData(float2 UV, bool bGetNormalizedNormal = true, bool bForceUnlitOrDefaultLit = false)
 {
-	float PixelPos = trunc(UV * View.BufferSizeAndInvSize.xy * View.BufferToSceneTextureScale.xy);
+	float2 PixelPos = trunc(UV * View.BufferSizeAndInvSize.xy * View.BufferToSceneTextureScale.xy);
 	FGBufferData GBuffer = MobileFetchAndDecodeGBuffer(UV, PixelPos);
 	if (bGetNormalizedNormal)
 	{
