Renaming TransmissionThickness which does not means what it represent to a more explicit name TransmittanceOrOpticalThickness.
Added comments.

#rnx
#tests PC editor dx12
#Synced-CL 49840788

diff --git a/Engine/Shaders/Private/DeferredLightingCommon.ush b/Engine/Shaders/Private/DeferredLightingCommon.ush
--- a/Engine/Shaders/Private/DeferredLightingCommon.ush
+++ b/Engine/Shaders/Private/DeferredLightingCommon.ush
@@ -122,7 +122,7 @@ void GetShadowTermsBase(
 			//@todo - how to do static shadowing of SSS correctly?
 			OutShadow.TransmissionShadow = LightAttenuation.w * StaticShadowing;
 
-			OutShadow.TransmissionThickness = LightAttenuation.w;
+			OutShadow.TransmittanceOrOpticalThickness = LightAttenuation.w;
 		}
 		else
 		{
@@ -140,7 +140,7 @@ void GetShadowTermsBase(
 			OutShadow.TransmissionShadow *= LightAttenuation.z;
 
 			// Need this min or backscattering will leak when in shadow which cast by non perobject shadow(Only for directional light)
-			OutShadow.TransmissionThickness = min(LightAttenuation.y, LightAttenuation.w);
+			OutShadow.TransmittanceOrOpticalThickness = min(LightAttenuation.y, LightAttenuation.w);
 		}
 	}
 
@@ -346,7 +346,7 @@ FLightAccumulator AccumulateDynamicLighting(
 		FShadowTerms Shadow;
 		Shadow.SurfaceShadow = AmbientOcclusion;
 		Shadow.TransmissionShadow = 1;
-		Shadow.TransmissionThickness = 1;
+		Shadow.TransmittanceOrOpticalThickness = 1;
 		Shadow.HairTransmittance.OpaqueVisibility = 1;
 		const float ContactShadowOpacity = GBuffer.CustomData.a;
 		GetShadowTerms(GBuffer.Depth, GBuffer.PrecomputedShadowFactors, GBuffer.ShadingModelID, ContactShadowOpacity,
diff --git a/Engine/Shaders/Private/ShadingModels.ush b/Engine/Shaders/Private/ShadingModels.ush
--- a/Engine/Shaders/Private/ShadingModels.ush
+++ b/Engine/Shaders/Private/ShadingModels.ush
@@ -92,9 +92,18 @@ struct FDirectLighting
 
 struct FShadowTerms
 {
+	// This is the surface shadow from opaque.
 	half	SurfaceShadow;
+
+	// This is the surface shadow from opaque multiplied by transmittance from subsurface elements.
 	half	TransmissionShadow;
-	half	TransmissionThickness;
+
+	// This can be OpticalThickness or Transmittance depending (see ShadowProjectionPixelShader.usf and GetShadowTermsBase(...)).
+	// Shading model - Subsurface			=> Transmittance in [0, 1]
+	// Shading model - Subsurface Profile	=> Encoded OpticalThickness to be decoded using DecodeOpticalDepthFromShadowMask.
+	// Substrate	 - Per Pixel MFP		=> Same as Subsurface Profile.
+	half	TransmittanceOrOpticalThickness;
+
 	FHairTransmittanceData HairTransmittance;
 };
 
@@ -645,7 +654,7 @@ FDirectLighting SubsurfaceProfileBxDF( FGBufferData GBuffer, half3 N, half3 V, F
 	const uint ProfileId = ExtractSubsurfaceProfileInt(GBuffer);
 	FTransmissionProfileParams TransmissionParams = GetTransmissionProfileParams(ProfileId);
 
-	float ShadowOpticalDepth= DecodeOpticalDepthFromShadowMask(Shadow.TransmissionThickness);
+	float ShadowOpticalDepth= DecodeOpticalDepthFromShadowMask(Shadow.TransmittanceOrOpticalThickness);
 	float3 Profile = GetTransmissionProfile(ProfileId, ShadowOpticalDepth).rgb; // GetTransmissionProfile Thickness input is optical depth?
 
 	float3 RefracV = refract(V, -N, TransmissionParams.OneOverIOR);
@@ -735,7 +744,8 @@ FDirectLighting SubsurfaceBxDF(FGBufferData GBuffer, half3 N, half3 V, FAreaLigh
 	const half3 TransmittedColor = HSV_2_LinearRGB(half3(LinearRGB_2_HSV(RawTransmittedColor).xy, LinearRGB_2_HSV(SubsurfaceColor).z));
 
 	// lerp to never exceed 1 (energy conserving)
-	Lighting.Transmission = AreaLight.FalloffColor * (AreaLight.Falloff * lerp(BackScatter, 1, InScatter)) * lerp(TransmittedColor, SubsurfaceColor, Shadow.TransmissionThickness);
+	const float SubsurfaceTransmittance = Shadow.TransmittanceOrOpticalThickness;
+	Lighting.Transmission = AreaLight.FalloffColor * (AreaLight.Falloff * lerp(BackScatter, 1, InScatter)) * lerp(TransmittedColor, SubsurfaceColor, SubsurfaceTransmittance);
 
 	return Lighting;
 }
diff --git a/Engine/Shaders/Private/Substrate/SubstrateEvaluation.ush b/Engine/Shaders/Private/Substrate/SubstrateEvaluation.ush
--- a/Engine/Shaders/Private/Substrate/SubstrateEvaluation.ush
+++ b/Engine/Shaders/Private/Substrate/SubstrateEvaluation.ush
@@ -916,7 +916,8 @@ FSubstrateEvaluateResult SubstrateEvaluateBSDFCommon(FSubstrateBSDFContext BSDFC
 					const half3 TransmittedColor = HSV_2_LinearRGB(half3(LinearRGB_2_HSV(RawTransmittedColor).xy, LinearRGB_2_HSV(SubSurfaceColor).z));
 
 					// Lerp to never exceed 1 (energy conserving)
-					TransmissionThroughput = lerp(BackScatter, 1, InScatter) * lerp(TransmittedColor, SubSurfaceColor, ShadowTerms.TransmissionThickness);
+					const float SubsurfaceTransmittance = ShadowTerms.TransmittanceOrOpticalThickness;
+					TransmissionThroughput = lerp(BackScatter, 1, InScatter) * lerp(TransmittedColor, SubSurfaceColor, SubsurfaceTransmittance);
 				}
 
 				Sample.TransmissionPDF = 1.0f / (4.0f * PI); // SUBSTRATE_TODO this currently match the uniform sphere sampling from SubstrateImportanceSampleBSDF
@@ -935,7 +936,7 @@ FSubstrateEvaluateResult SubstrateEvaluateBSDFCommon(FSubstrateBSDFContext BSDFC
 			#if SUBSTRATE_SSS_TRANSMISSION
 			if (bIsSubsurfaceDiffusionProfile || bIsSubsurfaceDiffusionMFP)
 			{
-				const float ShadowOpticalDepth = DecodeOpticalDepthFromShadowMask(ShadowTerms.TransmissionThickness);
+				const float ShadowOpticalDepth = DecodeOpticalDepthFromShadowMask(ShadowTerms.TransmittanceOrOpticalThickness);
 
 				// TODO: we end up with optical depth here but it seems SSS GetTransmissionProfile or GetBurleyTransmission expect Thickness in centimeters. A bad naming case?
 				float ThicknessInCm = ShadowOpticalDepth;
