Substrate - removed unused SSS Headers.

#tests PC editor dx12
#rnx

#Synced-CL 49923192

diff --git a/Engine/Shaders/Private/DefaultSSRTiles.usf b/Engine/Shaders/Private/DefaultSSRTiles.usf
--- a/Engine/Shaders/Private/DefaultSSRTiles.usf
+++ b/Engine/Shaders/Private/DefaultSSRTiles.usf
@@ -51,8 +51,6 @@ bool DoesPixelContainSSR(uint2 PixelPos)
 	BRANCH
 	if (SubstratePixelHeader.ClosureCount == 1) // Only non complex materials will be interpreted for exclusion for now when it comes to Specular values or legacy shading models.
 	{
-		const FSubstrateSubsurfaceHeader SSSHeader = SubstrateLoadSubsurfaceHeader(Substrate.MaterialTextureArray, Substrate.FirstSliceStoringSubstrateSSSData, SubstrateAddressing.PixelCoords);
-
 		const bool bSkipSSSMaterialOverride = true; // We do not want the material data to be affected by any override.
 		FSubstrateBSDF BSDF = UnpackSubstrateBSDFIn(Substrate.MaterialTextureArray, SubstrateAddressing, SubstratePixelHeader, bSkipSSSMaterialOverride);
 
diff --git a/Engine/Shaders/Private/PostProcessExposureMaterial.ush b/Engine/Shaders/Private/PostProcessExposureMaterial.ush
--- a/Engine/Shaders/Private/PostProcessExposureMaterial.ush
+++ b/Engine/Shaders/Private/PostProcessExposureMaterial.ush
@@ -33,7 +33,6 @@ FExposureMaterial GetExposureMaterial(uint2 InPixelPos)
 		FSubstrateAddressing SubstrateAddressing = GetSubstratePixelDataByteOffset(InPixelPos, uint2(View.BufferSizeAndInvSize.xy), Substrate.MaxBytesPerPixel);
 		FSubstratePixelHeader SubstratePixelHeader = UnpackSubstrateHeaderIn(Substrate.MaterialTextureArray, SubstrateAddressing, Substrate.TopLayerTexture);
 		FSubstrateTopLayerData TopLayerData = SubstrateUnpackTopLayerData(Substrate.TopLayerTexture.Load(uint3(InPixelPos, 0)));
-		FSubstrateSubsurfaceHeader SSSHeader = SubstrateLoadSubsurfaceHeader(Substrate.MaterialTextureArray, Substrate.FirstSliceStoringSubstrateSSSData, InPixelPos);
 
 		const float DeviceZ = ConvertFromDeviceZ(SceneDepthTexture.Load(int3(InPixelPos, 0)).r);
 		const float3 TranslatedWorldPosition = ReconstructTranslatedWorldPositionFromDeviceZ(InPixelPos, DeviceZ);
