Add uber shader permutation for VolumetricFog.

The uber shader permutation is used when the requested permutation is no available and requires PSO compilation. While the PSO is asyncly compiled, the uber permutation is used.

#Synced-CL 49918411

diff --git a/Engine/Shaders/Private/VolumetricFog.usf b/Engine/Shaders/Private/VolumetricFog.usf
--- a/Engine/Shaders/Private/VolumetricFog.usf
+++ b/Engine/Shaders/Private/VolumetricFog.usf
@@ -410,6 +410,16 @@ float ComputeRayTracedShadowFactor(float3 TranslatedWorldPosition, float3 LightD
 
 #endif
 
+uint bUseTemporalReprojection;
+uint bDistanceFieldSkyOcclusion;
+uint bLumenGI;
+uint bVirtualShadowMap;
+uint bUseRaytracedShadowsVolume;
+uint bUseLightFunctionAtlas;
+uint bUseMegaLights;
+uint bUseLightSoftFading;
+uint bUseSuperSampling;
+
 float4 InjectShadowedLocalLightCommon(uint3 GridCoordinate)
 {
 	float4 OutScattering = 0;
@@ -427,20 +437,27 @@ float4 InjectShadowedLocalLightCommon(uint3 GridCoordinate)
 
 		uint NumSuperSamples = 1;
 
-#if USE_TEMPORAL_REPROJECTION
-		float3 HistoryUV = ComputeHistoryVolumeUVFromTranslatedPos(ComputeCellTranslatedWorldPosition(GridCoordinate, .5f), UnjitteredPrevTranslatedWorldToClip);
-		float HistoryAlpha = HistoryWeight;
-
-		FLATTEN
-		if (any(HistoryUV < 0) || any(HistoryUV > 1))
+		#if USE_TEMPORAL_REPROJECTION
+		IF_UBERSHADER(bUseTemporalReprojection)
 		{
-			HistoryAlpha = 0;
-		}
+			float3 HistoryUV = ComputeHistoryVolumeUVFromTranslatedPos(ComputeCellTranslatedWorldPosition(GridCoordinate, .5f), UnjitteredPrevTranslatedWorldToClip);
+			float HistoryAlpha = HistoryWeight;
+
+			FLATTEN
+			if (any(HistoryUV < 0) || any(HistoryUV > 1))
+			{
+				HistoryAlpha = 0;
+			}
 
-		NumSuperSamples = HistoryAlpha < .001f ? HistoryMissSuperSampleCount : 1;
+			NumSuperSamples = HistoryAlpha < .001f ? HistoryMissSuperSampleCount : 1;
+		}
+		#endif
 
-#endif
-		const bool bSoftFadeEnabled = VolumetricFog.LightSoftFading > 0;
+		const bool bSoftFadeEnabled = 
+			#if UBERSHADER
+			bUseLightSoftFading &&
+			#endif
+			VolumetricFog.LightSoftFading > 0;
 
 		for (uint SampleIndex = 0; SampleIndex < NumSuperSamples; SampleIndex++)
 		{
@@ -516,18 +533,25 @@ float4 InjectShadowedLocalLightCommon(uint3 GridCoordinate)
 				bool bUnused = false;
 				ShadowFactor = ComputeVolumeShadowing(TranslatedWorldPosition, IsPointLight, LightData.bSpotLight, bUnused);
 
-#if VIRTUAL_SHADOW_MAP
-				FVirtualShadowMapSampleResult VirtualShadowMapSample = SampleVirtualShadowMapTranslatedWorld(VirtualShadowMapId, TranslatedWorldPosition);
-				ShadowFactor *= VirtualShadowMapSample.ShadowFactor;
-#endif // VIRTUAL_SHADOW_MAP
+				#if VIRTUAL_SHADOW_MAP
+				IF_UBERSHADER(bVirtualShadowMap)
+				{
+					FVirtualShadowMapSampleResult VirtualShadowMapSample = SampleVirtualShadowMapTranslatedWorld(VirtualShadowMapId, TranslatedWorldPosition);
+					ShadowFactor *= VirtualShadowMapSample.ShadowFactor;
+				}
+				#endif // VIRTUAL_SHADOW_MAP
+
 #endif // USE_RAYTRACED_SHADOWS
 			}
 #endif
 	
 			FLightFunctionColor LightFunctionColor = 1.0f;
-#if	USE_LIGHT_FUNCTION_ATLAS
-			LightFunctionColor = GetLocalLightFunctionCommon(TranslatedWorldPosition, LightData.LightFunctionAtlasLightIndex);
-#endif
+			#if	USE_LIGHT_FUNCTION_ATLAS
+			IF_UBERSHADER(bUseLightFunctionAtlas)
+			{
+				LightFunctionColor = GetLocalLightFunctionCommon(TranslatedWorldPosition, LightData.LightFunctionAtlasLightIndex);
+			}
+			#endif
 
 			OutScattering.rgb += LightColor * LightFunctionColor * (PhaseFunction(PhaseG, dot(L, -CameraVector)) * CombinedAttenuation * ShadowFactor * VolumetricScatteringIntensity);
 				
@@ -692,13 +716,14 @@ float ComputeSkyVisibility(float3 TranslatedWorldPosition, float3 BrickTextureUV
 {
 	float Visibility = 1;
 
-#if DISTANCE_FIELD_SKY_OCCLUSION
-	// Trace one 45 degree cone straight up for sky occlusion
-	float TanConeHalfAngle = tan((float)PI / 4);
-
-	Visibility = HemisphereConeTraceAgainstGlobalDistanceField(TranslatedWorldPosition, float3(0, 0, 1), TanConeHalfAngle);
-	
-#endif
+	#if DISTANCE_FIELD_SKY_OCCLUSION
+	IF_UBERSHADER(bDistanceFieldSkyOcclusion)
+	{
+		// Trace one 45 degree cone straight up for sky occlusion
+		float TanConeHalfAngle = tan((float)PI / 4);
+		Visibility = HemisphereConeTraceAgainstGlobalDistanceField(TranslatedWorldPosition, float3(0, 0, 1), TanConeHalfAngle);
+	}
+	#endif
 
 #if ALLOW_STATIC_LIGHTING
 	if (SkyLightUseStaticShadowing > 0)
@@ -719,12 +744,15 @@ uint DirectionalLightFunctionAtlasLightIndex;
 
 FLightFunctionColor GetDirectionalLightFunction(float3 TranslatedWorldPosition) // directional light
 {
-#if USE_LIGHT_FUNCTION_ATLAS  // Used when injecting lights from the light grid
-	if (DirectionalApplyLightFunctionFromAtlas > 0)
+	#if USE_LIGHT_FUNCTION_ATLAS  // Used when injecting lights from the light grid
+	IF_UBERSHADER(bUseLightFunctionAtlas)
 	{
-		return GetLocalLightFunctionCommon(TranslatedWorldPosition, DirectionalLightFunctionAtlasLightIndex);
+		if (DirectionalApplyLightFunctionFromAtlas > 0)
+		{
+			return GetLocalLightFunctionCommon(TranslatedWorldPosition, DirectionalLightFunctionAtlasLightIndex);
+		}
 	}
-#endif
+	#endif
 
 	float4 HomogeneousShadowPosition = mul(float4(TranslatedWorldPosition, 1), DirectionalLightFunctionTranslatedWorldToShadow);
 	float2 LightFunctionUV = HomogeneousShadowPosition.xy * .5f + .5f;
@@ -849,30 +877,41 @@ void LightScatteringCS(
 			return;
 		}
 
-#if USE_TEMPORAL_REPROJECTION
-		// TODO: investigate whether using CellTranslatedWorldPosition (which is biased towards camera) is correct here
-		// it might be more accurate to transform unbiased translated world position to previous frame clip space and apply bias towards camera in that space?
-		float4 PrevCellNDCPosition = mul(float4(CellTranslatedWorldPosition, 1), UnjitteredPrevTranslatedWorldToClip);
-		PrevCellNDCPosition.xyz /= PrevCellNDCPosition.w;
+		#if USE_TEMPORAL_REPROJECTION
+		IF_UBERSHADER(bUseTemporalReprojection)
+		{
+			// TODO: investigate whether using CellTranslatedWorldPosition (which is biased towards camera) is correct here
+			// it might be more accurate to transform unbiased translated world position to previous frame clip space and apply bias towards camera in that space?
+			float4 PrevCellNDCPosition = mul(float4(CellTranslatedWorldPosition, 1), UnjitteredPrevTranslatedWorldToClip);
+			PrevCellNDCPosition.xyz /= PrevCellNDCPosition.w;
 
-		HistoryUV.xy = FixupHistoryUV(HistoryUV.xy, View.VolumetricFogPrevResourceGridSize.xy, PrevCellNDCPosition.z, bHasValidHistory);
-#endif
+			HistoryUV.xy = FixupHistoryUV(HistoryUV.xy, View.VolumetricFogPrevResourceGridSize.xy, PrevCellNDCPosition.z, bHasValidHistory);
+		}
+		#endif
 	}
 
-#if USE_TEMPORAL_REPROJECTION
+	#if USE_TEMPORAL_REPROJECTION
 	float HistoryAlpha = HistoryWeight;
-
-	// We need to test with View.VolumetricFogPrevUVMaxForTemporalBlend for HistoryUV because that is what we clamp with in ComputeHistoryVolumeUVFromTranslatedPos.
-	FLATTEN
-	if (any(HistoryUV < 0) || any(HistoryUV >= float3(View.VolumetricFogPrevUVMaxForTemporalBlend, 1.0f)) || !bHasValidHistory)
+	IF_UBERSHADER(bUseTemporalReprojection)
 	{
-		HistoryAlpha = 0;
-	}
+		// We need to test with View.VolumetricFogPrevUVMaxForTemporalBlend for HistoryUV because that is what we clamp with in ComputeHistoryVolumeUVFromTranslatedPos.
+		FLATTEN
+		if (any(HistoryUV < 0) || any(HistoryUV >= float3(View.VolumetricFogPrevUVMaxForTemporalBlend, 1.0f)) || !bHasValidHistory)
+		{
+			HistoryAlpha = 0;
+		}
 
-	// Supersample if the history was outside the camera frustum
-	// The compute shader is dispatched with extra threads, make sure those don't supersample
-	NumSuperSamples = HistoryAlpha < .001f && all(int3(GridCoordinate) < VolumetricFog.ViewGridSizeInt) && USE_SUPER_SAMPLING ? HistoryMissSuperSampleCount : 1;
-#endif
+		// Supersample if the history was outside the camera frustum
+		// The compute shader is dispatched with extra threads, make sure those don't supersample
+		NumSuperSamples = 1;
+		#if USE_SUPER_SAMPLING
+		IF_UBERSHADER (bUseSuperSampling)
+		{
+			NumSuperSamples = HistoryAlpha < .001f && all(int3(GridCoordinate) < VolumetricFog.ViewGridSizeInt) ? HistoryMissSuperSampleCount : 1;
+		}
+		#endif
+	}
+	#endif
 
 	for (uint SampleIndex = 0; SampleIndex < NumSuperSamples; SampleIndex++)
 	{
@@ -898,7 +937,7 @@ void LightScatteringCS(
 		const bool	 bHasDirectionalLight		= ForwardLightStruct.HasDirectionalLight;
 		const float3 InDirectionalLightColor	= ForwardLightStruct.DirectionalLightColor * ForwardLightStruct.DirectionalLightVolumetricScatteringIntensity;
 		const float3 InDirectionalLightDirection= ForwardLightStruct.DirectionalLightDirection;
-		const bool   bEvaluateDirectionalLight	= !USE_MEGA_LIGHTS || !ForwardLightStruct.DirectionalLightHandledByMegaLights;
+		const bool   bEvaluateDirectionalLight	= !USE_MEGA_LIGHTS || (UBERSHADER && USE_MEGA_LIGHTS && !bUseMegaLights) || !ForwardLightStruct.DirectionalLightHandledByMegaLights;
 	#endif
 		BRANCH
 		if (bHasDirectionalLight && bEvaluateDirectionalLight)
@@ -911,14 +950,20 @@ void LightScatteringCS(
 				bool bUnused = false;
 				ShadowFactor *= ComputeDirectionalLightDynamicShadowing(TranslatedWorldPosition, SceneDepth, bUnused);
 
-#if VIRTUAL_SHADOW_MAP
-				FVirtualShadowMapSampleResult VirtualShadowMapSample = SampleVirtualShadowMapTranslatedWorld(ForwardLightStruct.DirectionalLightVSM, TranslatedWorldPosition);
-				ShadowFactor *= VirtualShadowMapSample.ShadowFactor;
-#endif
+				#if VIRTUAL_SHADOW_MAP
+				IF_UBERSHADER(bVirtualShadowMap)
+				{
+					FVirtualShadowMapSampleResult VirtualShadowMapSample = SampleVirtualShadowMapTranslatedWorld(ForwardLightStruct.DirectionalLightVSM, TranslatedWorldPosition);
+					ShadowFactor *= VirtualShadowMapSample.ShadowFactor;
+				}
+				#endif
 
-#if USE_RAYTRACED_SHADOWS_VOLUME
-				ShadowFactor *= RaytracedShadowsVolume[GridCoordinate];
-#endif
+				#if USE_RAYTRACED_SHADOWS_VOLUME
+				IF_UBERSHADER(bUseRaytracedShadowsVolume)
+				{
+					ShadowFactor *= RaytracedShadowsVolume[GridCoordinate];
+				}
+				#endif
 
 #if !SHADING_PATH_MOBILE
 				if (CloudShadowmapStrength > 0.0f)
@@ -954,40 +999,46 @@ void LightScatteringCS(
 		}
 #endif
 
-#if LUMEN_GI
-
-		// Lumen Dynamic GI + shadowed Skylight
-		FTwoBandSHVectorRGB TranslucencyGISH = GetTranslucencyGIVolumeLighting(DFPromote(WorldPosition), PrimaryView.WorldToClip, false); // LUMEN_LWC_TODO
+		#if (LUMEN_GI || UBERSHADER) && !SHADING_PATH_MOBILE
+		IF_UBERSHADER (bLumenGI)
+		{
+			// Lumen Dynamic GI + shadowed Skylight
+			FTwoBandSHVectorRGB TranslucencyGISH = GetTranslucencyGIVolumeLighting(DFPromote(WorldPosition), PrimaryView.WorldToClip, false); // LUMEN_LWC_TODO
 
-		LightScattering += max(DotSH(TranslucencyGISH, RotatedHGZonalHarmonic), 0);
+			LightScattering += max(DotSH(TranslucencyGISH, RotatedHGZonalHarmonic), 0);
+		}
+		#endif
 
-#else
+		#if !LUMEN_GI || UBERSHADER
 		// Skylight
-		if (View.SkyLightVolumetricScatteringIntensity > 0)
+		IF_UBERSHADER (!bLumenGI)
 		{
-			float3 SkyLighting;
-
-			if (UseHeightFogColors.y > 0)
+			if (View.SkyLightVolumetricScatteringIntensity > 0)
 			{
-				float3 HeightFogInscatteringColor = ComputeInscatteringColor(CameraVector, CameraVectorLength);
-				float ScalarFactor = SHAmbientFunction();
-				FTwoBandSHVectorRGB SkyIrradianceSH;
-				SkyIrradianceSH.R.V = float4(ScalarFactor * HeightFogInscatteringColor.r, 0, 0, 0);
-				SkyIrradianceSH.G.V = float4(ScalarFactor * HeightFogInscatteringColor.g, 0, 0, 0);
-				SkyIrradianceSH.B.V = float4(ScalarFactor * HeightFogInscatteringColor.b, 0, 0, 0);
-
-				SkyLighting = max(DotSH(SkyIrradianceSH, RotatedHGZonalHarmonic), 0);
-			}
-			else
-			{
-				// NOTE it should be CameraVector * PhaseG. But we keep it like that for backward compatibility for now. (visible when r.Lumen.TranslucencyVolume.Enable 0)
-				SkyLighting = SampleSkyLightDiffuseEnvMap > 0 ? View.SkyLightColor.rgb * GetSkySHDiffuseSimple(CameraVector * -PhaseG) : 0.0f.xxx;
-			}
+				float3 SkyLighting;
 
-			float SkyVisibility = ComputeSkyVisibility(TranslatedWorldPosition, BrickTextureUVs);
-			LightScattering += (SkyVisibility * View.SkyLightVolumetricScatteringIntensity) * SkyLighting;
+				if (UseHeightFogColors.y > 0)
+				{
+					float3 HeightFogInscatteringColor = ComputeInscatteringColor(CameraVector, CameraVectorLength);
+					float ScalarFactor = SHAmbientFunction();
+					FTwoBandSHVectorRGB SkyIrradianceSH;
+					SkyIrradianceSH.R.V = float4(ScalarFactor * HeightFogInscatteringColor.r, 0, 0, 0);
+					SkyIrradianceSH.G.V = float4(ScalarFactor * HeightFogInscatteringColor.g, 0, 0, 0);
+					SkyIrradianceSH.B.V = float4(ScalarFactor * HeightFogInscatteringColor.b, 0, 0, 0);
+
+					SkyLighting = max(DotSH(SkyIrradianceSH, RotatedHGZonalHarmonic), 0);
+				}
+				else
+				{
+					// NOTE it should be CameraVector * PhaseG. But we keep it like that for backward compatibility for now. (visible when r.Lumen.TranslucencyVolume.Enable 0)
+					SkyLighting = SampleSkyLightDiffuseEnvMap > 0 ? View.SkyLightColor.rgb * GetSkySHDiffuseSimple(CameraVector * -PhaseG) : 0.0f.xxx;
+				}
+
+				float SkyVisibility = ComputeSkyVisibility(TranslatedWorldPosition, BrickTextureUVs);
+				LightScattering += (SkyVisibility * View.SkyLightVolumetricScatteringIntensity) * SkyLighting;
+			}
 		}
-#endif
+		#endif
 
 #if ALLOW_STATIC_LIGHTING
 		// Indirect lighting of Stationary lights and Direct + Indirect lighting of Static lights
@@ -1009,27 +1060,32 @@ void LightScatteringCS(
 		float DistanceBiasSqr = max(CellRadius * InverseSquaredLightDistanceBiasScale, 1);
 		DistanceBiasSqr *= DistanceBiasSqr;
 
-#if USE_MEGA_LIGHTS
-		// when using Mega Lights Volume only need to apply "non-mega" lights here
-		const uint NumLights = CulledLightsGridHeader.NumLights - CulledLightsGridHeader.NumMegaLights;
-#else
-		const uint NumLights = CulledLightsGridHeader.NumLights;
-#endif
+		uint NumLights = CulledLightsGridHeader.NumLights;
+		#if USE_MEGA_LIGHTS
+		IF_UBERSHADER(bUseMegaLights)
+		{
+			// when using Mega Lights Volume only need to apply "non-mega" lights here
+			NumLights -= CulledLightsGridHeader.NumMegaLights;
+		}
+		#endif
 
 		// Forward lighting of unshadowed point and spot lights
 		LOOP
 		for (uint GridLightListIndex = 0; GridLightListIndex < NumLights; GridLightListIndex++)
 		{
 			const FLocalLightData LocalLight = GetLocalLightDataFromGrid(CulledLightsGridHeader.DataStartIndex + GridLightListIndex, 0);
-#if USE_MEGA_LIGHTS
-			checkSlow(!UnpackIsHandledByMegaLights(LocalLight));
-#endif
+			#if USE_MEGA_LIGHTS
+			IF_UBERSHADER(bUseMegaLights)
+			{
+				checkSlow(!UnpackIsHandledByMegaLights(LocalLight));
+			}
+			#endif
 
 			const float VolumetricScatteringIntensity = UnpackVolumetricScatteringIntensity(LocalLight);
 
 			if (VolumetricScatteringIntensity > 0)
 			{
-				const bool bSoftFadeEnabled = VolumetricFog.LightSoftFading > 0;
+				const bool bSoftFadeEnabled = VolumetricFog.LightSoftFading > 0 && bUseLightSoftFading;
 				const FDeferredLightData LightData = ConvertToDeferredLight(LocalLight);
 
 				float3 L = 0;
@@ -1071,9 +1127,13 @@ void LightScatteringCS(
 				float CombinedAttenuation = Lighting * LightMask;
 
 				FLightFunctionColor LightFunctionColor = 1.0f;
-			#if	USE_LIGHT_FUNCTION_ATLAS
-				LightFunctionColor = GetLocalLightFunctionCommon(TranslatedWorldPosition, LightData.LightFunctionAtlasLightIndex);
-			#endif
+				#if	USE_LIGHT_FUNCTION_ATLAS
+				IF_UBERSHADER(bUseLightFunctionAtlas)
+				{
+
+					LightFunctionColor = GetLocalLightFunctionCommon(TranslatedWorldPosition, LightData.LightFunctionAtlasLightIndex);
+				}
+				#endif
 
 				LightScattering += LightData.Color * LightFunctionColor * (PhaseFunction(PhaseG, dot(L, -CameraVector)) * CombinedAttenuation * VolumetricScatteringIntensity);
 
@@ -1090,6 +1150,7 @@ void LightScatteringCS(
 	LightScattering += View.OneOverPreExposure * LocalShadowedLightScattering[GridCoordinate].xyz;
 
 	#if USE_MEGA_LIGHTS
+	IF_UBERSHADER(bUseMegaLights)
 	{
 		// Inject MegaLights which were computed earlier. This texture was already pre-exposed, so need to rever it here.
 		LightScattering += MegaLightsVolume[GridCoordinate] * View.OneOverPreExposure;
@@ -1106,24 +1167,29 @@ void LightScatteringCS(
 	}
 	float4 PreExposedScatteringAndExtinction = float4(View.PreExposure * (LightScattering * MaterialScatteringAndExtinction.xyz + MaterialEmissive), Extinction);
 
-#if USE_TEMPORAL_REPROJECTION
-	BRANCH
-	if (HistoryAlpha > 0)
+	#if USE_TEMPORAL_REPROJECTION
+	IF_UBERSHADER(bUseTemporalReprojection)
 	{
-		float4 PreExposedHistoryScatteringAndExtinction = Texture3DSampleLevel(LightScatteringHistory, LightScatteringHistorySampler, min(HistoryUV, float3(View.VolumetricFogPrevUVMax, 1.0f)), 0);
-		PreExposedHistoryScatteringAndExtinction.rgb *= LightScatteringHistoryPreExposureAndInv.y * View.PreExposure; // Bring previous frame pre exposed history luminance as current frame pre exposed luminance. Leave extinction untouched! 
-		PreExposedScatteringAndExtinction = lerp(PreExposedScatteringAndExtinction, PreExposedHistoryScatteringAndExtinction, HistoryAlpha);
+		BRANCH
+		if (HistoryAlpha > 0)
+		{
+			float4 PreExposedHistoryScatteringAndExtinction = Texture3DSampleLevel(LightScatteringHistory, LightScatteringHistorySampler, min(HistoryUV, float3(View.VolumetricFogPrevUVMax, 1.0f)), 0);
+			PreExposedHistoryScatteringAndExtinction.rgb *= LightScatteringHistoryPreExposureAndInv.y * View.PreExposure; // Bring previous frame pre exposed history luminance as current frame pre exposed luminance. Leave extinction untouched! 
+			PreExposedScatteringAndExtinction = lerp(PreExposedScatteringAndExtinction, PreExposedHistoryScatteringAndExtinction, HistoryAlpha);
+		}
 	}
-	
-#endif
+	#endif
 
 	// Visualize history rejection for debugging purposes
-#if 0 && USE_TEMPORAL_REPROJECTION
-	if (HistoryAlpha < 0.001f)
+	#if 0 && USE_TEMPORAL_REPROJECTION
+	IF_UBERSHADER(bUseTemporalReprojection)
 	{
-		PreExposedScatteringAndExtinction = float4(View.PreExposure * float3(1,0,0), Extinction);
+		if (HistoryAlpha < 0.001f)
+		{
+			PreExposedScatteringAndExtinction = float4(View.PreExposure * float3(1, 0, 0), Extinction);
+		}
 	}
-#endif
+	#endif
 
 	if (all(int3(GridCoordinate) < VolumetricFog.ResourceGridSizeInt))
 	{
diff --git a/Engine/Source/Runtime/Renderer/Private/VolumetricFog.cpp b/Engine/Source/Runtime/Renderer/Private/VolumetricFog.cpp
--- a/Engine/Source/Runtime/Renderer/Private/VolumetricFog.cpp
+++ b/Engine/Source/Runtime/Renderer/Private/VolumetricFog.cpp
@@ -1037,6 +1037,7 @@ class FVolumetricFogLightScatteringCS : public FGlobalShader
 	class FMegaLights					: SHADER_PERMUTATION_BOOL("USE_MEGA_LIGHTS");
 	class FLightSoftFading				: SHADER_PERMUTATION_BOOL("USE_LIGHT_SOFT_FADING");
 	class FSuperSampling				: SHADER_PERMUTATION_BOOL("USE_SUPER_SAMPLING");
+	class FUbershader					: SHADER_PERMUTATION_BOOL("UBERSHADER");
 
 	using FPermutationDomain = TShaderPermutationDomain<
 		FTemporalReprojection,
@@ -1047,7 +1048,8 @@ class FVolumetricFogLightScatteringCS : public FGlobalShader
 		FSampleLightFunctionAtlas,
 		FMegaLights,
 		FLightSoftFading,
-		FSuperSampling>;
+		FSuperSampling,
+		FUbershader>;
 
 	BEGIN_SHADER_PARAMETER_STRUCT(FParameters, )
 		SHADER_PARAMETER_STRUCT_REF(FViewUniformShaderParameters, View)
@@ -1096,6 +1098,15 @@ class FVolumetricFogLightScatteringCS : public FGlobalShader
 		SHADER_PARAMETER(uint32, MobileHasDirectionalLight)
 		SHADER_PARAMETER(uint32, DirectionalApplyLightFunctionFromAtlas)
 		SHADER_PARAMETER(uint32, DirectionalLightFunctionAtlasLightIndex)
+		SHADER_PARAMETER(uint32, bUseTemporalReprojection)
+		SHADER_PARAMETER(uint32, bDistanceFieldSkyOcclusion)
+		SHADER_PARAMETER(uint32, bLumenGI)
+		SHADER_PARAMETER(uint32, bVirtualShadowMap)
+		SHADER_PARAMETER(uint32, bUseRaytracedShadowsVolume)
+		SHADER_PARAMETER(uint32, bUseLightFunctionAtlas)
+		SHADER_PARAMETER(uint32, bUseMegaLights)
+		SHADER_PARAMETER(uint32, bUseLightSoftFading)
+		SHADER_PARAMETER(uint32, bUseSuperSampling)
 	END_SHADER_PARAMETER_STRUCT()
 
 	static FIntVector GetGroupSize()
@@ -1123,13 +1134,6 @@ class FVolumetricFogLightScatteringCS : public FGlobalShader
 
 	static FPermutationDomain RemapPermutation(FPermutationDomain PermutationVector, EShaderPlatform ShaderPlatform)
 	{
-		if (IsMobileVolumetricFog(FStaticShaderPlatform(ShaderPlatform)))
-		{
-			PermutationVector.Set<FDistanceFieldSkyOcclusion>(false);
-			PermutationVector.Set<FTemporalReprojection>(false);
-			PermutationVector.Set<FSampleLightFunctionAtlas>(false);
-			PermutationVector.Set<FMegaLights>(false);
-		}
 
 		if (!FDataDrivenShaderPlatformInfo::GetSupportsLumenGI(ShaderPlatform))
 		{
@@ -1156,12 +1160,65 @@ class FVolumetricFogLightScatteringCS : public FGlobalShader
 			PermutationVector.Set<FVirtualShadowMap>(false);
 		}
 
+		if (PermutationVector.Get< FUbershader>())
+		{
+			PermutationVector.Set<FTemporalReprojection>(true);
+			PermutationVector.Set<FDistanceFieldSkyOcclusion>(true);
+			PermutationVector.Set<FLumenGI>(FDataDrivenShaderPlatformInfo::GetSupportsLumenGI(ShaderPlatform));
+			PermutationVector.Set<FVirtualShadowMap>(true);
+			PermutationVector.Set<FRaytracedShadowsVolume>(true);
+			PermutationVector.Set<FSampleLightFunctionAtlas>(true);
+			PermutationVector.Set<FMegaLights>(true);
+			PermutationVector.Set<FLightSoftFading>(true);
+			PermutationVector.Set<FSuperSampling>(true);
+		}
+
+		if (IsMobileVolumetricFog(FStaticShaderPlatform(ShaderPlatform)))
+		{
+			PermutationVector.Set<FDistanceFieldSkyOcclusion>(false);
+			PermutationVector.Set<FTemporalReprojection>(false);
+			PermutationVector.Set<FSampleLightFunctionAtlas>(false);
+			PermutationVector.Set<FMegaLights>(false);
+			PermutationVector.Set<FRaytracedShadowsVolume>(false);
+			PermutationVector.Set<FVirtualShadowMap>(false);
+		}
 		return PermutationVector;
 	}
 
 	static bool ShouldCompilePermutation(const FGlobalShaderPermutationParameters& Parameters)
 	{
 		FPermutationDomain PermutationVector(Parameters.PermutationId);
+		if (PermutationVector.Get<FUbershader>())
+		{
+			const bool bLumenGI = FDataDrivenShaderPlatformInfo::GetSupportsLumenGI(Parameters.Platform) ? PermutationVector.Get<FLumenGI>() : PermutationVector.Get<FLumenGI>() == 0;
+			if (IsMobileVolumetricFog(FStaticShaderPlatform(Parameters.Platform)))
+			{
+				return 
+					bLumenGI &&
+					PermutationVector.Get<FTemporalReprojection>()==0 &&
+					PermutationVector.Get<FDistanceFieldSkyOcclusion>()==0 &&
+					PermutationVector.Get<FVirtualShadowMap>()==0 &&
+					PermutationVector.Get<FRaytracedShadowsVolume>()==0 &&
+					PermutationVector.Get<FSampleLightFunctionAtlas>()==0 &&
+					PermutationVector.Get<FMegaLights>()==0 &&
+					PermutationVector.Get<FLightSoftFading>() &&
+					PermutationVector.Get<FSuperSampling>();
+			}
+			else
+			{
+				return 
+					bLumenGI &&
+					PermutationVector.Get<FTemporalReprojection>() &&
+					PermutationVector.Get<FDistanceFieldSkyOcclusion>() &&
+					PermutationVector.Get<FVirtualShadowMap>() &&
+					PermutationVector.Get<FRaytracedShadowsVolume>() &&
+					PermutationVector.Get<FSampleLightFunctionAtlas>() &&
+					PermutationVector.Get<FMegaLights>() &&
+					PermutationVector.Get<FLightSoftFading>() &&
+					PermutationVector.Get<FSuperSampling>();
+			}
+		}
+		
 		if (RemapPermutation(PermutationVector, Parameters.Platform) != PermutationVector)
 		{
 			return false;
@@ -1618,6 +1675,11 @@ void FSceneRenderer::ComputeVolumetricFog(FRDGBuilder& GraphBuilder,
 #if RHI_RAYTRACING		
 		RenderRaytracedDirectionalShadowVolume(GraphBuilder, View, *Scene, IntegrationData, RaytracedShadowsVolume);
 #endif
+		const bool bUseRaytracedShadowsVolume = RaytracedShadowsVolume != nullptr;
+		if (!bUseRaytracedShadowsVolume)
+		{
+			RaytracedShadowsVolume = VolumetricBlackDummyTexture;
+		}
 
 		{
 			RDG_EVENT_SCOPE(GraphBuilder, "VolumetricFog::InitialiseVolume");
@@ -1710,6 +1772,10 @@ void FSceneRenderer::ComputeVolumetricFog(FRDGBuilder& GraphBuilder,
 			PassParameters->MegaLightsVolume = View.GetMegaLightsVolume().Texture;
 			PassParameters->RWLightScattering = IntegrationData.LightScatteringUAV;
 			PassParameters->VirtualShadowMapSamplingParameters = VirtualShadowMapArray.GetSamplingParameters(GraphBuilder, ViewIndex);
+			if (!PassParameters->MegaLightsVolume)
+			{
+				PassParameters->MegaLightsVolume = VolumetricBlackDummyTexture;
+			}
 
 			FDistanceFieldAOParameters AOParameterData(Scene->DefaultMaxDistanceFieldOcclusionDistance);
 			if (Scene->SkyLight
@@ -1798,8 +1864,7 @@ void FSceneRenderer::ComputeVolumetricFog(FRDGBuilder& GraphBuilder,
 			PassParameters->CloudShadowmapFarDepthKm = CloudShadowmap_FarDepthKm;
 			PassParameters->CloudShadowmapStrength = CloudShadowmap_Strength;
 			PassParameters->CloudShadowmapTranslatedWorldToLightClipMatrix = CloudWorldToLightClipShadowMatrix;
-
-			PassParameters->RaytracedShadowsVolume = RaytracedShadowsVolume ? GraphBuilder.CreateSRV(RaytracedShadowsVolume) : nullptr;
+			PassParameters->RaytracedShadowsVolume = GraphBuilder.CreateSRV(RaytracedShadowsVolume);
 
 			PassParameters->LightFunctionAtlas = LightFunctionAtlas::BindGlobalParameters(GraphBuilder, View);
 			if (LightsToInject.DirectionalLightFunction && bUseLightFunctionAtlasEnabledAndSupported)
@@ -1816,7 +1881,6 @@ void FSceneRenderer::ComputeVolumetricFog(FRDGBuilder& GraphBuilder,
 			const bool bUseLumenGI = View.GetLumenTranslucencyGIVolume().Texture0 != nullptr && FDataDrivenShaderPlatformInfo::GetSupportsLumenGI(View.GetShaderPlatform());
 			const bool bUseMegaLights = View.GetMegaLightsVolume().Texture != nullptr && MegaLights::IsEnabled(ViewFamily);
 			const bool bUseGlobalDistanceField = UseGlobalDistanceField() && Scene->DistanceFieldSceneData.NumObjectsInBuffer > 0;
-			const bool bUseRaytracedShadowsVolume = RaytracedShadowsVolume != nullptr;
 
 			const bool bUseDistanceFieldSkyOcclusion =
 				ViewFamily.EngineShowFlags.AmbientOcclusion
@@ -1832,20 +1896,40 @@ void FSceneRenderer::ComputeVolumetricFog(FRDGBuilder& GraphBuilder,
 				&& !IsMobileVolumetricFog(View.GetShaderPlatform());
 
 			const int32 SuperSampleCount = FVolumetricFogLightScatteringCS::GetSuperSampleCount(GVolumetricFogHistoryMissSupersampleCount);
-			check(PassParameters->VolumetricFogParameters.HistoryMissSuperSampleCount == SuperSampleCount)
+			check(PassParameters->VolumetricFogParameters.HistoryMissSuperSampleCount == SuperSampleCount);
+
+			// Ubershader "dynamic-if" parameters
+			PassParameters->bUseTemporalReprojection = bUseTemporalReprojection;
+			PassParameters->bDistanceFieldSkyOcclusion = bUseDistanceFieldSkyOcclusion;
+			PassParameters->bLumenGI = bUseLumenGI;
+			PassParameters->bVirtualShadowMap = VirtualShadowMapArray.IsAllocated();
+			PassParameters->bUseRaytracedShadowsVolume = bUseRaytracedShadowsVolume;
+			PassParameters->bUseLightFunctionAtlas = bUseLightFunctionAtlasEnabledAndSupported;
+			PassParameters->bUseMegaLights = bUseMegaLights;
+			PassParameters->bUseLightSoftFading = GetVolumetricFogLightSoftFading() > 0.0f;
+			PassParameters->bUseSuperSampling = SuperSampleCount > 1;
 
 			FVolumetricFogLightScatteringCS::FPermutationDomain PermutationVector;
-			PermutationVector.Set< FVolumetricFogLightScatteringCS::FTemporalReprojection >(bUseTemporalReprojection);
-			PermutationVector.Set< FVolumetricFogLightScatteringCS::FDistanceFieldSkyOcclusion >(bUseDistanceFieldSkyOcclusion);
-			PermutationVector.Set< FVolumetricFogLightScatteringCS::FSuperSampling>(SuperSampleCount > 1);
-			PermutationVector.Set< FVolumetricFogLightScatteringCS::FLumenGI >(bUseLumenGI);
-			PermutationVector.Set< FVolumetricFogLightScatteringCS::FVirtualShadowMap >(VirtualShadowMapArray.IsAllocated() );
-			PermutationVector.Set< FVolumetricFogLightScatteringCS::FRaytracedShadowsVolume >(bUseRaytracedShadowsVolume);
-			PermutationVector.Set< FVolumetricFogLightScatteringCS::FSampleLightFunctionAtlas >(bUseLightFunctionAtlasEnabledAndSupported);
-			PermutationVector.Set< FVolumetricFogLightScatteringCS::FMegaLights>(bUseMegaLights);
-			PermutationVector.Set< FVolumetricFogLightScatteringCS::FLightSoftFading >(GetVolumetricFogLightSoftFading() > 0.0f);
-
+			PermutationVector.Set< FVolumetricFogLightScatteringCS::FTemporalReprojection >(PassParameters->bUseTemporalReprojection > 0);
+			PermutationVector.Set< FVolumetricFogLightScatteringCS::FDistanceFieldSkyOcclusion >(PassParameters->bDistanceFieldSkyOcclusion > 0);
+			PermutationVector.Set< FVolumetricFogLightScatteringCS::FSuperSampling>(PassParameters->bUseSuperSampling > 0);
+			PermutationVector.Set< FVolumetricFogLightScatteringCS::FLumenGI >(PassParameters->bLumenGI > 0);
+			PermutationVector.Set< FVolumetricFogLightScatteringCS::FVirtualShadowMap >(PassParameters->bVirtualShadowMap > 0);
+			PermutationVector.Set< FVolumetricFogLightScatteringCS::FRaytracedShadowsVolume >(PassParameters->bUseRaytracedShadowsVolume > 0);
+			PermutationVector.Set< FVolumetricFogLightScatteringCS::FSampleLightFunctionAtlas >(PassParameters->bUseLightFunctionAtlas > 0);
+			PermutationVector.Set< FVolumetricFogLightScatteringCS::FMegaLights>(PassParameters->bUseMegaLights > 0);
+			PermutationVector.Set< FVolumetricFogLightScatteringCS::FLightSoftFading >(PassParameters->bUseLightSoftFading > 0);
+
+			PermutationVector = FVolumetricFogLightScatteringCS::RemapPermutation(PermutationVector, View.GetShaderPlatform());
 			auto ComputeShader = View.ShaderMap->GetShader< FVolumetricFogLightScatteringCS >(PermutationVector);
+
+			// Ubershader default parameters and feature dynamic branching
+			if (!PrecachePSO(FComputePipelineStateInitializer(ComputeShader.GetComputeShader())))
+			{
+				PermutationVector.Set< FVolumetricFogLightScatteringCS::FUbershader >(true);
+				PermutationVector = FVolumetricFogLightScatteringCS::RemapPermutation(PermutationVector, View.GetShaderPlatform());
+				ComputeShader = View.ShaderMap->GetShader< FVolumetricFogLightScatteringCS >( PermutationVector );
+			}
 			ClearUnusedGraphResources(ComputeShader, PassParameters);
 
 			GraphBuilder.AddPass(
