Improve HDR UI composition

Modify r.HDR.UI.Luminance default from 300 nits to 203 nits to match BT.2408.

Add r.HDR.UI.CompositeBlend to choose between gamma/linear composition, and change the default blend space from linear to gamma, which approximates blending behavior for UI in SDR. Undo premultiplied alpha and reapply during HDR alpha blend. This is a behavior change for the previous linear blend implementation.

Add r.HDR.UI.CompositeBlend.Gamma to specify the power for the HDR UI gamma blend space.

Also add r.HDR.UI.CompositeEOTF.Gamma to specify the SDR EOTF.

#Synced-CL 49917480

diff --git a/Engine/Shaders/Private/CompositeUIPixelShader.usf b/Engine/Shaders/Private/CompositeUIPixelShader.usf
--- a/Engine/Shaders/Private/CompositeUIPixelShader.usf
+++ b/Engine/Shaders/Private/CompositeUIPixelShader.usf
@@ -33,47 +33,66 @@ SamplerState SceneSampler;
 float UILevel;
 float UILuminance;
 int UICompositeEOTF;
+float UICompositeEOTFGamma;
+int UICompositeBlend;
+float UICompositeBlendGamma;
+float UICompositeBlendInvGamma;
 
 // Apply "EOTF" on SDR UI for HDR composition.
 float3 LinearizeColor(float3 EncodedColor)
 {
-	return (UICompositeEOTF == 0) ? pow(EncodedColor, 2.4) : sRGBToLinear(EncodedColor);
-}
-
-static const float ScRGBScaleFactor = 80.0f; // used to convert between ScRGB and nits
-
-void HDRBackBufferConvertPS(
-	FScreenVertexOutput Input,
-	out float3 OutColor : SV_Target0
-	)
-{
-	OutColor = 0;
-	float3 SceneColor = Texture2DSample(SceneTexture, SceneSampler, Input.UV).xyz ;
-	SceneColor = ST2084ToLinear( SceneColor ) / ScRGBScaleFactor; 
-	OutColor.rgb = SceneColor;
+	return (UICompositeEOTF == 0) ? pow(EncodedColor, UICompositeEOTFGamma) : sRGBToLinear(EncodedColor);
 }
 
 float3 ComposeUIAndScene(float3 SceneColor, float4 InUIColor, float InUILevel)
 {
 	BRANCH
-	if (InUIColor.w > 0.f && InUIColor.w < 1.f)
+	if (InUIColor.a > 0.f)
 	{
-		// Clamp gamut to sRGB as extended gamut colors bleeding into the UI can look funny
-		SceneColor = max(SceneColor, 0.f);
-
-		// Tonemap HDR under transparent UI with a simple Reinhard to the max luminance of the UI
-		// This prevents HDR bleed through destroying UI legibility
-		// Rec2020 coefficients to compute luminance
-		float KR = 0.2627, KG = 0.678, KB = 0.0593;
-		float Luminance = dot(SceneColor, half3(KR, KG, KB)) / InUILevel;
-		float OutL = 1.f / (Luminance + 1.f);
-
-		// Ease out remapping to avoid hard transitions where UI is near zero opacity
-		SceneColor *= lerp(1.f, OutL * InUILevel, InUIColor.w);
+		InUIColor.rgb *= InUILevel;
+
+		BRANCH
+		if (InUIColor.a < 1.f)
+		{
+			BRANCH
+			if (UICompositeBlend)
+			{
+				// Clamp gamut to sRGB as extended gamut colors bleeding into the UI can look funny
+				SceneColor = max(SceneColor, 0.f);
+
+				// Tonemap HDR under transparent UI with a simple Reinhard to the max luminance of the UI
+				// This prevents HDR bleeding through and destroying UI legibility
+				// Rec2020 coefficients to compute luminance
+				float KR = 0.2627, KG = 0.678, KB = 0.0593;
+				float Luminance = dot(SceneColor, half3(KR, KG, KB)) / InUILevel;
+				float OutL = 1.f / (Luminance + 1.f);
+
+				// Ease out remapping to avoid hard transitions where UI is near zero opacity
+				SceneColor *= lerp(1.f, OutL * InUILevel, InUIColor.a);
+
+				// Composite in linear blend space
+				SceneColor = lerp(SceneColor, InUIColor.rgb, InUIColor.a);
+			}
+			else
+			{
+				// Convert inputs to gamma blend space
+				InUIColor.rgb = pow(saturate(InUIColor.rgb / UE_PQ_WHITE_NITS), UICompositeBlendInvGamma);
+				SceneColor = pow(saturate(SceneColor / UE_PQ_WHITE_NITS), UICompositeBlendInvGamma);
+
+				// Composite in gamma blend space
+				SceneColor = lerp(SceneColor, InUIColor.rgb, InUIColor.a);
+
+				// Convert result to nits
+				SceneColor = pow(SceneColor, UICompositeBlendGamma) * UE_PQ_WHITE_NITS;
+			}
+		}
+		else
+		{
+			SceneColor = InUIColor.rgb;
+		}
 	}
 
-	// Composite, assuming pre-multiplied alpha
-	return SceneColor * (1.f - InUIColor.w) + InUIColor.xyz * InUILevel;
+	return SceneColor;
 }
 
 float3 ApplyColorDeficiency(float3 Rec2020Color)
@@ -96,18 +115,22 @@ float4 ComputeHDRUIColor(float2 InputUV)
 {
 	const float3x3 sRGB_2_Rec2020 = mul(XYZ_2_Rec2020_MAT, sRGB_2_XYZ_MAT);
 	float4 UIColor = Texture2DSample(UITexture, UISampler, InputUV);
-	UIColor.xyz = LinearizeColor(UIColor.xyz);
-	UIColor.xyz = mul(sRGB_2_Rec2020, UIColor.xyz) * UILuminance;
+
+	// Undo premultiplied alpha to apply in HDR blend space
+	UIColor.rgb = (UIColor.a > 0.) ? (UIColor.rgb / UIColor.a) : 0.;
+
+	UIColor.rgb = LinearizeColor(UIColor.rgb);
+	UIColor.rgb = mul(sRGB_2_Rec2020, UIColor.rgb) * UILuminance;
 	return UIColor;
 }
 
 float3 ComputeHDRSceneColor(float3 SceneColor)
 {
-	const float3x3 sRGB_2_Rec2020 = mul(XYZ_2_Rec2020_MAT, sRGB_2_XYZ_MAT);
 #if !SCRGB_ENCODING
-	SceneColor.xyz = ST2084ToLinear(SceneColor.xyz);
+	SceneColor = ST2084ToLinear(SceneColor);
 #else
-	SceneColor.xyz = mul(sRGB_2_Rec2020, SceneColor.xyz) * ScRGBScaleFactor;
+	const float3x3 sRGB_2_Rec2020 = mul(XYZ_2_Rec2020_MAT, sRGB_2_XYZ_MAT);
+	SceneColor = mul(sRGB_2_Rec2020, SceneColor) * UE_SCRGB_WHITE_NITS;
 #endif
 	return SceneColor;
 }
@@ -158,7 +181,7 @@ void Main(
 	OutColor.xyz = LinearToST2084(OutColor.xyz);
 #else
 	const float3x3 Rec2020_2_sRGB = mul(XYZ_2_sRGB_MAT, Rec2020_2_XYZ_MAT);
-	OutColor.xyz = mul(Rec2020_2_sRGB, OutColor.xyz / ScRGBScaleFactor);
+	OutColor.xyz = mul(Rec2020_2_sRGB, OutColor.xyz / UE_SCRGB_WHITE_NITS);
 #endif
 }
 
@@ -214,7 +237,7 @@ void CompositeUICS(uint3 DispatchThreadId : SV_DispatchThreadID)
 	OutColor.xyz = LinearToST2084(OutColor.xyz);
 #else
 	const float3x3 Rec2020_2_sRGB = mul(XYZ_2_sRGB_MAT, Rec2020_2_XYZ_MAT);
-	OutColor.xyz = mul(Rec2020_2_sRGB, OutColor.xyz / ScRGBScaleFactor);
+	OutColor.xyz = mul(Rec2020_2_sRGB, OutColor.xyz / UE_SCRGB_WHITE_NITS);
 #endif
 
 	RWSceneTexture[DispatchThreadId.xy] = OutColor;
@@ -241,7 +264,7 @@ void CompositeUIForBlur(FScreenVertexOutput Input, out float4 OutColor : SV_Targ
 
 	float3 SceneColor = Texture2DSample(SceneTexture, UISampler, Input.UV).xyz;
 	SceneColor = ComputeHDRSceneColor(SceneColor);
-	OutColor.xyz = ComposeUIAndScene(SceneColor.xyz, UIColor, UILevel) / ScRGBScaleFactor;
+	OutColor.xyz = ComposeUIAndScene(SceneColor.xyz, UIColor, UILevel) / UE_SCRGB_WHITE_NITS;
 	OutColor.a = 1;
 }
 
diff --git a/Engine/Shaders/Private/GammaCorrectionCommon.ush b/Engine/Shaders/Private/GammaCorrectionCommon.ush
--- a/Engine/Shaders/Private/GammaCorrectionCommon.ush
+++ b/Engine/Shaders/Private/GammaCorrectionCommon.ush
@@ -133,40 +133,46 @@ aces100nitFitInverse(float3 FilmColor)
 	return inverse;
 }
 
+static const float UE_PQ_WHITE_NITS = 10000.;
+
 //
 // Dolby PQ transforms
 //
 float3
-ST2084ToLinear(float3 pq)
+ST2084ToLinearNormalized(float3 pq)
 {
 	const float m1 = 0.1593017578125; // = 2610. / 4096. * .25;
 	const float m2 = 78.84375; // = 2523. / 4096. *  128;
 	const float c1 = 0.8359375; // = 2392. / 4096. * 32 - 2413./4096.*32 + 1;
 	const float c2 = 18.8515625; // = 2413. / 4096. * 32;
 	const float c3 = 18.6875; // = 2392. / 4096. * 32;
-	const float C = 10000.;
 
 	float3 Np = pow( pq, 1./m2 );
 	float3 L = Np - c1;
 	L = max(0., L);
 	L = L / (c2 - c3 * Np);
 	L = pow( L, 1./m1 );
-	float3 P = L * C;
+	float3 P = L;
 
 	return P;
 }
 
 float3
-LinearToST2084(float3 lin)
+ST2084ToLinear(float3 pq)
+{
+	return ST2084ToLinearNormalized(pq) * UE_PQ_WHITE_NITS;
+}
+
+float3
+LinearToST2084Normalized(float3 lin)
 {
 	const float m1 = 0.1593017578125; // = 2610. / 4096. * .25;
 	const float m2 = 78.84375; // = 2523. / 4096. *  128;
 	const float c1 = 0.8359375; // = 2392. / 4096. * 32 - 2413./4096.*32 + 1;
 	const float c2 = 18.8515625; // = 2413. / 4096. * 32;
 	const float c3 = 18.6875; // = 2392. / 4096. * 32;
-	const float C = 10000.;
 
-	float3 L = lin/C;
+	float3 L = lin;
 	float3 Lm = pow(L, m1);
 	float3 N1 = ( c1 + c2 * Lm );
 	float3 N2 = ( 1.0 + c3 * Lm );
@@ -176,6 +182,12 @@ LinearToST2084(float3 lin)
 	return P;
 }
 
+float3
+LinearToST2084(float3 lin)
+{
+	return LinearToST2084Normalized(lin / UE_PQ_WHITE_NITS);
+}
+
 //
 // Sony S-Log3
 //
diff --git a/Engine/Source/Runtime/SlateRHIRenderer/Private/SlatePostProcessor.cpp b/Engine/Source/Runtime/SlateRHIRenderer/Private/SlatePostProcessor.cpp
--- a/Engine/Source/Runtime/SlateRHIRenderer/Private/SlatePostProcessor.cpp
+++ b/Engine/Source/Runtime/SlateRHIRenderer/Private/SlatePostProcessor.cpp
@@ -34,12 +34,30 @@ float GetSlateHDRUILuminance()
 	return CVar ? CVar->GetFloat() : 300.0f;
 }
 
-int GetSlateHDRUICompositeEOTF()
+int32 GetSlateHDRUICompositeEOTF()
 {
 	static auto CVar = IConsoleManager::Get().FindConsoleVariable(TEXT("r.HDR.UI.CompositeEOTF"));
 	return CVar ? CVar->GetInt() : 0;
 }
 
+float GetSlateHDRUICompositeEOTFGamma()
+{
+	static auto CVar = IConsoleManager::Get().FindConsoleVariable(TEXT("r.HDR.UI.CompositeEOTF.Gamma"));
+	return CVar ? CVar->GetFloat() : 2.2f;
+}
+
+int32 GetSlateHDRUICompositeBlend()
+{
+	static auto CVar = IConsoleManager::Get().FindConsoleVariable(TEXT("r.HDR.UI.CompositeBlend"));
+	return CVar ? CVar->GetInt() : 0;
+}
+
+float GetSlateHDRUICompositeBlendGamma()
+{
+	static auto CVar = IConsoleManager::Get().FindConsoleVariable(TEXT("r.HDR.UI.CompositeBlend.Gamma"));
+	return CVar ? CVar->GetFloat() : 2.3f;
+}
+
 ETextureCreateFlags GetSlateTransientRenderTargetFlags()
 {
 	return ETextureCreateFlags::RenderTargetable | ETextureCreateFlags::ShaderResource | ETextureCreateFlags::FastVRAM
@@ -75,6 +93,10 @@ public:
 		SHADER_PARAMETER(float, UILevel)
 		SHADER_PARAMETER(float, UILuminance)
 		SHADER_PARAMETER(int32, UICompositeEOTF)
+		SHADER_PARAMETER(float, UICompositeEOTFGamma)
+		SHADER_PARAMETER(int32, UICompositeBlend)
+		SHADER_PARAMETER(float, UICompositeBlendGamma)
+		SHADER_PARAMETER(float, UICompositeBlendInvGamma)
 		SHADER_PARAMETER(FVector2f, UITextureSize)
 		RENDER_TARGET_BINDING_SLOTS()
 	END_SHADER_PARAMETER_STRUCT()
@@ -134,6 +156,11 @@ FScreenPassTexture AddSlateCompositeHDRForBlurPass(FRDGBuilder& GraphBuilder, co
 	PassParameters->UILevel = GetSlateHDRUILevel();
 	PassParameters->UILuminance = GetSlateHDRUILuminance();
 	PassParameters->UICompositeEOTF = GetSlateHDRUICompositeEOTF();
+	PassParameters->UICompositeEOTFGamma = GetSlateHDRUICompositeEOTFGamma();
+	PassParameters->UICompositeBlend = GetSlateHDRUICompositeBlend();
+	const float UICompositeBlendGamma = GetSlateHDRUICompositeBlendGamma();
+	PassParameters->UICompositeBlendGamma = UICompositeBlendGamma;
+	PassParameters->UICompositeBlendInvGamma = 1.f / UICompositeBlendGamma;
 
 	TShaderMapRef<FCompositeHDRForBlurPS> PixelShader(ShaderMap, PermutationVector);
 	AddDrawScreenPass(GraphBuilder, RDG_EVENT_NAME("CompositeHDR"), FeatureLevel, OutputViewport, InputViewport, PixelShader, PassParameters);
diff --git a/Engine/Source/Runtime/SlateRHIRenderer/Private/SlateRHIRenderer.cpp b/Engine/Source/Runtime/SlateRHIRenderer/Private/SlateRHIRenderer.cpp
--- a/Engine/Source/Runtime/SlateRHIRenderer/Private/SlateRHIRenderer.cpp
+++ b/Engine/Source/Runtime/SlateRHIRenderer/Private/SlateRHIRenderer.cpp
@@ -80,18 +80,40 @@ static TAutoConsoleVariable<float> CVarUILevel(
 
 static TAutoConsoleVariable<float> CVarHDRUILuminance(
 	TEXT("r.HDR.UI.Luminance"),
-	300.0f,
+	203.f,
 	TEXT("Base Luminance in nits for UI elements when compositing into HDR framebuffer. Gets multiplied by r.HDR.UI.Level"),
 	ECVF_RenderThreadSafe);
 
-static TAutoConsoleVariable<int> CVarHDRUICompositeEOTF(
+static TAutoConsoleVariable<int32> CVarHDRUICompositeEOTF(
 	TEXT("r.HDR.UI.CompositeEOTF"),
 	0,
 	TEXT("EOTF used when compositing the UI layer:\n")
-	TEXT("0: BT.1886 to approximate the viewing environment specified by Rec. 2100 (default)\n")
+	TEXT("0: Gamma (default), see r.HDR.UI.CompositeEOTF.Gamma\n")
 	TEXT("1: sRGB"),
 	ECVF_RenderThreadSafe);
 
+static TAutoConsoleVariable<float> CVarHDRUICompositeEOTFGamma(
+	TEXT("r.HDR.UI.CompositeEOTF.Gamma"),
+	2.4f,
+	TEXT("The power used if using gamma to composite the UI layer:\n")
+	TEXT("Defaults to 2.4 for BT.1886 specified by BT.2408."),
+	ECVF_RenderThreadSafe);
+
+static TAutoConsoleVariable<int32> CVarHDRUICompositeBlend(
+	TEXT("r.HDR.UI.CompositeBlend"),
+	0,
+	TEXT("The blend space used when compositing the UI layer:\n")
+	TEXT("0: Gamma (default), see r.HDR.UI.CompositeBlend.Gamma\n")
+	TEXT("1: Linear"),
+	ECVF_RenderThreadSafe);
+
+static TAutoConsoleVariable<float> CVarHDRUICompositeBlendGamma(
+	TEXT("r.HDR.UI.CompositeBlend.Gamma"),
+	2.4f,
+	TEXT("The power used if using gamma to composite the UI layer:\n")
+	TEXT("Defaults to 2.4 to match r.HDR.UI.CompositeEOTF.Gamma."),
+	ECVF_RenderThreadSafe);
+
 static TAutoConsoleVariable<int32> CVarUICompositeMode(
 	TEXT("r.HDR.UI.CompositeMode"),
 	1,
@@ -582,6 +604,10 @@ BEGIN_SHADER_PARAMETER_STRUCT(FCompositeShaderCommonParameters, )
 	SHADER_PARAMETER(float, UILevel)
 	SHADER_PARAMETER(float, UILuminance)
 	SHADER_PARAMETER(int32, UICompositeEOTF)
+	SHADER_PARAMETER(float, UICompositeEOTFGamma)
+	SHADER_PARAMETER(int32, UICompositeBlend)
+	SHADER_PARAMETER(float, UICompositeBlendGamma)
+	SHADER_PARAMETER(float, UICompositeBlendInvGamma)
 	SHADER_PARAMETER(float, ColorVisionDeficiencyType)
 	SHADER_PARAMETER(float, ColorVisionDeficiencySeverity)
 	SHADER_PARAMETER(float, bCorrectDeficiency)
@@ -808,6 +834,11 @@ FSlateDrawWindowPassOutputs FSlateRHIRenderer::DrawWindow_RenderThread(FRDGBuild
 			CommonParameters.UILevel = CVarUILevel.GetValueOnRenderThread();
 			CommonParameters.UILuminance = CVarHDRUILuminance.GetValueOnRenderThread();
 			CommonParameters.UICompositeEOTF = CVarHDRUICompositeEOTF.GetValueOnRenderThread();
+			CommonParameters.UICompositeEOTFGamma = CVarHDRUICompositeEOTFGamma.GetValueOnRenderThread();
+			CommonParameters.UICompositeBlend = CVarHDRUICompositeBlend.GetValueOnRenderThread();
+			const float UICompositeBlendGamma = CVarHDRUICompositeBlendGamma.GetValueOnRenderThread();
+			CommonParameters.UICompositeBlendGamma = UICompositeBlendGamma;
+			CommonParameters.UICompositeBlendInvGamma = 1.f / UICompositeBlendGamma;
 			CommonParameters.ColorVisionDeficiencySeverity = (float)GSlateColorDeficiencySeverity;
 			CommonParameters.ColorVisionDeficiencyType = (float)GSlateColorDeficiencyType;
 			CommonParameters.bSimulateCorrectionWithDeficiency = GSlateShowColorDeficiencyCorrectionWithDeficiency ? 1.0f : 0.0f;
