Fix buffer visualizations incorrectly applying pre-exposure.
- using "Disable Pre Exposure Scale" (bDisablePreExposureScale) in the post process materials used to render the visualizations.

#virtualized
#Synced-CL 49955598

diff --git a/Engine/Shaders/Private/PostProcessMaterialShaders.usf b/Engine/Shaders/Private/PostProcessMaterialShaders.usf
--- a/Engine/Shaders/Private/PostProcessMaterialShaders.usf
+++ b/Engine/Shaders/Private/PostProcessMaterialShaders.usf
@@ -586,7 +586,7 @@ void MainPS(
 	FullSceneColor.rgb = sqrt(FullSceneColor.rgb);
 #endif
 
-#if !MATERIALBLENDING_MODULATE
+#if !MATERIALBLENDING_MODULATE && !POST_PROCESS_DISABLE_PRE_EXPOSURE_SCALE
 	FullSceneColor.xyz *= View.PreExposure;
 #endif
 #endif
