Substrate - Toon - added shadow density control for soft fading.
SubstrateLoadSubsurfaceHeader call sites have been checked and should be all good (no diffusion on pixels with toon shading for instance).
Added new general shadow model that could be merge with substrate per pixel shadow later.

#tests PC editor dx12
#rnx

#Synced-CL 49930030

diff --git a/Engine/Shaders/Private/DeferredShadingCommon.ush b/Engine/Shaders/Private/DeferredShadingCommon.ush
--- a/Engine/Shaders/Private/DeferredShadingCommon.ush
+++ b/Engine/Shaders/Private/DeferredShadingCommon.ush
@@ -310,7 +310,8 @@ bool IsSubsurfaceModel(int ShadingModel)
 		|| ShadingModel == SHADINGMODELID_SUBSURFACE_PROFILE
 		|| ShadingModel == SHADINGMODELID_TWOSIDED_FOLIAGE
 		|| ShadingModel == SHADINGMODELID_HAIR
-		|| ShadingModel == SHADINGMODELID_EYE;
+		|| ShadingModel == SHADINGMODELID_EYE
+		|| ShadingModel == SHADINGMODELID_SUBSTRATE_TOON; // For shadow transmittance similar to Subsurface.
 }
 
 bool UseSubsurfaceProfile(int ShadingModel)
diff --git a/Engine/Shaders/Private/PostProcessGBufferHints.usf b/Engine/Shaders/Private/PostProcessGBufferHints.usf
--- a/Engine/Shaders/Private/PostProcessGBufferHints.usf
+++ b/Engine/Shaders/Private/PostProcessGBufferHints.usf
@@ -8,6 +8,11 @@
 #include "DeferredShadingCommon.ush"
 #include "ScreenPass.ush"
 
+#if SUBSTRATE_ENABLED
+#include "ToonProfileCommon.ush"
+#include "Substrate/Substrate.ush"
+#endif
+
 #define SUPPORT_GBUFFER_FETCH (SHADING_PATH_DEFERRED && FEATURE_LEVEL >= FEATURE_LEVEL_SM5)
 
 float2 Input_ExtentInverse;
@@ -143,6 +148,7 @@ void PrintShadingMode(inout FShaderPrintContext Ctx, uint In)
 	if (In == SHADINGMODELID_SINGLELAYERWATER)		{ Print(Ctx, TEXT("SINGLELAYERWATER"), FontRed); return; }
 	if (In == SHADINGMODELID_THIN_TRANSLUCENT)		{ Print(Ctx, TEXT("THIN_TRANSLUCENT"), FontRed); return; }
 	if (In == SHADINGMODELID_SUBSTRATE)				{ Print(Ctx, TEXT("SUBSTRATE"), FontRed); return; }
+	if (In == SHADINGMODELID_SUBSTRATE_TOON)		{ Print(Ctx, TEXT("TOON"), FontRed); return; }
 	if (In == SHADINGMODELID_NUM)					{ Print(Ctx, TEXT("NUM"), FontRed); return; }
 	if (In == SHADINGMODELID_MASK) 					{ Print(Ctx, TEXT("MASK"), FontRed); return; }
 }
@@ -221,7 +227,7 @@ void MainCS(uint3 DispatchThreadId : SV_DispatchThreadID)
 		RectMin = Ctx.Pos;
 		Print(Ctx, TEXT("Custom"), FontOrange);Newline(Ctx, RectMax);
 		Print(Ctx, TEXT("Shading Model ID  : "), FontWhite); PrintShadingMode(Ctx, Data.ShadingModelID); Newline(Ctx, RectMax);
-		Print(Ctx, TEXT("Custom Data       : "), FontWhite); Print(Ctx, Data.CustomData, FontYellow); Newline(Ctx, RectMax);
+		Print(Ctx, TEXT(" Raw Custom Data  : "), FontWhite); Print(Ctx, Data.CustomData, FontYellow); Newline(Ctx, RectMax);
 		if (Data.ShadingModelID == SHADINGMODELID_CLEAR_COAT)
 		{
 		Print(Ctx, TEXT("Clear Coat        : "), FontWhite); Print(Ctx, Data.CustomData.x, FontYellow); Newline(Ctx, RectMax);
@@ -253,6 +259,19 @@ void MainCS(uint3 DispatchThreadId : SV_DispatchThreadID)
 			Print(Ctx, TEXT("SSS SpecColor     : "), FontWhite); Print(Ctx, Data.SpecularColor); Newline(Ctx, RectMax);
 			
 		}
+	#if SUBSTRATE_ENABLED
+		else if (Data.ShadingModelID == SHADINGMODELID_SUBSTRATE_TOON)
+		{
+			float3 SpecularColor = 0.0f;
+			float Anisotropy = 0.0f;
+			float ToonProfileId = 0.0f;
+			UnpackToonCustomData(PackRGBA8(Data.CustomData), SpecularColor, Anisotropy, ToonProfileId);
+
+			Print(Ctx, TEXT("Specular Color     : "), FontWhite); Print(Ctx, SpecularColor); Newline(Ctx, RectMax);
+			Print(Ctx, TEXT("Toon Profil ID     : "), FontWhite); Print(Ctx, ToonProfileId); Newline(Ctx, RectMax);
+			Print(Ctx, TEXT("Anisotropy         : "), FontWhite); Print(Ctx, Anisotropy); Newline(Ctx, RectMax);
+		}
+	#endif
 		PrintBackgroundRect(Ctx, RectMin, RectMax);
 		Newline(Ctx);
 		
diff --git a/Engine/Shaders/Private/ShadowProjectionPixelShader.usf b/Engine/Shaders/Private/ShadowProjectionPixelShader.usf
--- a/Engine/Shaders/Private/ShadowProjectionPixelShader.usf
+++ b/Engine/Shaders/Private/ShadowProjectionPixelShader.usf
@@ -42,6 +42,7 @@
 #endif
 #include "DeferredShadingCommon.ush"
 #include "DynamicLightingCommon.ush"
+#include "ToonProfileCommon.ush"
 
 float ShadowFadeFraction;
 float ShadowSharpen;
@@ -301,18 +302,42 @@ void Main(
 			const uint SSSType = SubstrateSubSurfaceHeaderGetSSSType(SSSHeader);
 			#endif
 
+			bool bApplyPCFOverBlurCorrection = true; // Set to false when the transmittance from extinction is correctly respected. Toon & TransmitttanceShadow disables that to have a physically based non squared transmittance curve.
 			float Density = 1;
 			if (!bIsEyeOrHair)
 			{
 			#if SUBSTRATE_GBUFFER_FORMAT==1
-				if (SSSType != SSS_TYPE_NONE)
+				if (SSSType == SSS_TYPE_TRANSMITTANCE_SHADOW)
+				{
+					bApplyPCFOverBlurCorrection = false;
+					const float3 MFPcm = SubstrateSubSurfaceHeaderGetTransmittanceShadowMFPInMeters(SSSHeader) * METER_TO_CENTIMETER;
+					Density = 1.0f / max(0.000001f, MFPcm.x); // Recover extinction coefficient as 1/cm from MFP in cm.
+				}
+				else if (SSSType != SSS_TYPE_NONE)
 				{
 					// Disble accurate SSS transmission if SSS type is no using Diffusion or Diffusion Profile
 					bHasSubsurfaceTransmission = bHasSubsurfaceTransmission && (SSSType >= SSS_TYPE_DIFFUSION);
 					Density = SubsurfaceDensityFromOpacity(SubstrateSubSurfaceHeaderGetOpacity(SSSHeader));
 				}
 			#else // SUBSTRATE_GBUFFER_FORMAT==1
-				Density = SubsurfaceDensityFromOpacity(GBufferData.CustomData.a);
+				#if SUBSTRATE_ENABLED
+				if (GBufferData.ShadingModelID == SHADINGMODELID_SUBSTRATE_TOON)
+				{
+					bApplyPCFOverBlurCorrection = false;
+
+					float3 SpecularColor = 0.0f;
+					float Anisotropy = 0.0f;
+					float ToonProfileId = 0.0f;
+					UnpackToonCustomData(PackRGBA8(GBufferData.CustomData), SpecularColor, Anisotropy, ToonProfileId);
+
+					// Density is extinction as 1/cm (not 1/m) to match with the transmittance shadow evaluation from ManualPCF below.
+					Density = GetToonProfileShadowSettings(ToonProfileId).ShadowExtinctionCoefficientInInvCentimeter;
+				}
+				else
+				#endif
+				{
+					Density = SubsurfaceDensityFromOpacity(GBufferData.CustomData.a);
+				}
 			#endif // SUBSTRATE_GBUFFER_FORMAT==1
 			}
 			
@@ -362,14 +387,15 @@ void Main(
 			else
 			#endif
 			{
-				// This branch is taken by Subsurface shading model (Wrap for Substrate). 
+				// This branch is taken by Subsurface shading model (Wrap for Substrate) or the toon shading model. 
 				// => SSSTransmittanceOrOpticalDepth stores the result of ManualPCF() evaluating Transmittance.
 				SSSTransmittanceOrOpticalDepth = ManualPCF(ShadowPosition.xy, Settings);
 			}
 
 			// Legacy code maintained: here we have a square applied over the transmittance curve or encoded optical depth... 
 			// This could help reduce banding artefact, but since the square is not decoded in the lighting side, it is just wrong or here for an unknown reason now.
-			SSSTransmittanceOrOpticalDepth = ApplyPCFOverBlurCorrection(SSSTransmittanceOrOpticalDepth);
+			// It affects all transmittance and encoded optical depth: substrate per pixel MFP and Subsurface Profile optical depth as well as Subsurface shading model and Toon transmittance.
+			SSSTransmittanceOrOpticalDepth = bApplyPCFOverBlurCorrection ? ApplyPCFOverBlurCorrection(SSSTransmittanceOrOpticalDepth) : SSSTransmittanceOrOpticalDepth;
 		}
 
 	#endif
@@ -556,6 +582,9 @@ void MainOnePassPointLightPS(
 	{
 		EncodedOpticalDepth = CalculateEncodedOpticalDepth(WorldNormal, TransmissionParams, WorldSampleToLightVec, PointLightProjParameters, LightPositionAndInvRadius.w, 0);
 	}
+
+	// Transmittance for Subsurface shading model and Toon is not implemented today.
+
 	OutColor.a = IsSubsurfaceProfile ? EncodeLightAttenuation(EncodedOpticalDepth) : OutColor.b;
 #endif // USE_TRANSMISSION
 }
\ No newline at end of file
diff --git a/Engine/Shaders/Private/Substrate/Substrate.ush b/Engine/Shaders/Private/Substrate/Substrate.ush
--- a/Engine/Shaders/Private/Substrate/Substrate.ush
+++ b/Engine/Shaders/Private/Substrate/Substrate.ush
@@ -4902,6 +4902,8 @@ FSubstrateBSDF UnpackSubstrateBSDFIn(FSubstrateMaterialContainer SubstrateBuffer
 
 				UnpackToonCustomData(DataY, TOON_SPECULARCOLOR(OutBSDF), TOON_ANISOTROPY(OutBSDF), TOON_PROFILEID(OutBSDF));
 
+				BSDF_SETSSSTYPE(OutBSDF, SSS_TYPE_TRANSMITTANCE_SHADOW);
+
 				return OutBSDF;
 			}
 			else
diff --git a/Engine/Shaders/Private/Substrate/SubstrateEvaluation.ush b/Engine/Shaders/Private/Substrate/SubstrateEvaluation.ush
--- a/Engine/Shaders/Private/Substrate/SubstrateEvaluation.ush
+++ b/Engine/Shaders/Private/Substrate/SubstrateEvaluation.ush
@@ -1329,10 +1329,13 @@ FSubstrateEvaluateResult SubstrateEvaluateBSDFCommon(FSubstrateBSDFContext BSDFC
 			float ToonProfileDiffuseDither  = EvaluateToonProfileDiffuseDither(ToonProfileId, ScreenUV, ToonEvalSettings);
 			float ToonProfileSpecularDither = EvaluateToonProfileSpecularDither(ToonProfileId, ScreenUV, ToonEvalSettings);
 
+			// TransmissionThickness in fact already stores transmittance.
+			float TransmittanceShadow = ShadowTerms.TransmittanceOrOpticalThickness;
+
 			// Diffuse
-			Sample.DiffuseColor				= DiffuseColor * EvaluateToonProfileDiffuseRamp(ToonProfileId, saturate(BSDFContext.Context.NoL + ToonProfileDiffuseDither)); // AreaLight.NoL to handle capsule shape?
+			Sample.DiffuseColor				= DiffuseColor * EvaluateToonProfileDiffuseRamp(ToonProfileId, BSDFContext.Context.NoL + ToonProfileDiffuseDither); // AreaLight.NoL to handle capsule shape?
 			Sample.DiffusePathValue			= Diffuse_Lambert(Sample.DiffuseColor);
-			Sample.IntegratedDiffuseValue	= ShadowTerms.SurfaceShadow * AreaLight.Falloff * AreaLight.FalloffColor * Sample.DiffusePathValue;
+			Sample.IntegratedDiffuseValue	= TransmittanceShadow * AreaLight.Falloff * AreaLight.FalloffColor * Sample.DiffusePathValue;
 			
 			const float RoughnessExponent	= saturate(1.0 - Roughness);
 			const float RoughnessExponent2 = RoughnessExponent * RoughnessExponent; // Arbitrary: this is so that the roughness control appears smooth and linear without big jumps due to 8bit encoding.
diff --git a/Engine/Shaders/Private/Substrate/SubstrateExport.ush b/Engine/Shaders/Private/Substrate/SubstrateExport.ush
--- a/Engine/Shaders/Private/Substrate/SubstrateExport.ush
+++ b/Engine/Shaders/Private/Substrate/SubstrateExport.ush
@@ -2,6 +2,8 @@
 
 #pragma once
 
+#include "../ToonProfileCommon.ush"
+
 // Do not include that on mobile since FMobileBasePassInterpolantsVSToPS is not the same structure as FBasePassInterpolantsVSToPS
 #if SUBSTRATE_ENABLED && !defined(FMobileBasePassInterpolantsVSToPS)
 
@@ -308,7 +310,8 @@ void PackSubstrateOut(
 
 						TopLayerData.Material = TOP_LAYER_MATERIAL_VALID;
 
-						SubstratePixelHeader.SetMaterialMode(HEADER_MATERIALMODE_SLAB_COMPLEX); // Use the complex path for Toon for now.
+						SubstratePixelHeader.SetHasSubsurface(true);									// This is true as Toon uses transmittance shadow. This also prevents classification pass to clear SSSHeader.
+						SubstratePixelHeader.SetMaterialMode(HEADER_MATERIALMODE_SLAB_COMPLEX);			// Use the complex path for Toon for now.
 						SubstratePixelHeader.SetComplexPathMode(COMPLEXPATH_MODE_TOON);
 						break;
 					}
@@ -879,6 +882,10 @@ void PackSubstrateOut(
 						{
 							SUBSTRATE_STORE_UINT1(PackColorLinearToGamma2AlphaLinear(float4(TOON_DIFFUSEALBEDO(BSDF), TOON_ROUGHNESS(BSDF))));
 							SUBSTRATE_STORE_UINT1(PackToonCustomData(TOON_SPECULARCOLOR(BSDF), TOON_ANISOTROPY(BSDF), TOON_PROFILEID(BSDF), Dither));
+							SSSDataType = SSS_TYPE_TRANSMITTANCE_SHADOW;
+
+							const float ShadowMFPInMeters = GetToonProfileShadowSettings(TOON_PROFILEID(BSDF)).ShadowMFPInMeters;
+							SSSDataMFP = ShadowMFPInMeters.xxx;
 							// 12 bytes
 						}
 						break;
@@ -919,6 +926,10 @@ void PackSubstrateOut(
 				{
 					SubstrateSubSurfaceHeaderSetProfile(SSSData.Header, SSSDataRadiusScale, SSSDataProfileID);
 				}
+				else if (SSSDataType == SSS_TYPE_TRANSMITTANCE_SHADOW)
+				{
+					SubstrateSubSurfaceHeaderSetTransmittanceShadowMFPInMeters(SSSData.Header, SSSDataMFP);
+				}
 				else
 				{
 					SubstrateSubSurfaceHeaderSetNonProfile(SSSData.Header, SSSDataMFP);
diff --git a/Engine/Shaders/Private/Substrate/SubstrateSubsurface.ush b/Engine/Shaders/Private/Substrate/SubstrateSubsurface.ush
--- a/Engine/Shaders/Private/Substrate/SubstrateSubsurface.ush
+++ b/Engine/Shaders/Private/Substrate/SubstrateSubsurface.ush
@@ -25,15 +25,17 @@ struct FSubstrateSubsurfaceHeader
 	// 3: Diffusion
 	// 4: Diffusion with profile
 	// 5: SimpleVolume
+	// 6: Transmittance shadow
 
 	// 29 bits of data.
-	// If  bIsProfile: 8bit profile radius, 8bits ProfileId, 16bits unused
-	// If !bIsProfile: 10109 mean free path
+	// If  TransmittanceShadow	: 10109 mean free path. SUBSTRATE_TODO unify with Substrate per pixel MFP with option for squared transmittance and skip screen space SSS post process. That would then be a MFP of 999bits. NonProfile can also be renamed to SSSPerPixelMFP.
+	// If  bIsProfile			: 8bit profile radius, 8bits ProfileId, 16bits unused
+	// If !bIsProfile			: 10109 mean free path
 	uint Bytes;
 };
 
-#if SSS_TYPE_COUNT != 6
-#error Update this code to ensure all SSS types requiering side payload can be represented.
+#if SSS_TYPE_COUNT != 7
+#error Update this code to ensure all SSS types requiring side payload can be represented.
 #endif
 #define SSSHEADER_TYPE(Header)	(Header.Bytes & 0x7)
 #define SSSHEADER_TYPE_MASK		(0x00000007)
@@ -58,6 +60,11 @@ void SubstrateSubSurfaceHeaderSetNonProfile(inout FSubstrateSubsurfaceHeader SSS
 	SSSHeader.Bytes |= (Pack10F(MeanFreePath.x) << 22) | (Pack10F(MeanFreePath.y) << 12) | ((Pack10F(MeanFreePath.z)&0x3FE) << 2);
 }
 
+void SubstrateSubSurfaceHeaderSetTransmittanceShadowMFPInMeters(inout FSubstrateSubsurfaceHeader SSSHeader, float3 MeanFreePath)
+{
+	SubstrateSubSurfaceHeaderSetNonProfile(SSSHeader, MeanFreePath);
+}
+
 float SubstrateSubSurfaceGetWrapOpacityFromAnisotropy(float PhaseAnisotropy)
 {
 	// Reinterpret the phase function anisotropy as 'opacity' value
@@ -124,7 +131,7 @@ float SubstrateSubSurfaceHeaderGetProfileRadiusScale(in FSubstrateSubsurfaceHead
 	return bIsValidAndProfile ? UnpackR8(SSSHeader.Bytes >> 16) : 1.f;
 }
 
-float3 SubstrateSubSurfaceHeaderGetMFP(in FSubstrateSubsurfaceHeader SSSHeader)				// Only when !IsProfile
+float3 SubstrateSubSurfaceHeaderGetMFP(in FSubstrateSubsurfaceHeader SSSHeader)				// Only when !Profile (== Substrate SSSPerPixelMFP) or TransmittanceShadow.
 {
 	// xxxxxxxxxx
 	//           xxxxxxxxxx
@@ -138,6 +145,11 @@ float3 SubstrateSubSurfaceHeaderGetMFP(in FSubstrateSubsurfaceHeader SSSHeader)
 		Unpack10F((SSSHeader.Bytes >> 2) & 0x3FE)); //  9bits
 }
 
+float3 SubstrateSubSurfaceHeaderGetTransmittanceShadowMFPInMeters(in FSubstrateSubsurfaceHeader SSSHeader)
+{
+	return SubstrateSubSurfaceHeaderGetMFP(SSSHeader);
+}
+
 float SubstrateSubSurfaceHeaderGetWrapOpacity(in FSubstrateSubsurfaceHeader SSSHeader)
 {
 	const uint SSSType = SSSHEADER_TYPE(SSSHeader);
diff --git a/Engine/Shaders/Private/Substrate/SubstrateVisualizeCommon.ush b/Engine/Shaders/Private/Substrate/SubstrateVisualizeCommon.ush
--- a/Engine/Shaders/Private/Substrate/SubstrateVisualizeCommon.ush
+++ b/Engine/Shaders/Private/Substrate/SubstrateVisualizeCommon.ush
@@ -58,6 +58,10 @@ void DrawPixelFootprint(float3 P, float3 dPdx, float3 dPdy, uint2 PixelCoord)
 		{
 			MFP = GetSubsurfaceProfileMFPInCm(SubstrateSubSurfaceHeaderGetProfileId(SSSHeader)).xyz * SubstrateSubSurfaceHeaderGetProfileRadiusScale(SSSHeader);
 		}
+		else if(SSSHEADER_TYPE(SSSHeader) == SSS_TYPE_TRANSMITTANCE_SHADOW)
+		{
+			MFP = SubstrateSubSurfaceHeaderGetTransmittanceShadowMFPInMeters(SSSHeader);
+		}
 		else
 		{
 			MFP = SubstrateSubSurfaceHeaderGetMFP(SSSHeader);
diff --git a/Engine/Shaders/Private/ToonProfile.usf b/Engine/Shaders/Private/ToonProfile.usf
--- a/Engine/Shaders/Private/ToonProfile.usf
+++ b/Engine/Shaders/Private/ToonProfile.usf
@@ -9,6 +9,8 @@
 int2 ToonProfileTextureResolution;
 uint ToonProfileIndex;
 
+float ToonProfileShadowExtinctionCoefficientInInvCentimeter;
+
 float DiffuseDitherPatternStrength;
 float DiffuseDitherPatternSize;
 
@@ -51,6 +53,10 @@ void ToonProfileUpdateAtlasTileCS(uint3 DispatchThreadId : SV_DispatchThreadID)
 			{
 				ToonProfileTexture[uint3(PixelPos, ToonProfileIndex)] = float4(DiffuseDitherPatternStrength, DiffuseDitherPatternSize, SpecularDitherPatternStrength, SpecularDitherPatternSize);
 			}
+			else if (PixelPos.x == 1)
+			{
+				ToonProfileTexture[uint3(PixelPos, ToonProfileIndex)] = float4(ToonProfileShadowExtinctionCoefficientInInvCentimeter, 0.0f, 0.0f, 0.0f);
+			}
 			else
 			{
 				ToonProfileTexture[uint3(PixelPos, ToonProfileIndex)] = DefaultColor;
diff --git a/Engine/Shaders/Private/ToonProfileCommon.ush b/Engine/Shaders/Private/ToonProfileCommon.ush
--- a/Engine/Shaders/Private/ToonProfileCommon.ush
+++ b/Engine/Shaders/Private/ToonProfileCommon.ush
@@ -8,8 +8,8 @@ float3 EvaluateToonProfileRamp(in float ToonProfileId, in float NoL, in float TO
 {
 	const float3 ToonProfileCoord = float3(
 		NoL,																	// The toon ramp is applied over NoL in [-1, 1]
-		(TOON_PROFILE_RAMP + 0.5f) * View.ToonProfileTextureSizeAndInvSize.w,		// Center of the first row of texel
-		ToonProfileId																		// Slice index
+		(TOON_PROFILE_RAMP + 0.5f) * View.ToonProfileTextureSizeAndInvSize.w,	// Center of the row of texel
+		ToonProfileId															// Slice index
 	);
 	return View.ToonProfileTexture.SampleLevel(View.ToonProfileSampler, ToonProfileCoord, 0).xyz;
 }
@@ -35,7 +35,7 @@ struct FToonProfileEvaluateSettings
 
 FToonProfileEvaluateSettings GetToonProfileEvaluateSettings(in float ToonProfileId)
 {
-	const float3 ToonProfileCoord = float3(0,
+	const float3 ToonProfileCoord = float3(0.5f * View.ToonProfileTextureSizeAndInvSize.z,
 		(TOON_PROFILE_PARAMETERS_Y + 0.5f) * View.ToonProfileTextureSizeAndInvSize.w,		// Center of the row of texel
 		ToonProfileId																		// EvaluateToonProfileDiffuseDither / Slice index
 	);
@@ -99,3 +99,26 @@ float EvaluateToonProfileSpecularDither(in float ToonProfileId, in float2 Screen
 	return 0.0f;
 }
 
+struct FToonProfileShadowSettings
+{
+	float ShadowExtinctionCoefficientInInvCentimeter;	// Needed because transmittance evaluation code from ShadowProjectionPixelShader.usf uses centimeters.
+	float ShadowMFPInMeters;							// Needed for storage in substrate Subsurface Header.
+};
+
+FToonProfileShadowSettings GetToonProfileShadowSettings(in float ToonProfileId)
+{
+	const float3 ToonProfileCoord = float3(1.5f * View.ToonProfileTextureSizeAndInvSize.z,
+		(TOON_PROFILE_PARAMETERS_Y + 0.5f) * View.ToonProfileTextureSizeAndInvSize.w,		// Center of the row of texel
+		ToonProfileId																		// Slice index
+	);
+	float4 Data = View.ToonProfileTexture.SampleLevel(View.ToonProfileSampler, ToonProfileCoord, 0).xyzw;
+
+	FToonProfileShadowSettings Ret;
+	Ret.ShadowExtinctionCoefficientInInvCentimeter = Data.x;
+
+	const float ShadowExtinctionCoefficientInInvMeter = Ret.ShadowExtinctionCoefficientInInvCentimeter * (1.0f / CENTIMETER_TO_METER); // 1/cm => 1/m
+	Ret.ShadowMFPInMeters = 1.0f / max(0.000001f, ShadowExtinctionCoefficientInInvMeter);
+
+	return Ret;
+}
+
diff --git a/Engine/Shaders/Public/SubstratePublic.ush b/Engine/Shaders/Public/SubstratePublic.ush
--- a/Engine/Shaders/Public/SubstratePublic.ush
+++ b/Engine/Shaders/Public/SubstratePublic.ush
@@ -80,7 +80,8 @@ FSubstrateGBuffer SubstratePublic_GetSubstrateGBufferFromFirstBSDF(uint2 PixelCo
 		Out.Opacity				= SubstrateGetBSDFOpacity(BSDF);
 
 		// Override opacity with SSSHeader value if needed
-		if (SubstrateSubSurfaceHeaderGetSSSType(SSSHeader) != SSS_TYPE_NONE)
+		const uint SSSType = SubstrateSubSurfaceHeaderGetSSSType(SSSHeader);
+		if (SSSType != SSS_TYPE_NONE && SSSType != SSS_TYPE_TRANSMITTANCE_SHADOW)
 		{
 			Out.Opacity = SubstrateSubSurfaceHeaderGetOpacity(SSSHeader);
 		}
diff --git a/Engine/Shaders/Shared/SubstrateDefinitions.h b/Engine/Shaders/Shared/SubstrateDefinitions.h
--- a/Engine/Shaders/Shared/SubstrateDefinitions.h
+++ b/Engine/Shaders/Shared/SubstrateDefinitions.h
@@ -150,4 +150,5 @@
 #define SSS_TYPE_DIFFUSION						(3u)
 #define SSS_TYPE_DIFFUSION_PROFILE				(4u)
 #define SSS_TYPE_SIMPLEVOLUME					(5u)
-#define SSS_TYPE_COUNT							(6u)
+#define SSS_TYPE_TRANSMITTANCE_SHADOW			(6u)
+#define SSS_TYPE_COUNT							(7u)
diff --git a/Engine/Shaders/Shared/SubstrateVisualizeDefinitions.h b/Engine/Shaders/Shared/SubstrateVisualizeDefinitions.h
--- a/Engine/Shaders/Shared/SubstrateVisualizeDefinitions.h
+++ b/Engine/Shaders/Shared/SubstrateVisualizeDefinitions.h
@@ -348,12 +348,23 @@ FSubstrateSubsurfaceHeader SubstrateLoadSubsurfaceHeader(inout FSubstrateAddress
 		SubstrateSubSurfaceHeaderSetWrapOpacity(Out, GBufferData.CustomData.a);
 
 	}
-	if (GBufferData.ShadingModelID == SHADINGMODELID_SUBSURFACE_PROFILE ||
+	else if (GBufferData.ShadingModelID == SHADINGMODELID_SUBSURFACE_PROFILE ||
 		GBufferData.ShadingModelID == SHADINGMODELID_EYE)
 	{
 		Mode = SSS_TYPE_DIFFUSION_PROFILE;
 		SubstrateSubSurfaceHeaderSetProfile(Out, GBufferData.CustomData.a, ExtractSubsurfaceProfileInt(GBufferData));
 	}
+	else if (GBufferData.ShadingModelID == SHADINGMODELID_SUBSTRATE_TOON)
+	{
+		Mode = SSS_TYPE_TRANSMITTANCE_SHADOW;
+
+		float3 SpecularColor = 0.0f;
+		float Anisotropy = 0.0f;
+		float ToonProfileId = 0.0f;
+		UnpackToonCustomData(PackRGBA8(GBufferData.CustomData), SpecularColor, Anisotropy, ToonProfileId);
+
+		SubstrateSubSurfaceHeaderSetTransmittanceShadowMFPInMeters(Out, GetToonProfileShadowSettings(ToonProfileId).ShadowMFPInMeters.xxx);
+	}
 
 	SubstrateSubSurfaceHeaderSetSSSType(Out, Mode);
 	return Out;
@@ -689,6 +700,16 @@ FSubstratePixelDebugData ConvertToSerializableSubstratePixelDebugData(uint2 InCo
 			Data.Closures[i].Roughness		= TOON_ROUGHNESS(BSDF);
 			Data.Closures[i].Anisotropy		= TOON_ANISOTROPY(BSDF);
 			Data.Closures[i].SSSPRofileID	= TOON_PROFILEID(BSDF);
+
+			if (Data.Closures[i].SSSType == SSS_TYPE_TRANSMITTANCE_SHADOW)
+			{
+			#if SUBSTRATE_MATERIALCONTAINER_IS_VIEWRESOURCE
+				const float ShadowMFPInMeters = GetToonProfileShadowSettings(TOON_PROFILEID(BSDF)).ShadowMFPInMeters;
+				Data.Closures[i].SSSMFPR = ShadowMFPInMeters;
+				Data.Closures[i].SSSMFPG = ShadowMFPInMeters;
+				Data.Closures[i].SSSMFPB = ShadowMFPInMeters;
+			#endif
+			}
 		}
 		break;
 
diff --git a/Engine/Source/Runtime/Engine/Classes/Engine/ToonProfile.h b/Engine/Source/Runtime/Engine/Classes/Engine/ToonProfile.h
--- a/Engine/Source/Runtime/Engine/Classes/Engine/ToonProfile.h
+++ b/Engine/Source/Runtime/Engine/Classes/Engine/ToonProfile.h
@@ -25,7 +25,7 @@ struct FToonProfileStruct
 	GENERATED_USTRUCT_BODY()
 
 	/**
-	* Define the diffuse ramp.
+	* Defines the diffuse ramp.
 	*/
 	UPROPERTY(Category = "Diffuse", EditAnywhere, meta = (AllowZoomOutput = "false", ShowZoomButtons = "false", ViewMinInput = "0", ViewMaxInput = "1", ViewMinOutput = "0", ViewMaxOutput = "1", TimelineLength = "1", ShowInputGridNumbers="false", ShowOutputGridNumbers="false"))
 	FRuntimeCurveLinearColor DiffuseRamp;
@@ -49,7 +49,7 @@ struct FToonProfileStruct
 	float DiffuseDitherPatternSize;
 
 	/**
-	* Define the specular ramp.
+	* Defines the specular ramp.
 	*/
 	UPROPERTY(Category = "Specular", EditAnywhere, meta = (AllowZoomOutput = "false", ShowZoomButtons = "false", ViewMinInput = "0", ViewMaxInput = "1", ViewMinOutput = "0", ViewMaxOutput = "1", TimelineLength = "1", ShowInputGridNumbers="false", ShowOutputGridNumbers="false"))
 	FRuntimeCurveLinearColor SpecularRamp;
@@ -72,12 +72,25 @@ struct FToonProfileStruct
 	UPROPERTY(Category = "Specular", EditAnywhere, meta = (UIMin = 0.0, UIMax = 1.0, ClampMin = 0.0))
 	float SpecularDitherPatternSize;
 
+	/**
+	* Defines how deep light will penetrate inside the model for the shadow map occluder distance. 
+	* This is effectively the extinction (unit=1/meter) of the virtual mesh medium. A value of 1 means that exp(-1)=36% of the energy will remain after 1 meter.
+	*/
+	UPROPERTY(Category = "Shadow", EditAnywhere, meta = (DisplayName = "Shadow Extinction", UIMin = 0.0, UIMax = 10.0, ClampMin = 0.0, ClampMax = 100.0, SliderExponent = 2.0))
+	float ShadowExtinctionCoefficient;
+
 	FToonProfileStruct();
 
 	void Invalidate()
 	{
 		*this = FToonProfileStruct();
 	}
+
+	float GetShadowExtinctionCoefficientInInvCentimeter() const
+	{
+		// Convert from 1/m to 1/cm
+		return ShadowExtinctionCoefficient * (1.0f / 100.0f);
+	}
 };
 
 /**
diff --git a/Engine/Source/Runtime/Engine/Private/Rendering/ToonProfile.cpp b/Engine/Source/Runtime/Engine/Private/Rendering/ToonProfile.cpp
--- a/Engine/Source/Runtime/Engine/Private/Rendering/ToonProfile.cpp
+++ b/Engine/Source/Runtime/Engine/Private/Rendering/ToonProfile.cpp
@@ -66,6 +66,8 @@ public:
 		SHADER_PARAMETER(float, SpecularDitherPatternStrength)
 		SHADER_PARAMETER(float, SpecularDitherPatternSize)
 
+		SHADER_PARAMETER(float, ToonProfileShadowExtinctionCoefficientInInvCentimeter)
+
 		SHADER_PARAMETER_RDG_TEXTURE_UAV(RWTexture2DArray, ToonProfileTexture)
 
 		SHADER_PARAMETER_RDG_TEXTURE(Texture2D, DiffuseDitherPattern)
@@ -353,6 +355,7 @@ IPooledRenderTarget* FToonProfileTextureManager::GetAtlasTexture(FRDGBuilder& Gr
 			PassParameters->DiffuseDitherPatternSize = Data.DiffuseDitherPatternSize;
 			PassParameters->SpecularDitherPatternStrength = bHasSpecularDitherPaternTexture ? Data.SpecularDitherPatternStrength : 0.0f;
 			PassParameters->SpecularDitherPatternSize = Data.SpecularDitherPatternSize;
+			PassParameters->ToonProfileShadowExtinctionCoefficientInInvCentimeter = Data.GetShadowExtinctionCoefficientInInvCentimeter();
 
 			auto CreateBufferFRomCurve = [&](FRuntimeCurveLinearColor ColorCurve)
 			{
@@ -466,6 +469,8 @@ FToonProfileStruct::FToonProfileStruct()
 	SpecularDitherPattern = nullptr;
 	SpecularDitherPatternStrength = 0.0f;	// By default, skip sampling dither pattern.
 	SpecularDitherPatternSize = 0.1;
+
+	ShadowExtinctionCoefficient = 10.0f;	// Default extinction coefficients that work well.
 }
 
 ///////////////////////////////////////////////////////////////////////////////////////////////////
diff --git a/Engine/Source/Runtime/Engine/Public/Materials/MaterialExpressionSubstrate.h b/Engine/Source/Runtime/Engine/Public/Materials/MaterialExpressionSubstrate.h
--- a/Engine/Source/Runtime/Engine/Public/Materials/MaterialExpressionSubstrate.h
+++ b/Engine/Source/Runtime/Engine/Public/Materials/MaterialExpressionSubstrate.h
@@ -27,6 +27,7 @@ enum EMaterialSubSurfaceType : int
 	MSS_Diffusion			= SSS_TYPE_DIFFUSION 			UMETA(DisplayName="Diffusion", ToolTip="Diffusion based sub-surface scattering"),
 	MSS_DiffusionProfile	= SSS_TYPE_DIFFUSION_PROFILE 	UMETA(Hidden),
 	MSS_SimpleVolume		= SSS_TYPE_SIMPLEVOLUME 		UMETA(DisplayName="Simple Volume", ToolTip="Approximation of optically thin slab (e.g.: glass) where light is visible through the material"),
+	MSS_TransmittanceShadow	= SSS_TYPE_TRANSMITTANCE_SHADOW	UMETA(Hidden),	// Only used from shaders as of today and might be refactor so not added here explicitely in UI.
 	MSS_MAX
 };
 
diff --git a/Engine/Source/Runtime/Renderer/Private/Substrate/SubstrateVisualize.cpp b/Engine/Source/Runtime/Renderer/Private/Substrate/SubstrateVisualize.cpp
--- a/Engine/Source/Runtime/Renderer/Private/Substrate/SubstrateVisualize.cpp
+++ b/Engine/Source/Runtime/Renderer/Private/Substrate/SubstrateVisualize.cpp
@@ -887,6 +887,16 @@ void AddProcessAndPrintSubstrateMaterialPropertiesPasses(FRDGBuilder& GraphBuild
 						PrintS(DrawPosX + 0.0f, DrawPosY, TEXT("Anisotropy"), Orange);
 						PrintFSmall(DrawPosX + DataOffset, DrawPosY, Closure.Anisotropy);
 						NewLine();
+
+						if (Closure.SSSType == SSS_TYPE_TRANSMITTANCE_SHADOW)
+						{
+							PrintS(DrawPosX + 0.0f, DrawPosY, TEXT("Transmittance Shadow"), Orange);
+							NewLine();
+				
+							PrintS(DrawPosX + 0.0f, DrawPosY, TEXT(" - MFP"), Orange);
+							PrintFColorAdapt(DrawPosX + DataOffset, DrawPosY, Closure.SSSMFPR, Closure.SSSMFPG, Closure.SSSMFPB);
+							NewLine();
+						}
 					}
 
 					DispatchBatch();
