Substrate - rough translucent refraction now account for second roughness when used.

#tests PC edito dx12
#rnx

#Synced-CL 49930030

diff --git a/Engine/Shaders/Private/Substrate/SubstrateTree.ush b/Engine/Shaders/Private/Substrate/SubstrateTree.ush
--- a/Engine/Shaders/Private/Substrate/SubstrateTree.ush
+++ b/Engine/Shaders/Private/Substrate/SubstrateTree.ush
@@ -49,9 +49,26 @@ void UpdateSingleBSDFOperatorCoverageTransmittance(
 #endif
 
 	float CurrentBSDFCoverage = CurrentBSDF.Coverage;
-	float Roughness = SubstrateGetBSDFRoughness(CurrentBSDF);
 	float3x3 TangentBasis = SubstrateGetBSDFSharedBasis_InlineShading(SubstratePixelHeader, CurrentBSDF);
 
+	float Roughness = SubstrateGetBSDFRoughness(CurrentBSDF);
+#if SUBSTRATE_FASTPATH==0
+	// Combine second roughness to allso affect opqaue and translucent rough refraction
+	const bool bHasHaziness = BSDF_GETHASHAZINESS(CurrentBSDF);
+	BRANCH
+	if (bHasHaziness)
+	{
+		FHaziness Haziness = UnpackHaziness(SLAB_HAZINESS(CurrentBSDF));
+		const bool  bHazeAsSimpleClearCoat = Haziness.bSimpleClearCoat;
+		if (!bHazeAsSimpleClearCoat)
+		{
+			const float HazeWeight = Haziness.Weight;
+			const float HazeSafeRoughness = MakeRoughnessSafe(Haziness.Roughness);
+			Roughness = lerp(Roughness, HazeSafeRoughness, HazeWeight);
+		}
+	}
+#endif
+
 #if USE_DEVELOPMENT_SHADERS
 	Roughness = Settings.bRoughnessTracking ? Roughness : 0.0;
 #endif
