Add ubershader permutation to VirtualShadowMap projection shader.

The uber shader permutation is used when the requested permutation is no available and requires PSO compilation. While the PSO is asyncly compiled, the uber permutation is used.

#Synced-CL 49992015

diff --git a/Engine/Shaders/Private/VirtualShadowMaps/VirtualShadowMapProjection.usf b/Engine/Shaders/Private/VirtualShadowMaps/VirtualShadowMapProjection.usf
--- a/Engine/Shaders/Private/VirtualShadowMaps/VirtualShadowMapProjection.usf
+++ b/Engine/Shaders/Private/VirtualShadowMaps/VirtualShadowMapProjection.usf
@@ -11,6 +11,10 @@
 #include "../HairStrands/HairStrandsVisibilityCommon.ush"
 #include "../Visualization.ush"
 #include "../BlueNoise.ush"
+
+int SMRTSamplePerRay;
+uint bSMRTExtrapolateSlope;
+
 #include "VirtualShadowMapPageAccessCommon.ush"
 #include "VirtualShadowMapProjectionCommon.ush"
 #include "VirtualShadowMapProjectionFilter.ush"
@@ -237,13 +241,18 @@ FVirtualShadowMapSampleResult SampleVirtualShadowMapDirectionalPrefilteredDistan
 
 #endif
 
+uint bPrefilteredDistantClipmap;
+uint bFirstPersonShadow;
+uint bOnePassProjection;
+uint bVisualizeOutput;
 int4 ProjectionRect;
 float SubsurfaceMinSourceRadius;
 
 // One pass projection
-#if ONE_PASS_PROJECTION
+#if ONE_PASS_PROJECTION || UBERSHADER
 RWTexture2D< uint4 > OutShadowMaskBits;
-#else
+#endif
+#if !ONE_PASS_PROJECTION || UBERSHADER
 // Single light per pass
 // Light parameters loaded via GetRootLightShaderParameters();
 int LightUniformVirtualShadowMapId;
@@ -358,14 +367,16 @@ FVirtualShadowMapSampleResult ProjectLight(
 
 #if DIRECTIONAL_LIGHT
 
-#if PREFILTERED_DISTANT_CLIPMAP
-		// Sample the perfiltered distant SM first and early out if that is already producing fully shadowed samples.
-		// TODO: reuse the mapped clipmap search (create some intermediate struct that can be passed between them?)
-		Result = SampleVirtualShadowMapDirectionalPrefilteredDistant(VirtualShadowMapId, Light, PixelPos, TranslatedWorldPosition);
-		bool bFullyShadowed = Result.bValid && Result.ShadowFactor <= 0.0001f;
-#else
 		bool bFullyShadowed = false;
-#endif // PREFILTERED_DISTANT_CLIPMAP
+		#if PREFILTERED_DISTANT_CLIPMAP
+		IF_UBERSHADER(bPrefilteredDistantClipmap)
+		{
+			// Sample the perfiltered distant SM first and early out if that is already producing fully shadowed samples.
+			// TODO: reuse the mapped clipmap search (create some intermediate struct that can be passed between them?)
+			Result = SampleVirtualShadowMapDirectionalPrefilteredDistant(VirtualShadowMapId, Light, PixelPos, TranslatedWorldPosition);
+			bFullyShadowed = Result.bValid && Result.ShadowFactor <= 0.0001f;
+		}
+		#endif // PREFILTERED_DISTANT_CLIPMAP
 
 		if (!bFullyShadowed)
 		{
@@ -415,52 +426,56 @@ FVirtualShadowMapSampleResult ProjectLight(
 				Result = SmrtResult;
 			}
 		}
-#if FIRST_PERSON_SHADOW
-		// If there is an extra FP shadow map, it represents the shadow cast from the 3rd person mesh of the player onto the world only, thus it should only be sampled if
-		// 1. the pixel is not marked as FP (e.g., the weapon)
-		// 
-		// 2. the pixel is not already in shadow (100%)
-		bFullyShadowed = Result.bValid && Result.ShadowFactor == 0.0f;
 
-		if (FirstPersonVirtualShadowMapId >= 0 && !ShadingInfo.bIsFirstPerson && !bFullyShadowed)
+		#if FIRST_PERSON_SHADOW
+		IF_UBERSHADER(bFirstPersonShadow)
 		{
-			FVirtualShadowMapSampleResult FPResult = InitVirtualShadowMapSampleResult();
-			if (GetSMRTRayCount() > 0)
-			{
-				FPResult = TraceDirectional(
-					FirstPersonVirtualShadowMapId,
-					Light,
-					PixelPos,
-					SceneDepth,
-					TranslatedWorldPosition,
-					SMRTRayOffset,
-					Noise,
-					WorldNormal);
-			}
-			else
-			{
-				FPResult = SampleVirtualShadowMapDirectional(
-					FirstPersonVirtualShadowMapId,
-					TranslatedWorldPosition,
-					SMRTRayOffset,
-					WorldNormal,
-					SceneDepth);
-			}
-			if (FPResult.bValid)
+			// If there is an extra FP shadow map, it represents the shadow cast from the 3rd person mesh of the player onto the world only, thus it should only be sampled if
+			// 1. the pixel is not marked as FP (e.g., the weapon)
+			// 
+			// 2. the pixel is not already in shadow (100%)
+			bFullyShadowed = Result.bValid && Result.ShadowFactor == 0.0f;
+
+			if (FirstPersonVirtualShadowMapId >= 0 && !ShadingInfo.bIsFirstPerson && !bFullyShadowed)
 			{
-				if (!Result.bValid)
+				FVirtualShadowMapSampleResult FPResult = InitVirtualShadowMapSampleResult();
+				if (GetSMRTRayCount() > 0)
 				{
-					Result = FPResult;
+					FPResult = TraceDirectional(
+						FirstPersonVirtualShadowMapId,
+						Light,
+						PixelPos,
+						SceneDepth,
+						TranslatedWorldPosition,
+						SMRTRayOffset,
+						Noise,
+						WorldNormal);
 				}
 				else
 				{
-					// Merge the result, picking the most-shadowed value & nearest occluder distance
-					Result.ShadowFactor = min(Result.ShadowFactor, FPResult.ShadowFactor);
-					Result.OccluderDistance = min(Result.OccluderDistance, FPResult.OccluderDistance);
+					FPResult = SampleVirtualShadowMapDirectional(
+						FirstPersonVirtualShadowMapId,
+						TranslatedWorldPosition,
+						SMRTRayOffset,
+						WorldNormal,
+						SceneDepth);
+				}
+				if (FPResult.bValid)
+				{
+					if (!Result.bValid)
+					{
+						Result = FPResult;
+					}
+					else
+					{
+						// Merge the result, picking the most-shadowed value & nearest occluder distance
+						Result.ShadowFactor = min(Result.ShadowFactor, FPResult.ShadowFactor);
+						Result.OccluderDistance = min(Result.OccluderDistance, FPResult.OccluderDistance);
+					}
 				}
 			}
 		}
-#endif // FIRST_PERSON_SHADOW
+		#endif // FIRST_PERSON_SHADOW
 
 
 #else // !DIRECTIONAL_LIGHT
@@ -903,20 +918,30 @@ void OutputVisualize(
 }
 
 #if USE_TILE_LIST
+uint bUseTileList;
 Buffer<uint> TileListData;
 #endif
 
+uint2 GetTileCoord(uint2 GroupId)
+{
+	uint2 Out = GroupId.xy;
+	#if USE_TILE_LIST
+	IF_UBERSHADER(bUseTileList)
+	{
+		Out = UnpackTileCoord12bits(TileListData.Load(GroupId.x)); // Need to match SingleLayerWater tile encoding
+	}
+	#endif // USE_TILE_LIST
+	return Out;
+}
+
 [numthreads(WORK_TILE_SIZE, WORK_TILE_SIZE, 1)]
 void VirtualShadowMapProjection(
 	uint3	GroupId				: SV_GroupID,
 	uint	GroupIndex			: SV_GroupIndex,
 	uint3	DispatchThreadId	: SV_DispatchThreadID )
 {
-#if USE_TILE_LIST
-	const uint2 TileCoord = UnpackTileCoord12bits(TileListData.Load(GroupId.x)); // Need to match SingleLayerWater tile encoding
-#else
-	const uint2 TileCoord = GroupId.xy;
-#endif // USE_TILE_LIST
+	const uint2 TileCoord = GetTileCoord(GroupId.xy);
+
 	// Morton order within a group so page access/atomics are more coherent and wave-swizzled gradients are possible.
 	uint2 LocalPixelPos = WORK_TILE_SIZE * TileCoord + MortonDecode(GroupIndex);
 	uint2 PixelPos = LocalPixelPos + uint2( ProjectionRect.xy );
@@ -953,57 +978,65 @@ void VirtualShadowMapProjection(
 		SubsurfaceOpacityMFP = GetSubsurfaceOpacityFromGbuffer(PixelPos);
 	}
 
-#if ONE_PASS_PROJECTION
-	const FCulledLightsGridHeader LightGridHeader = VirtualShadowMapGetLightsGridHeader(LocalPixelPos, SceneDepth);	
-	// Allow one extra loop iteration to detect potential overflow
-	const uint LoopCount = min(GetPackedShadowMaskMaxLightCount() + 1, LightGridHeader.NumLights);
-
-	uint4 ShadowMask = InitializePackedShadowMask();
-	LOOP
-	for (uint Index = 0; Index < LoopCount; ++Index)
+#if ONE_PASS_PROJECTION || UBERSHADER
+	IF_UBERSHADER(bOnePassProjection)
 	{
-		const FLocalLightData LightData = VirtualShadowMapGetLocalLightData(LightGridHeader, Index);
-		
-		// Megalights are sorted to the end (See LightSceneInfo.h SortyKey) so if we hit one we're done
-		if (UnpackIsHandledByMegaLights(LightData))
-		{
-			break;
-		}
+		const FCulledLightsGridHeader LightGridHeader = VirtualShadowMapGetLightsGridHeader(LocalPixelPos, SceneDepth);	
+		// Allow one extra loop iteration to detect potential overflow
+		const uint LoopCount = min(GetPackedShadowMaskMaxLightCount() + 1, LightGridHeader.NumLights);
 
-		if (Index < GetPackedShadowMaskMaxLightCount())
+		uint4 ShadowMask = InitializePackedShadowMask();
+		LOOP
+		for (uint Index = 0; Index < LoopCount; ++Index)
 		{
-			// VSM pruned light grid, so this should always be valid
-			int VirtualShadowMapId = LightData.Internal.VirtualShadowMapId;
+			const FLocalLightData LightData = VirtualShadowMapGetLocalLightData(LightGridHeader, Index);
+		
+			// Megalights are sorted to the end (See LightSceneInfo.h SortyKey) so if we hit one we're done
+			if (UnpackIsHandledByMegaLights(LightData))
+			{
+				break;
+			}
 
-			FLightShaderParameters Light = ConvertFromLocal(LightData);
-			FVirtualShadowMapSampleResult Result = ProjectLight(
-				VirtualShadowMapId,
-				Light,
-				ShadingInfo,
-				PixelPos,
-				SceneDepth,
-				ScreenRayLengthWorld,
-				TranslatedWorldPosition,
-				Noise,
-				SubsurfaceOpacityMFP);
+			if (Index < GetPackedShadowMaskMaxLightCount())
+			{
+				// VSM pruned light grid, so this should always be valid
+				int VirtualShadowMapId = LightData.Internal.VirtualShadowMapId;
+
+				FLightShaderParameters Light = ConvertFromLocal(LightData);
+				FVirtualShadowMapSampleResult Result = ProjectLight(
+					VirtualShadowMapId,
+					Light,
+					ShadingInfo,
+					PixelPos,
+					SceneDepth,
+					ScreenRayLengthWorld,
+					TranslatedWorldPosition,
+					Noise,
+					SubsurfaceOpacityMFP);
 
-			FilterVirtualShadowMapSampleResult(PixelPos, Result);
-			PackShadowMask(ShadowMask, Result.ShadowFactor, Index);
+				FilterVirtualShadowMapSampleResult(PixelPos, Result);
+				PackShadowMask(ShadowMask, Result.ShadowFactor, Index);
 
-#if VISUALIZE_OUTPUT
-			OutputVisualize( VirtualShadowMapId, PixelPos - View.ViewRectMin.xy, TranslatedWorldPosition, Result );
-#endif // VISUALIZE_OUTPUT
-		}
-		else
-		{
-			// Record overflow
-			StatsBufferInterlockedEnableFlags(VSM_STAT_OVERFLOW_FLAGS, VSM_STAT_OVERFLOW_FLAG_OPP_MAX_LIGHTS, true, true);
+				#if VISUALIZE_OUTPUT
+				if (bVisualizeOutput)
+				{
+					OutputVisualize( VirtualShadowMapId, PixelPos - View.ViewRectMin.xy, TranslatedWorldPosition, Result );
+				}
+				#endif // VISUALIZE_OUTPUT
+			}
+			else
+			{
+				// Record overflow
+				StatsBufferInterlockedEnableFlags(VSM_STAT_OVERFLOW_FLAGS, VSM_STAT_OVERFLOW_FLAG_OPP_MAX_LIGHTS, true, true);
+			}
 		}
-	}
 
-	OutShadowMaskBits[PixelPos] = ShadowMask;
+		OutShadowMaskBits[PixelPos] = ShadowMask;
+	}
+#endif // ONE_PASS_PROJECTION || UBERSHADER
 
-#else // !ONE_PASS_PROJECTION
+#if !ONE_PASS_PROJECTION || UBERSHADER
+	IF_UBERSHADER(!bOnePassProjection)
 	{
 		int VirtualShadowMapId = LightUniformVirtualShadowMapId;
 		FLightShaderParameters Light = GetRootLightShaderParameters();
@@ -1032,10 +1065,13 @@ void VirtualShadowMapProjection(
 
 		OutShadowFactor[ PixelPos ] = Result.ShadowFactor.xx;
 
-#if VISUALIZE_OUTPUT
-		OutputVisualize( VirtualShadowMapId, PixelPos - View.ViewRectMin.xy, TranslatedWorldPosition, Result );
-#endif // VISUALIZE_OUTPUT
+		#if VISUALIZE_OUTPUT
+		IF_UBERSHADER(bVisualizeOutput)
+		{
+			OutputVisualize(VirtualShadowMapId, PixelPos - View.ViewRectMin.xy, TranslatedWorldPosition, Result);
+		}
+		#endif // VISUALIZE_OUTPUT
 
 	}
-#endif // ONE_PASS_PROJECTION
+#endif // !ONE_PASS_PROJECTION || UBERSHADER
 }
diff --git a/Engine/Shaders/Private/VirtualShadowMaps/VirtualShadowMapSMRTTemplate.ush b/Engine/Shaders/Private/VirtualShadowMaps/VirtualShadowMapSMRTTemplate.ush
--- a/Engine/Shaders/Private/VirtualShadowMaps/VirtualShadowMapSMRTTemplate.ush
+++ b/Engine/Shaders/Private/VirtualShadowMaps/VirtualShadowMapSMRTTemplate.ush
@@ -46,7 +46,11 @@ FSMRTResult SMRTRayCast(
 	float PrevReferenceDepth = -1;
 
 	bool bValidHit =  false;
-#if SMRT_TEMPLATE_STATIC_SAMPLES_PER_RAY >= 0
+	float SlopeScale = 1;
+#if UBERSHADER
+	NumSteps = SMRTSamplePerRay;
+	SlopeScale = bSMRTExtrapolateSlope ? 1 : 0;
+#elif SMRT_TEMPLATE_STATIC_SAMPLES_PER_RAY >= 0
 	NumSteps = SMRT_TEMPLATE_STATIC_SAMPLES_PER_RAY;
 	UNROLL
 #endif
@@ -95,7 +99,7 @@ FSMRTResult SMRTRayCast(
 				if (bBehind)
 				{
 					#if SMRT_EXTRAPOLATE_SLOPE
-						DepthForComparison = DepthSlope * DeltaHistoryTime + DepthHistory;
+						DepthForComparison = SlopeScale * DepthSlope * DeltaHistoryTime + DepthHistory;
 					#else
 						DepthForComparison = DepthHistory;
 					#endif
diff --git a/Engine/Source/Runtime/Renderer/Private/VirtualShadowMaps/VirtualShadowMapProjection.cpp b/Engine/Source/Runtime/Renderer/Private/VirtualShadowMaps/VirtualShadowMapProjection.cpp
--- a/Engine/Source/Runtime/Renderer/Private/VirtualShadowMaps/VirtualShadowMapProjection.cpp
+++ b/Engine/Source/Runtime/Renderer/Private/VirtualShadowMaps/VirtualShadowMapProjection.cpp
@@ -110,6 +110,7 @@ class FVirtualShadowMapProjectionCS : public FGlobalShader
 	class FUseTileList					: SHADER_PERMUTATION_BOOL("USE_TILE_LIST");
 	class FFirstPersonShadow			: SHADER_PERMUTATION_BOOL("FIRST_PERSON_SHADOW");
 	class FPrefilteredDistantClipmap	: SHADER_PERMUTATION_BOOL("PREFILTERED_DISTANT_CLIPMAP");
+	class FUbershader					: SHADER_PERMUTATION_BOOL("UBERSHADER");
 	
 	// -1 means dynamic count
 	class FSMRTStaticSampleCount	: SHADER_PERMUTATION_SPARSE_INT("SMRT_TEMPLATE_STATIC_SAMPLES_PER_RAY", -1, 2, 4);
@@ -125,7 +126,8 @@ class FVirtualShadowMapProjectionCS : public FGlobalShader
 		FUseTileList,
 		FFirstPersonShadow,
 		FPrefilteredDistantClipmap,
-		FSMRTStaticSampleCount
+		FSMRTStaticSampleCount,
+		FUbershader
 #if MAX_TEST_PERMUTATION > 0
 		, FTestDim
 #endif
@@ -167,6 +169,13 @@ class FVirtualShadowMapProjectionCS : public FGlobalShader
 		// DoF
 		SHADER_PARAMETER(float, DOFBiasStrength)
 		SHADER_PARAMETER_STRUCT_INCLUDE(DiaphragmDOF::FDOFCocModelShaderParameters, CocModel)
+		SHADER_PARAMETER(int32, SMRTSamplePerRay)
+		SHADER_PARAMETER(uint32, bOnePassProjection)
+		SHADER_PARAMETER(uint32, bVisualizeOutput)
+		SHADER_PARAMETER(uint32, bSMRTExtrapolateSlope)
+		SHADER_PARAMETER(uint32, bUseTileList)
+		SHADER_PARAMETER(uint32, bFirstPersonShadow)
+		SHADER_PARAMETER(uint32, bPrefilteredDistantClipmap)
 	END_SHADER_PARAMETER_STRUCT()
 
 	static void ModifyCompilationEnvironment( const FGlobalShaderPermutationParameters& Parameters, FShaderCompilerEnvironment& OutEnvironment )
@@ -195,6 +204,31 @@ class FVirtualShadowMapProjectionCS : public FGlobalShader
 	{
 		FPermutationDomain PermutationVector(Parameters.PermutationId);
 
+		if (PermutationVector.Get<FUbershader>())
+		{
+			// Actual permutation 
+			// - FDirectionalLightDim
+			if (PermutationVector.Get<FOnePassProjectionDim>() &&
+				PermutationVector.Get<FHairStrandsDim>() &&
+				PermutationVector.Get<FVisualizeOutputDim>() &&
+				PermutationVector.Get<FExtrapolateSlopeDim>() &&
+				PermutationVector.Get<FUseTileList>() &&
+				PermutationVector.Get<FFirstPersonShadow>() &&
+				PermutationVector.Get<FPrefilteredDistantClipmap>() &&
+				PermutationVector.Get<FSMRTStaticSampleCount>() == 4
+				#if MAX_TEST_PERMUTATION > 0
+				&& !PermutationVector.Get<FTestDim>(false);
+				#endif
+			)
+			{
+				return DoesPlatformSupportVirtualShadowMaps(Parameters.Platform);
+			}
+			else
+			{
+				return false;
+			}
+		}
+
 		// Directional lights are always in separate passes as forward light data structure currently
 		// only contains a single directional light.
 		if( PermutationVector.Get< FDirectionalLightDim >() && PermutationVector.Get< FOnePassProjectionDim >() )
@@ -223,6 +257,27 @@ class FVirtualShadowMapProjectionCS : public FGlobalShader
 		return DoesPlatformSupportVirtualShadowMaps(Parameters.Platform);
 	}
 
+	static FPermutationDomain RemapPermutation(FPermutationDomain PermutationVector)
+	{		
+		if (PermutationVector.Get< FUbershader>())
+		{
+			// Actual permutation 
+			// - FDirectionalLightDim
+			PermutationVector.Set<FOnePassProjectionDim>(true);
+			PermutationVector.Set<FHairStrandsDim>(true);
+			PermutationVector.Set<FVisualizeOutputDim>(true);
+			PermutationVector.Set<FExtrapolateSlopeDim>(true);
+			PermutationVector.Set<FUseTileList>(true);
+			PermutationVector.Set<FFirstPersonShadow>(true);
+			PermutationVector.Set<FPrefilteredDistantClipmap>(true);
+			PermutationVector.Set<FSMRTStaticSampleCount>(4);
+			#if MAX_TEST_PERMUTATION > 0
+			PermutationVector.Set<FTestDim>(false);
+			#endif
+		}
+		return PermutationVector;
+	}
+
 	static EShaderPermutationPrecacheRequest ShouldPrecachePermutation(const FGlobalShaderPermutationParameters& Parameters)
 	{
 		FPermutationDomain PermutationVector(Parameters.PermutationId);
@@ -231,6 +286,11 @@ class FVirtualShadowMapProjectionCS : public FGlobalShader
 			return EShaderPermutationPrecacheRequest::NotPrecached;
 		}
 
+		if (!PermutationVector.Get<FUbershader>())
+		{
+			return EShaderPermutationPrecacheRequest::NotPrecached;
+		}
+
 		return EShaderPermutationPrecacheRequest::Precached;
 	}
 };
@@ -285,13 +345,13 @@ static void RenderVirtualShadowMapProjectionCommon(
 
 	bool bDirectionalLight = false;
 	bool bOnePassProjection = LightProxy == nullptr;
+	PassParameters->ForwardLightStruct = View.ForwardLightingResources.ForwardLightUniformBuffer;
 	if (bOnePassProjection)
 	{
 		// These are mutually exclusive in practice, enforce to cut down on permutations
 		check(!bUseTileList);
 
 		// One pass projection
-		PassParameters->ForwardLightStruct = View.ForwardLightingResources.ForwardLightUniformBuffer;
 		PassParameters->OutShadowMaskBits = GraphBuilder.CreateUAV( OutputTexture );
 	}
 	else
@@ -333,16 +393,17 @@ static void RenderVirtualShadowMapProjectionCommon(
 
 	// If the requested samples per ray matches one of our static permutations, pick that one
 	// Otherwise use the dynamic samples per ray permutation (-1).
-	int SamplesPerRay = (bDirectionalLight ? VirtualShadowMapArray.UniformParameters.SMRTSamplesPerRayDirectional : VirtualShadowMapArray.UniformParameters.SMRTSamplesPerRayLocal);
-	int StaticSamplesPerRay = UE::ShaderPermutationUtils::DoesDimensionContainValue<FVirtualShadowMapProjectionCS::FSMRTStaticSampleCount>(SamplesPerRay) ? SamplesPerRay : -1;
-	float ExtrapolateMaxSlope = (bDirectionalLight ? VirtualShadowMapArray.UniformParameters.SMRTExtrapolateMaxSlopeDirectional : VirtualShadowMapArray.UniformParameters.SMRTExtrapolateMaxSlopeLocal);
+	const int SamplesPerRay = (bDirectionalLight ? VirtualShadowMapArray.UniformParameters.SMRTSamplesPerRayDirectional : VirtualShadowMapArray.UniformParameters.SMRTSamplesPerRayLocal);
+	const int StaticSamplesPerRay = UE::ShaderPermutationUtils::DoesDimensionContainValue<FVirtualShadowMapProjectionCS::FSMRTStaticSampleCount>(SamplesPerRay) ? SamplesPerRay : -1;
+	const float ExtrapolateMaxSlope = (bDirectionalLight ? VirtualShadowMapArray.UniformParameters.SMRTExtrapolateMaxSlopeDirectional : VirtualShadowMapArray.UniformParameters.SMRTExtrapolateMaxSlopeLocal);
+	const bool bSMRTExtrapolateSlope = ExtrapolateMaxSlope > 0.0f;
 
 	FVirtualShadowMapProjectionCS::FPermutationDomain PermutationVector;
 	PermutationVector.Set< FVirtualShadowMapProjectionCS::FDirectionalLightDim >( bDirectionalLight );
 	PermutationVector.Set< FVirtualShadowMapProjectionCS::FOnePassProjectionDim >( bOnePassProjection );
 	PermutationVector.Set< FVirtualShadowMapProjectionCS::FHairStrandsDim >( bHasHairStrandsData );
 	PermutationVector.Set< FVirtualShadowMapProjectionCS::FVisualizeOutputDim >( bDebugOutput );
-	PermutationVector.Set< FVirtualShadowMapProjectionCS::FExtrapolateSlopeDim >(ExtrapolateMaxSlope > 0.0f);
+	PermutationVector.Set< FVirtualShadowMapProjectionCS::FExtrapolateSlopeDim >(bSMRTExtrapolateSlope);
 	PermutationVector.Set< FVirtualShadowMapProjectionCS::FUseTileList >(bUseTileList);
 	PermutationVector.Set< FVirtualShadowMapProjectionCS::FSMRTStaticSampleCount >(StaticSamplesPerRay);
 	
@@ -367,7 +428,41 @@ static void RenderVirtualShadowMapProjectionCommon(
 	}
 #endif
 
+	PassParameters->SMRTSamplePerRay = StaticSamplesPerRay;
+	PassParameters->bOnePassProjection = bOnePassProjection;
+	PassParameters->bVisualizeOutput = bDebugOutput;
+	PassParameters->bSMRTExtrapolateSlope = bSMRTExtrapolateSlope;
+	PassParameters->bUseTileList = bUseTileList;
+	PassParameters->bFirstPersonShadow = bHasFirstPersonShadow;
+	PassParameters->bPrefilteredDistantClipmap = bPrefilteredDistantClipmap;
+	
+	PermutationVector = FVirtualShadowMapProjectionCS::RemapPermutation(PermutationVector);
 	auto ComputeShader = View.ShaderMap->GetShader< FVirtualShadowMapProjectionCS >( PermutationVector );
+
+	// Ubershader default parameters and feature dynamic branching
+	if (!PrecachePSO(FComputePipelineStateInitializer(ComputeShader.GetComputeShader())))
+	{
+		PermutationVector.Set< FVirtualShadowMapProjectionCS::FUbershader >(true);
+		PermutationVector = FVirtualShadowMapProjectionCS::RemapPermutation(PermutationVector);
+		ComputeShader = View.ShaderMap->GetShader< FVirtualShadowMapProjectionCS >( PermutationVector );
+
+		// Ubershader default parameters and feature dynamic branching
+		GSystemTextures.SetDefaultIfNull(GraphBuilder, PassParameters->PrefilteredDistantPhysicalPagePool);
+		GSystemTextures.SetDefaultIfNull(GraphBuilder, PassParameters->OutShadowMaskBits);
+		GSystemTextures.SetDefaultIfNull(GraphBuilder, PassParameters->OutShadowFactor);
+		GSystemTextures.SetDefaultIfNull(GraphBuilder, PassParameters->TileListData);
+		GSystemTextures.SetDefaultIfNull(GraphBuilder, PassParameters->IndirectDispatchArgs);
+		GSystemTextures.SetDefaultIfNull(GraphBuilder, PassParameters->OutStatsBuffer);
+		#if VSM_ENABLE_VISUALIZATION
+		GSystemTextures.SetDefaultIfNull(GraphBuilder, PassParameters->OutVisualize);
+		#endif
+
+		if (!bHasHairStrandsData)
+		{
+			PassParameters->HairStrands = HairStrands::CreateDefaultHairStrandsViewUniformBuffer(GraphBuilder);
+			PassParameters->HairStrandsVoxel = HairStrands::CreateDefaultHairStrandsVoxelUniformBuffer(GraphBuilder);
+		}
+	}
 	ClearUnusedGraphResources( ComputeShader, PassParameters );
 	ValidateShaderParameters( ComputeShader, *PassParameters );
 
