Fixed double evaluation of animation data. Added shared memory for concatenation pass which improves performance by ~10-15%.

#Synced-CL 49999797

diff --git a/Engine/Shaders/Private/Skinning/AnimSequenceTransformProvider.usf b/Engine/Shaders/Private/Skinning/AnimSequenceTransformProvider.usf
--- a/Engine/Shaders/Private/Skinning/AnimSequenceTransformProvider.usf
+++ b/Engine/Shaders/Private/Skinning/AnimSequenceTransformProvider.usf
@@ -201,16 +201,40 @@ ByteAddressBuffer ConcatenateBlockHeaderBuffer;
 ByteAddressBuffer LeafBonePathBuffer;
 RWByteAddressBuffer GlobalSpaceTransformBufferUAV;
 
-groupshared FConcatenateBlockHeader SharedConcatenateBlockHeader;
+groupshared FBoneTransformWithScale SharedTransform[BONES_PER_GROUP];
+
+FBoneTransformWithScale LoadSharedTransform(uint BoneIndex, uint TransformBufferOffset)
+{
+	if (BoneIndex < BONES_PER_GROUP)
+	{
+		return SharedTransform[BoneIndex];
+	}
+	else
+	{
+		return LoadBoneTransformWithScale(LocalSpaceTransformBuffer, TransformBufferOffset, BoneIndex);
+	}
+}
+
+void StoreSharedTransform(uint BoneIndex, FBoneTransformWithScale Transform)
+{
+	SharedTransform[BoneIndex] = Transform;
+}
 
 [numthreads(BONES_PER_GROUP, 1, 1)]
 void ConcatenateLocalToGlobalCS(uint3 GroupId : SV_GroupID, uint GroupIndex : SV_GroupIndex)
 {
 	const uint BlockIndex         = GroupId.x;
 	const uint LocalLeafBoneIndex = GroupIndex;
+	const uint PrefetchBoneIndex  = GroupIndex;
 	const uint PathElementSize    = 4;
 
-	FConcatenateBlockHeader BlockHeader = ConcatenateBlockHeaderBuffer.Load<FConcatenateBlockHeader>(GroupId.x * uint(sizeof(FConcatenateBlockHeader)));
+	FConcatenateBlockHeader BlockHeader = ConcatenateBlockHeaderBuffer.Load<FConcatenateBlockHeader>(BlockIndex * uint(sizeof(FConcatenateBlockHeader)));
+
+	if (PrefetchBoneIndex < BlockHeader.NumPrefetchBones)
+	{
+		StoreSharedTransform(PrefetchBoneIndex, LoadBoneTransformWithScale(LocalSpaceTransformBuffer, BlockHeader.TransformBufferOffset, PrefetchBoneIndex));
+	}
+	GroupMemoryBarrierWithGroupSync();
 
 	if (LocalLeafBoneIndex >= BlockHeader.BlockLeafBoneCount)
 	{
@@ -224,7 +248,7 @@ void ConcatenateLocalToGlobalCS(uint3 GroupId : SV_GroupID, uint GroupIndex : SV
 	const uint PathHeaderOffset = BlockHeader.LeafBonePathBufferOffset + RelativePathHeaderOffset;
 
 	FLeafBonePathHeader PathHeader = LeafBonePathBuffer.Load<FLeafBonePathHeader>(PathHeaderOffset);
-	FBoneTransformWithScale ParentTransform = LoadBoneTransformWithScale(LocalSpaceTransformBuffer, BlockHeader.TransformBufferOffset, 0);
+	FBoneTransformWithScale ParentTransform = LoadSharedTransform(0, BlockHeader.TransformBufferOffset);
 
 	if (PathHeader.bIsRootOwned)
 	{
@@ -241,7 +265,7 @@ void ConcatenateLocalToGlobalCS(uint3 GroupId : SV_GroupID, uint GroupIndex : SV
 			PathNode = LeafBonePathBuffer.Load<FLeafBonePathNode>(PathNodeOffset);
 			PathNodeOffset += PathElementSize;
 
-			FBoneTransformWithScale ChildTransform = LoadBoneTransformWithScale(LocalSpaceTransformBuffer, BlockHeader.TransformBufferOffset, PathNode.BoneIndex);
+			FBoneTransformWithScale ChildTransform = LoadSharedTransform(PathNode.BoneIndex, BlockHeader.TransformBufferOffset);
 			ParentTransform = MultiplyBoneTransformWithScale(ChildTransform, ParentTransform);
 
 			if (PathNode.bIsOwned)
diff --git a/Engine/Source/Runtime/Renderer/Private/Skinning/AnimSequenceTransformProvider.cpp b/Engine/Source/Runtime/Renderer/Private/Skinning/AnimSequenceTransformProvider.cpp
--- a/Engine/Source/Runtime/Renderer/Private/Skinning/AnimSequenceTransformProvider.cpp
+++ b/Engine/Source/Runtime/Renderer/Private/Skinning/AnimSequenceTransformProvider.cpp
@@ -847,10 +847,15 @@ void FAnimSequenceTransformProvider::ProvideGPUTransforms(FSkinningTransformProv
 			FTrackHistory::FUpdateResult UpdateResult = ProxyUserData.TrackHistories[TrackIndex].Update(Data, CurrentTransformSlot, bForceUpdateAll);
 
 			FTrackUpdateRequest UpdateRequest;
+			UpdateRequest.SetUpdateMode(EPreviousBoneTransformUpdateMode::None);
 			UpdateRequest.bNeedsCurrent = UpdateResult.bNeedsCurrent;
-			UpdateRequest.UpdateMode = uint8(UpdateResult.bNeedsCurrent && UpdateResult.bNeedsPrevious && TargetPosition == TargetPreviousPosition && SourcePosition == SourcePreviousPosition
-				? EPreviousBoneTransformUpdateMode::DuplicateCurrentToPrevious
-				: EPreviousBoneTransformUpdateMode::UpdatePrevious);
+
+			if (UpdateResult.bNeedsPrevious)
+			{
+				UpdateRequest.UpdateMode = uint8(UpdateResult.bNeedsCurrent && UpdateResult.bNeedsPrevious && TargetPosition == TargetPreviousPosition && SourcePosition == SourcePreviousPosition
+					? EPreviousBoneTransformUpdateMode::DuplicateCurrentToPrevious
+					: EPreviousBoneTransformUpdateMode::UpdatePrevious);
+			}
 
 			const int32 NumLocalSubTracks = UpdateResult.bNeedsCurrent + (UpdateRequest.GetUpdateMode() == EPreviousBoneTransformUpdateMode::UpdatePrevious ? 1 : 0);
 
@@ -1067,6 +1072,7 @@ void FAnimSequenceTransformProvider::ProvideGPUTransforms(FSkinningTransformProv
 
 					BlockHeader.LeafBonePathBufferOffset = LeafBonePathBufferOffset;
 					BlockHeader.TransformBufferOffset    = IntermediateTransformOffset;
+					BlockHeader.NumPrefetchBones         = FMath::Min(ProxyUserData.NumMeshBones, BonesPerGroup);
 					BlockHeader.BlockLocalIndex          = BlockIndex;
 					BlockHeader.BlockLeafBoneCount       = GetBlockTransformCount(BlockIndex, NumLeafBones, BonesPerGroup);
 				}
