Fixed FXC compilation failure with landscape grass weight shader

#rnx

#Synced-CL 49985552

diff --git a/Engine/Shaders/Private/LandscapeGrassWeight.usf b/Engine/Shaders/Private/LandscapeGrassWeight.usf
--- a/Engine/Shaders/Private/LandscapeGrassWeight.usf
+++ b/Engine/Shaders/Private/LandscapeGrassWeight.usf
@@ -95,9 +95,12 @@ void PSMain(
 		break; \
 	}
 
-	UNROLL
+	// Do not remove : this unroll directive is actually necessary to avoid a FXC warning (error X3511: ... unrolled loop is too large)
+	UNROLL_N(4)
 	for (uint ChannelIndex = 0; ChannelIndex < 4; ++ChannelIndex)
 	{
+		// Do not remove : this branch directive is actually necessary to avoid FXC screwing up the generated shader
+		BRANCH
 		switch(PerChannelGrassWeightIndices[ChannelIndex])
 		{
 			case GrassWeightInvalidIndex: 
