Fix for shader compiler not supporting template store.

#Synced-CL 49999797
#Synced-CL 49999797

diff --git a/Engine/Shaders/Private/Skinning/AnimSequenceTransformProvider.usf b/Engine/Shaders/Private/Skinning/AnimSequenceTransformProvider.usf
--- a/Engine/Shaders/Private/Skinning/AnimSequenceTransformProvider.usf
+++ b/Engine/Shaders/Private/Skinning/AnimSequenceTransformProvider.usf
@@ -31,9 +31,9 @@ FBoneTransformWithScale LoadBoneTransformWithScale(ByteAddressBuffer Buffer, uin
 	uint Offset = BaseOffset + Index * (uint)sizeof(FBoneTransformWithScale);
 
 	FBoneTransformWithScale Output;
-	Output.Rotation       = Buffer.Load<float4>(Offset);
-	float4 LocationScaleX = Buffer.Load<float4>(Offset + 16);
-	float2 ScaleYZ        = Buffer.Load<float2>(Offset + 32);
+	Output.Rotation       = asfloat(Buffer.Load4(Offset));
+	float4 LocationScaleX = asfloat(Buffer.Load4(Offset + 16));
+	float2 ScaleYZ        = asfloat(Buffer.Load2(Offset + 32));
 
 	Output.Location = LocationScaleX.xyz;
 	Output.Scale = float3(LocationScaleX.w, ScaleYZ);
@@ -44,9 +44,9 @@ void StoreBoneTransformWithScale(RWByteAddressBuffer Buffer, uint BaseOffset, ui
 {
 	uint Offset = BaseOffset + Index * (uint)sizeof(FBoneTransformWithScale);
 
-	Buffer.Store<float4>(Offset, Transform.Rotation);
-	Buffer.Store<float4>(Offset + 16, float4(Transform.Location, Transform.Scale.x));
-	Buffer.Store<float2>(Offset + 32, Transform.Scale.yz);
+	Buffer.Store4(Offset, asuint(Transform.Rotation));
+	Buffer.Store4(Offset + 16, asuint(float4(Transform.Location, Transform.Scale.x)));
+	Buffer.Store2(Offset + 32, asuint(Transform.Scale.yz));
 }
 
 FBoneTransformWithScale EvaluateAndRetargetBone(
@@ -203,7 +203,7 @@ RWByteAddressBuffer GlobalSpaceTransformBufferUAV;
 
 groupshared FBoneTransformWithScale SharedTransform[BONES_PER_GROUP];
 
-FBoneTransformWithScale LoadSharedTransform(uint BoneIndex, uint TransformBufferOffset)
+FBoneTransformWithScale LoadSharedTransform(uint TransformBufferOffset, uint BoneIndex)
 {
 	if (BoneIndex < BONES_PER_GROUP)
 	{
@@ -248,7 +248,7 @@ void ConcatenateLocalToGlobalCS(uint3 GroupId : SV_GroupID, uint GroupIndex : SV
 	const uint PathHeaderOffset = BlockHeader.LeafBonePathBufferOffset + RelativePathHeaderOffset;
 
 	FLeafBonePathHeader PathHeader = LeafBonePathBuffer.Load<FLeafBonePathHeader>(PathHeaderOffset);
-	FBoneTransformWithScale ParentTransform = LoadSharedTransform(0, BlockHeader.TransformBufferOffset);
+	FBoneTransformWithScale ParentTransform = LoadSharedTransform(BlockHeader.TransformBufferOffset, 0);
 
 	if (PathHeader.bIsRootOwned)
 	{
@@ -265,7 +265,7 @@ void ConcatenateLocalToGlobalCS(uint3 GroupId : SV_GroupID, uint GroupIndex : SV
 			PathNode = LeafBonePathBuffer.Load<FLeafBonePathNode>(PathNodeOffset);
 			PathNodeOffset += PathElementSize;
 
-			FBoneTransformWithScale ChildTransform = LoadSharedTransform(PathNode.BoneIndex, BlockHeader.TransformBufferOffset);
+			FBoneTransformWithScale ChildTransform = LoadSharedTransform(BlockHeader.TransformBufferOffset, PathNode.BoneIndex);
 			ParentTransform = MultiplyBoneTransformWithScale(ChildTransform, ParentTransform);
 
 			if (PathNode.bIsOwned)
