Mobile: Fixed a few cases where DepthBuffer texture sample is converted to mediump

#Synced-CL 50034370

diff --git a/Engine/Shaders/Private/PostProcessMobile.usf b/Engine/Shaders/Private/PostProcessMobile.usf
--- a/Engine/Shaders/Private/PostProcessMobile.usf
+++ b/Engine/Shaders/Private/PostProcessMobile.usf
@@ -197,9 +197,9 @@ void SunMaskPS_Mobile(
 	half4 SceneColor = SampleSceneColor(InUVPos.xy,GetEyeIndex(StereoInput));
 
 	#if MOBILE_USEDEPTHTEXTURE 
-		half InDepth = ConvertFromDeviceZ(Texture2DSampleLevel(MobileSceneTextures.SceneDepthTexture, MobileSceneTextures.SceneDepthTextureSampler, InUVPos.xy, 0).r);
+		half InDepth = ConvertFromDeviceZ(MobileSceneTextures.SceneDepthTexture.SampleLevel(MobileSceneTextures.SceneDepthTextureSampler, InUVPos.xy, 0).r);
 	#else
-		half InDepth = ConvertFromDeviceZ(Texture2DSampleLevel(MobileSceneTextures.SceneDepthAuxTexture, MobileSceneTextures.SceneDepthAuxTextureSampler, InUVPos.xy, 0).r);
+		half InDepth = ConvertFromDeviceZ(MobileSceneTextures.SceneDepthAuxTexture.SampleLevel(MobileSceneTextures.SceneDepthAuxTextureSampler, InUVPos.xy, 0).r);
 	#endif
 	
 	#if MOBILE_USESUN
diff --git a/Engine/Shaders/Private/SceneTexturesCommon.ush b/Engine/Shaders/Private/SceneTexturesCommon.ush
--- a/Engine/Shaders/Private/SceneTexturesCommon.ush
+++ b/Engine/Shaders/Private/SceneTexturesCommon.ush
@@ -212,9 +212,9 @@ float LookupDeviceZ( float2 ScreenUV, float ArrayIndex )
 #elif FORCE_DEPTH_TEXTURE_READS || PLATFORM_NEEDS_DEPTH_TEXTURE_READS
 	// native Depth buffer lookup
 	#if MOBILE_MULTI_VIEW
-		return Texture2DArraySample(MobileSceneTextures.SceneDepthTextureArray, MobileSceneTextures.SceneDepthTextureSampler, float3(ScreenUV.xy,ArrayIndex)).r;
+		return MobileSceneTextures.SceneDepthTextureArray.SampleLevel(MobileSceneTextures.SceneDepthTextureSampler, float3(ScreenUV.xy,ArrayIndex), 0).r;
 	#else
-		return Texture2DSampleLevel(MobileSceneTextures.SceneDepthTexture, MobileSceneTextures.SceneDepthTextureSampler, ScreenUV, 0).r;
+		return MobileSceneTextures.SceneDepthTexture.SampleLevel(MobileSceneTextures.SceneDepthTextureSampler, ScreenUV, 0).r;
 	#endif
 #elif MOBILE_DEPTHFECTH && COMPILER_GLSL_ES3_1
 	return DepthbufferFetchES2();
@@ -226,16 +226,16 @@ float LookupDeviceZ( float2 ScreenUV, float ArrayIndex )
 	// SceneDepth texture is discarded after BasePass (with forward shading) 
 	// instead fetch DeviceZ from SceneDepthAuxTexture
 	#if MOBILE_MULTI_VIEW
-		return Texture2DArraySample(MobileSceneTextures.SceneDepthAuxTextureArray, MobileSceneTextures.SceneDepthAuxTextureSampler, float3(ScreenUV.xy,ArrayIndex)).r;
+		return MobileSceneTextures.SceneDepthAuxTextureArray.SampleLevel(MobileSceneTextures.SceneDepthAuxTextureSampler, float3(ScreenUV.xy,ArrayIndex), 0).r;
 	#else
-		return Texture2DSampleLevel(MobileSceneTextures.SceneDepthAuxTexture, MobileSceneTextures.SceneDepthAuxTextureSampler, ScreenUV, 0).r;
+		return MobileSceneTextures.SceneDepthAuxTexture.SampleLevel(MobileSceneTextures.SceneDepthAuxTextureSampler, ScreenUV, 0).r;
 	#endif
 #else
 	// native Depth buffer lookup
 	#if MOBILE_MULTI_VIEW
-		return Texture2DArraySample(MobileSceneTextures.SceneDepthTextureArray, MobileSceneTextures.SceneDepthTextureSampler, float3(ScreenUV.xy,ArrayIndex)).r;
+		return MobileSceneTextures.SceneDepthTextureArray.SampleLevel(MobileSceneTextures.SceneDepthTextureSampler, float3(ScreenUV.xy, ArrayIndex), 0).r;
 	#else
-		return Texture2DSampleLevel(MobileSceneTextures.SceneDepthTexture, MobileSceneTextures.SceneDepthTextureSampler, ScreenUV, 0).r;
+		return MobileSceneTextures.SceneDepthTexture.SampleLevel(MobileSceneTextures.SceneDepthTextureSampler, ScreenUV, 0).r;
 	#endif
 #endif
 }
