Fix for 16-bit Rotation not supporting asuint on certain shader compilers.

#Synced-CL 49999797

diff --git a/Engine/Shaders/Private/Skinning/AnimSequenceTransformProvider.usf b/Engine/Shaders/Private/Skinning/AnimSequenceTransformProvider.usf
--- a/Engine/Shaders/Private/Skinning/AnimSequenceTransformProvider.usf
+++ b/Engine/Shaders/Private/Skinning/AnimSequenceTransformProvider.usf
@@ -44,7 +44,7 @@ void StoreBoneTransformWithScale(RWByteAddressBuffer Buffer, uint BaseOffset, ui
 {
 	uint Offset = BaseOffset + Index * (uint)sizeof(FBoneTransformWithScale);
 
-	Buffer.Store4(Offset, asuint(Transform.Rotation));
+	Buffer.Store4(Offset, asuint(float4(Transform.Rotation)));
 	Buffer.Store4(Offset + 16, asuint(float4(Transform.Location, Transform.Scale.x)));
 	Buffer.Store2(Offset + 32, asuint(Transform.Scale.yz));
 }
