HairStrands: Fix position math not accounting for PreViewTranslation on triangle fallback

When intersection shaders are not available, the math for computing the world position was missing the view translation offset, leading to inconsistent  results.

Fix submitted by [REDACTED_PERSON]

#Synced-CL 49956609

diff --git a/Engine/Shaders/Private/HairStrands/HairStrandsVertexFactory.ush b/Engine/Shaders/Private/HairStrands/HairStrandsVertexFactory.ush
--- a/Engine/Shaders/Private/HairStrands/HairStrandsVertexFactory.ush
+++ b/Engine/Shaders/Private/HairStrands/HairStrandsVertexFactory.ush
@@ -640,7 +640,7 @@ FMaterialPixelParameters GetMaterialPixelParameters(float3 RayOrigin, float3 Ray
 		float4 B = float4(P1.Position, P1.WorldRadius);
 		float4 C = lerp(A, B, UCoord);
 		// approximate position along the width of the ribbon
-		float3 WorldC = DFDemote(TransformLocalToWorld(C.xyz, SceneData.Primitive.LocalToWorld));
+		float3 WorldC = DFDemote(CalcWorldPosition(C.xyz, SceneData.Primitive.LocalToWorld));
 		float3 rd = RayDirection;
 		float3 cp = WorldC - RayOrigin;
 		float t = dot(cp, rd) / dot(rd, rd);
