Fixed instance SkinningData not being written for Mobile GPU Scene.

#Synced-CL 49999797

diff --git a/Engine/Shaders/Private/SceneDataMobileLoader.ush b/Engine/Shaders/Private/SceneDataMobileLoader.ush
--- a/Engine/Shaders/Private/SceneDataMobileLoader.ush
+++ b/Engine/Shaders/Private/SceneDataMobileLoader.ush
@@ -248,6 +248,17 @@ FInstanceSceneData LoadInstanceDataUBO(uint DrawInstanceId)
 	DataIdx += 1u;
 #endif
 
+	// SkinningData
+	{
+		InstanceData.SkinningData = 0;
+		const bool bHasSkinningData = (InstanceData.Flags & INSTANCE_SCENE_DATA_FLAG_HAS_SKINNING_DATA) != 0u;
+		if (bHasSkinningData)
+		{
+			InstanceData.SkinningData = asuint(LoadInstanceDataElementUBO(DrawInstanceId, DataIdx+0).x);
+		}
+		DataIdx += 1u;
+	}
+
 	// CustomDataCount, TODO: yzw
 	{
 		float4 Data = LoadInstanceDataElementUBO(DrawInstanceId, DataIdx+0);
diff --git a/Engine/Shaders/Private/SceneDataMobileWriter.ush b/Engine/Shaders/Private/SceneDataMobileWriter.ush
--- a/Engine/Shaders/Private/SceneDataMobileWriter.ush
+++ b/Engine/Shaders/Private/SceneDataMobileWriter.ush
@@ -204,6 +204,17 @@ void WriteInstanceDataUBO(RWStructuredBuffer<float4> Output, uint Offset, uint I
 		DataIdx += 1u;
 	}
 #endif
+
+	// SkinningData
+	{
+		if (Offsets.SkinningData != INVALID_INSTANCE_PAYLOAD_OFFSET)
+		{
+			uint SkinningDataValue = asuint(LoadInstancePayloadDataElement(Offsets.SkinningData)).y;
+			WriteDataElementUBO(Output, Offset, DataIdx + 0, float4(asfloat(SkinningDataValue), 0.f, 0.f, 0.f));
+		}
+		DataIdx += 1u;
+	}
+
 	// CustomDataCount, TODO: yzw
 	{
 		CustomDataCount = min(CustomDataCount, (BATCHED_INSTANCE_DATA_STRIDE_FLOAT4 - DataIdx - 1u) * 4u);
