Reduced ASTP transform bandwidth by 20% by using uniform scale.

#Synced-CL 49999797

diff --git a/Engine/Shaders/Private/BoneTransform.ush b/Engine/Shaders/Private/BoneTransform.ush
--- a/Engine/Shaders/Private/BoneTransform.ush
+++ b/Engine/Shaders/Private/BoneTransform.ush
@@ -18,6 +18,13 @@ struct FBoneTransformWithScale
 	float3 Scale;
 };
 
+struct FBoneTransformWithUniformScale
+{
+	FQuat  Rotation;
+	float3 Location;
+	float  Scale;
+};
+
 ////
 
 FBoneTransform MakeBoneTransform(FQuat Rotation, float3 Location)
@@ -58,6 +65,14 @@ FBoneTransform RemoveBoneTransformScale(FBoneTransformWithScale Transform)
 	return Result;
 }
 
+FBoneTransform RemoveBoneTransformUniformScale(FBoneTransformWithUniformScale Transform)
+{
+	FBoneTransform Result;
+	Result.Location = Transform.Location;
+	Result.Rotation = QuatNormalize(Transform.Rotation);
+	return Result;
+}
+
 float4x4 MatrixFromBoneTransform(FBoneTransform Transform)
 {
 	float3x3 RotationMatrix = QuatToMatrix(Transform.Rotation);
@@ -126,4 +141,54 @@ FBoneTransformWithScale BlendBoneTransformsWithScale(FBoneTransformWithScale Ato
 	Blended.Rotation = QuatSlerpApproximate(Atom1.Rotation, Atom2.Rotation, Alpha);
 	Blended.Scale    = lerp(Atom1.Scale, Atom2.Scale, Alpha);
 	return Blended;
+}
+
+////
+
+FBoneTransformWithUniformScale MakeBoneTransformWithUniformScale(FQuat Rotation, float3 Location, float Scale)
+{
+	FBoneTransformWithUniformScale Transform;
+	Transform.Rotation = Rotation;
+	Transform.Location = Location;
+	Transform.Scale = Scale;
+	return Transform;
+}
+
+FBoneTransformWithUniformScale MakeBoneTransformWithUniformScaleIdentity()
+{
+	return MakeBoneTransformWithUniformScale(QuatIdentity(), float3(0.0f, 0.0f, 0.0f), 1.0);
+}
+
+FBoneTransformWithUniformScale MultiplyBoneTransformWithUniformScale(FBoneTransformWithUniformScale Atom1, FBoneTransformWithUniformScale Atom2)
+{
+	// NOTE: Assumes no negative scale! Otherwise should multiply using matrices
+	FQuat Rotation	= QuatMultiply(Atom2.Rotation, Atom1.Rotation);
+	float3 Location	= QuatRotateVector(Atom2.Rotation, Atom2.Scale * Atom1.Location) + Atom2.Location;
+	float Scale = Atom1.Scale * Atom2.Scale;
+	return MakeBoneTransformWithUniformScale(Rotation, Location, Scale);
+}
+
+float4x4 MatrixFromBoneTransformWithUniformScale(FBoneTransformWithUniformScale Transform)
+{
+	float3x3 RotationMatrix = QuatToMatrix(Transform.Rotation);
+
+	RotationMatrix[0] *= Transform.Scale;
+	RotationMatrix[1] *= Transform.Scale;
+	RotationMatrix[2] *= Transform.Scale;
+
+	float4x4 OutMatrix;
+	OutMatrix[0] = float4(RotationMatrix[0], 0.0);
+	OutMatrix[1] = float4(RotationMatrix[1], 0.0);
+	OutMatrix[2] = float4(RotationMatrix[2], 0.0);
+	OutMatrix[3] = float4(Transform.Location, 1.0);
+	return OutMatrix;
+}
+
+FBoneTransformWithUniformScale BlendBoneTransformsWithUniformScale(FBoneTransformWithUniformScale Atom1, FBoneTransformWithUniformScale Atom2, float Alpha)
+{
+	FBoneTransformWithUniformScale Blended;
+	Blended.Location = lerp(Atom1.Location, Atom2.Location, Alpha);
+	Blended.Rotation = QuatSlerpApproximate(Atom1.Rotation, Atom2.Rotation, Alpha);
+	Blended.Scale    = lerp(Atom1.Scale, Atom2.Scale, Alpha);
+	return Blended;
 }
\ No newline at end of file
diff --git a/Engine/Shaders/Private/Skinning/AnimSequenceTransformProvider.usf b/Engine/Shaders/Private/Skinning/AnimSequenceTransformProvider.usf
--- a/Engine/Shaders/Private/Skinning/AnimSequenceTransformProvider.usf
+++ b/Engine/Shaders/Private/Skinning/AnimSequenceTransformProvider.usf
@@ -26,30 +26,28 @@ FBoneRetargetingData LoadBoneRetargetingData(ByteAddressBuffer Buffer, uint Base
 	return Buffer.Load<FBoneRetargetingData>(Offset);
 }
 
-FBoneTransformWithScale LoadBoneTransformWithScale(ByteAddressBuffer Buffer, uint BaseOffset, uint Index)
+FBoneTransformWithUniformScale LoadBoneTransformWithUniformScale(ByteAddressBuffer Buffer, uint BaseOffset, uint Index)
 {
-	uint Offset = BaseOffset + Index * (uint)sizeof(FBoneTransformWithScale);
+	uint Offset = BaseOffset + Index * (uint)sizeof(FBoneTransformWithUniformScale);
 
-	FBoneTransformWithScale Output;
+	FBoneTransformWithUniformScale Output;
 	Output.Rotation       = asfloat(Buffer.Load4(Offset));
 	float4 LocationScaleX = asfloat(Buffer.Load4(Offset + 16));
-	float2 ScaleYZ        = asfloat(Buffer.Load2(Offset + 32));
 
 	Output.Location = LocationScaleX.xyz;
-	Output.Scale = float3(LocationScaleX.w, ScaleYZ);
+	Output.Scale    = LocationScaleX.w;
 	return Output;
 }
 
-void StoreBoneTransformWithScale(RWByteAddressBuffer Buffer, uint BaseOffset, uint Index, FBoneTransformWithScale Transform)
+void StoreBoneTransformWithUniformScale(RWByteAddressBuffer Buffer, uint BaseOffset, uint Index, FBoneTransformWithUniformScale Transform)
 {
-	uint Offset = BaseOffset + Index * (uint)sizeof(FBoneTransformWithScale);
+	uint Offset = BaseOffset + Index * (uint)sizeof(FBoneTransformWithUniformScale);
 
 	Buffer.Store4(Offset, asuint(float4(Transform.Rotation)));
 	Buffer.Store4(Offset + 16, asuint(float4(Transform.Location, Transform.Scale.x)));
-	Buffer.Store2(Offset + 32, asuint(Transform.Scale.yz));
 }
 
-FBoneTransformWithScale EvaluateAndRetargetBone(
+FBoneTransformWithUniformScale EvaluateAndRetargetBone(
 	uint MeshBoneIndex,
 	uint SequenceTransformBufferOffset,
 	uint RetargetingTransformBufferOffset,
@@ -64,7 +62,7 @@ FBoneTransformWithScale EvaluateAndRetargetBone(
 	const uint RetargetingTransformIndex = MeshBoneIndex * (RetargetingData.bRequiresRefPoseRetarget ? 2 : 1);
 	const bool bIsRootBone               = MeshBoneIndex == 0;
 
-	FBoneTransformWithScale RetargetingTransform = LoadBoneTransformWithScale(SequenceTransformBuffer, RetargetingTransformBufferOffset, RetargetingTransformIndex);
+	FBoneTransformWithUniformScale RetargetingTransform = LoadBoneTransformWithUniformScale(SequenceTransformBuffer, RetargetingTransformBufferOffset, RetargetingTransformIndex);
 
 	if (RetargetingData.AnimBoneIndex == BONE_RETARGETING_DATA_INVALID_BONE_INDEX)
 	{
@@ -75,13 +73,13 @@ FBoneTransformWithScale EvaluateAndRetargetBone(
 	const uint KeyIndexA = RetargetingData.AnimBoneIndex + KeyIndex0 * NumAnimBones;
 	const uint KeyIndexB = RetargetingData.AnimBoneIndex + KeyIndex1 * NumAnimBones;
 
-	FBoneTransformWithScale MeshLocalTransformA    = LoadBoneTransformWithScale(SequenceTransformBuffer, SequenceTransformBufferOffset, KeyIndexA);
-	FBoneTransformWithScale MeshLocalTransformB    = LoadBoneTransformWithScale(SequenceTransformBuffer, SequenceTransformBufferOffset, KeyIndexB);
-	FBoneTransformWithScale MeshLocalAnimTransform = BlendBoneTransformsWithScale(MeshLocalTransformA, MeshLocalTransformB, Alpha);
+	FBoneTransformWithUniformScale MeshLocalTransformA    = LoadBoneTransformWithUniformScale(SequenceTransformBuffer, SequenceTransformBufferOffset, KeyIndexA);
+	FBoneTransformWithUniformScale MeshLocalTransformB    = LoadBoneTransformWithUniformScale(SequenceTransformBuffer, SequenceTransformBufferOffset, KeyIndexB);
+	FBoneTransformWithUniformScale MeshLocalAnimTransform = BlendBoneTransformsWithUniformScale(MeshLocalTransformA, MeshLocalTransformB, Alpha);
 
 	if (RetargetingData.bRequiresRefPoseRetarget && !bIsRootBone)
 	{
-		FBoneTransformWithScale RefPoseData = LoadBoneTransformWithScale(SequenceTransformBuffer, RetargetingTransformBufferOffset, RetargetingTransformIndex + 1);
+		FBoneTransformWithUniformScale RefPoseData = LoadBoneTransformWithUniformScale(SequenceTransformBuffer, RetargetingTransformBufferOffset, RetargetingTransformIndex + 1);
 
 		// Unpack Q0 from Rotation
 		FQuat Q0 = FQuat(RefPoseData.Rotation[0], RefPoseData.Rotation[1], RefPoseData.Rotation[2], RefPoseData.Rotation[3]);
@@ -143,7 +141,7 @@ void EvaluateAnimSequenceCS(uint3 GroupId : SV_GroupID, uint GroupIndex : SV_Gro
 
 	const uint MeshBoneIndex = LocalMeshBoneIndex + (BlockHeader.BlockLocalIndex * BONES_PER_GROUP);
 
-	FBoneTransformWithScale MeshLocalAnimTransform = EvaluateAndRetargetBone(
+	FBoneTransformWithUniformScale MeshLocalAnimTransform = EvaluateAndRetargetBone(
 		MeshBoneIndex,
 		BlockHeader.SequenceTransformBufferOffset,
 		BlockHeader.RetargetingTransformBufferOffset,
@@ -153,7 +151,7 @@ void EvaluateAnimSequenceCS(uint3 GroupId : SV_GroupID, uint GroupIndex : SV_Gro
 		BlockHeader.Alpha,
 		BlockHeader.NumAnimBones);
 
-	StoreBoneTransformWithScale(LocalSpaceTransformBufferUAV, BlockHeader.DstTransformBufferOffset, MeshBoneIndex, MeshLocalAnimTransform);
+	StoreBoneTransformWithUniformScale(LocalSpaceTransformBufferUAV, BlockHeader.DstTransformBufferOffset, MeshBoneIndex, MeshLocalAnimTransform);
 }
 
 [numthreads(BONES_PER_GROUP, 1, 1)]
@@ -171,7 +169,7 @@ void EvaluateAnimSequenceBlendCS(uint3 GroupId : SV_GroupID, uint GroupIndex : S
 
 	const uint MeshBoneIndex = LocalMeshBoneIndex + (BlockHeader.BlockLocalIndex * BONES_PER_GROUP);
 
-	FBoneTransformWithScale SourceMeshLocalAnimTransform = EvaluateAndRetargetBone(
+	FBoneTransformWithUniformScale SourceMeshLocalAnimTransform = EvaluateAndRetargetBone(
 		MeshBoneIndex,
 		BlockHeader.SourceSequenceTransformBufferOffset,
 		BlockHeader.SourceRetargetingTransformBufferOffset,
@@ -181,7 +179,7 @@ void EvaluateAnimSequenceBlendCS(uint3 GroupId : SV_GroupID, uint GroupIndex : S
 		BlockHeader.SourceAlpha,
 		BlockHeader.SourceNumAnimBones);
 
-	FBoneTransformWithScale TargetMeshLocalAnimTransform = EvaluateAndRetargetBone(
+	FBoneTransformWithUniformScale TargetMeshLocalAnimTransform = EvaluateAndRetargetBone(
 		MeshBoneIndex,
 		BlockHeader.TargetSequenceTransformBufferOffset,
 		BlockHeader.TargetRetargetingTransformBufferOffset,
@@ -191,9 +189,9 @@ void EvaluateAnimSequenceBlendCS(uint3 GroupId : SV_GroupID, uint GroupIndex : S
 		BlockHeader.TargetAlpha,
 		BlockHeader.TargetNumAnimBones);
 
-	FBoneTransformWithScale MeshLocalAnimTransform = BlendBoneTransformsWithScale(SourceMeshLocalAnimTransform, TargetMeshLocalAnimTransform, BlockHeader.BlendWeight);
+	FBoneTransformWithUniformScale MeshLocalAnimTransform = BlendBoneTransformsWithUniformScale(SourceMeshLocalAnimTransform, TargetMeshLocalAnimTransform, BlockHeader.BlendWeight);
 
-	StoreBoneTransformWithScale(LocalSpaceTransformBufferUAV, BlockHeader.DstTransformBufferOffset, MeshBoneIndex, MeshLocalAnimTransform);
+	StoreBoneTransformWithUniformScale(LocalSpaceTransformBufferUAV, BlockHeader.DstTransformBufferOffset, MeshBoneIndex, MeshLocalAnimTransform);
 }
 
 ByteAddressBuffer LocalSpaceTransformBuffer;
@@ -201,9 +199,9 @@ ByteAddressBuffer ConcatenateBlockHeaderBuffer;
 ByteAddressBuffer LeafBonePathBuffer;
 RWByteAddressBuffer GlobalSpaceTransformBufferUAV;
 
-groupshared FBoneTransformWithScale SharedTransform[BONES_PER_GROUP];
+groupshared FBoneTransformWithUniformScale SharedTransform[BONES_PER_GROUP];
 
-FBoneTransformWithScale LoadSharedTransform(uint TransformBufferOffset, uint BoneIndex)
+FBoneTransformWithUniformScale LoadSharedTransform(uint TransformBufferOffset, uint BoneIndex)
 {
 	if (BoneIndex < BONES_PER_GROUP)
 	{
@@ -211,11 +209,11 @@ FBoneTransformWithScale LoadSharedTransform(uint TransformBufferOffset, uint Bon
 	}
 	else
 	{
-		return LoadBoneTransformWithScale(LocalSpaceTransformBuffer, TransformBufferOffset, BoneIndex);
+		return LoadBoneTransformWithUniformScale(LocalSpaceTransformBuffer, TransformBufferOffset, BoneIndex);
 	}
 }
 
-void StoreSharedTransform(uint BoneIndex, FBoneTransformWithScale Transform)
+void StoreSharedTransform(uint BoneIndex, FBoneTransformWithUniformScale Transform)
 {
 	SharedTransform[BoneIndex] = Transform;
 }
@@ -232,7 +230,7 @@ void ConcatenateLocalToGlobalCS(uint3 GroupId : SV_GroupID, uint GroupIndex : SV
 
 	if (PrefetchBoneIndex < BlockHeader.NumPrefetchBones)
 	{
-		StoreSharedTransform(PrefetchBoneIndex, LoadBoneTransformWithScale(LocalSpaceTransformBuffer, BlockHeader.TransformBufferOffset, PrefetchBoneIndex));
+		StoreSharedTransform(PrefetchBoneIndex, LoadBoneTransformWithUniformScale(LocalSpaceTransformBuffer, BlockHeader.TransformBufferOffset, PrefetchBoneIndex));
 	}
 	GroupMemoryBarrierWithGroupSync();
 
@@ -248,11 +246,11 @@ void ConcatenateLocalToGlobalCS(uint3 GroupId : SV_GroupID, uint GroupIndex : SV
 	const uint PathHeaderOffset = BlockHeader.LeafBonePathBufferOffset + RelativePathHeaderOffset;
 
 	FLeafBonePathHeader PathHeader = LeafBonePathBuffer.Load<FLeafBonePathHeader>(PathHeaderOffset);
-	FBoneTransformWithScale ParentTransform = LoadSharedTransform(BlockHeader.TransformBufferOffset, 0);
+	FBoneTransformWithUniformScale ParentTransform = LoadSharedTransform(BlockHeader.TransformBufferOffset, 0);
 
 	if (PathHeader.bIsRootOwned)
 	{
-		StoreBoneTransformWithScale(GlobalSpaceTransformBufferUAV, BlockHeader.TransformBufferOffset, 0, ParentTransform);
+		StoreBoneTransformWithUniformScale(GlobalSpaceTransformBufferUAV, BlockHeader.TransformBufferOffset, 0, ParentTransform);
 	}
 
 	if (PathHeader.PathDepth > 0)
@@ -265,12 +263,12 @@ void ConcatenateLocalToGlobalCS(uint3 GroupId : SV_GroupID, uint GroupIndex : SV
 			PathNode = LeafBonePathBuffer.Load<FLeafBonePathNode>(PathNodeOffset);
 			PathNodeOffset += PathElementSize;
 
-			FBoneTransformWithScale ChildTransform = LoadSharedTransform(BlockHeader.TransformBufferOffset, PathNode.BoneIndex);
-			ParentTransform = MultiplyBoneTransformWithScale(ChildTransform, ParentTransform);
+			FBoneTransformWithUniformScale ChildTransform = LoadSharedTransform(BlockHeader.TransformBufferOffset, PathNode.BoneIndex);
+			ParentTransform = MultiplyBoneTransformWithUniformScale(ChildTransform, ParentTransform);
 
 			if (PathNode.bIsOwned)
 			{
-				StoreBoneTransformWithScale(GlobalSpaceTransformBufferUAV, BlockHeader.TransformBufferOffset, PathNode.BoneIndex, ParentTransform);
+				StoreBoneTransformWithUniformScale(GlobalSpaceTransformBufferUAV, BlockHeader.TransformBufferOffset, PathNode.BoneIndex, ParentTransform);
 			}
 
 			if (PathNode.bIsLeaf)
@@ -308,10 +306,10 @@ void ScatterTransformsCS(uint3 GroupId : SV_GroupID, uint GroupIndex : SV_GroupI
 		MeshBoneIndex = BoneMapBuffer.Load<uint>(BlockHeader.BoneMapOffset + TransformIndex * uint(sizeof(uint)));
 	}
 
-	FBoneTransformWithScale InvGlobalRefPoseTransform = LoadBoneTransformWithScale(SequenceTransformBuffer, BlockHeader.InvGlobalRefPoseTransformBufferOffset, MeshBoneIndex);
-	FBoneTransformWithScale GlobalSpaceTransform = LoadBoneTransformWithScale(GlobalSpaceTransformBuffer, BlockHeader.SrcTransformBufferOffset, MeshBoneIndex);
-	GlobalSpaceTransform = MultiplyBoneTransformWithScale(InvGlobalRefPoseTransform, GlobalSpaceTransform);
-	float4x3 GlobalSpaceMatrix = (float4x3)MatrixFromBoneTransformWithScale(GlobalSpaceTransform);
+	FBoneTransformWithUniformScale InvGlobalRefPoseTransform = LoadBoneTransformWithUniformScale(SequenceTransformBuffer, BlockHeader.InvGlobalRefPoseTransformBufferOffset, MeshBoneIndex);
+	FBoneTransformWithUniformScale GlobalSpaceTransform = LoadBoneTransformWithUniformScale(GlobalSpaceTransformBuffer, BlockHeader.SrcTransformBufferOffset, MeshBoneIndex);
+	GlobalSpaceTransform = MultiplyBoneTransformWithUniformScale(InvGlobalRefPoseTransform, GlobalSpaceTransform);
+	float4x3 GlobalSpaceMatrix = (float4x3)MatrixFromBoneTransformWithUniformScale(GlobalSpaceTransform);
 
 	StoreCompressedBoneTransform(SkinningTransformBufferUAV, BlockHeader.DstTransformBufferOffset1, TransformIndex, GlobalSpaceMatrix);
 
diff --git a/Engine/Source/Runtime/Renderer/Private/Skinning/AnimSequenceTransformProvider.cpp b/Engine/Source/Runtime/Renderer/Private/Skinning/AnimSequenceTransformProvider.cpp
--- a/Engine/Source/Runtime/Renderer/Private/Skinning/AnimSequenceTransformProvider.cpp
+++ b/Engine/Source/Runtime/Renderer/Private/Skinning/AnimSequenceTransformProvider.cpp
@@ -57,7 +57,7 @@ static TAutoConsoleVariable<float> CVarAnimSequenceTransformProviderTimeScale(
 
 static TAutoConsoleVariable<bool> CVarAnimSequenceTransformProviderGPU(
 	TEXT("r.AnimSequenceTransformProvider.GPU"),
-	false,
+	true,
 	TEXT(""),
 	FConsoleVariableDelegate::CreateLambda([](IConsoleVariable* InVariable)
 	{
@@ -422,9 +422,11 @@ void FAnimSequenceTransformProvider::BuildBoneRetargetingTables(
 					FQuat DeltaRotation = RetargetedAnimRefTransform.GetRotation().Inverse() * MeshRefTransform.GetRotation();
 					DeltaRotation.Normalize();
 
+					const float ReciprocalScale = FMath::Abs(RetargetedAnimRefTransform.GetScale3D().X) < UE_SMALL_NUMBER ? 0.0f : RetargetedAnimRefTransform.GetScale3D().X;
+
 					Transform.SetRotation(FQuat4f(DeltaRotation));
 					Transform.SetTranslation(FVector3f(MeshRefTransform.GetTranslation() - RetargetedAnimRefTransform.GetTranslation()));
-					Transform.SetScale3D(FVector3f(MeshRefTransform.GetScale3D() * FTransform::GetSafeScaleReciprocal(RetargetedAnimRefTransform.GetScale3D())));
+					Transform.SetScale3D(FVector3f(MeshRefTransform.GetScale3D().X * ReciprocalScale));
 				}
 				break;
 				case EBoneTranslationRetargetingMode::OrientAndScale:
@@ -1231,13 +1233,11 @@ void FAnimSequenceTransformProvider::ProvideGPUTransforms(FSkinningTransformProv
 						RefPoseSlot.Rotation[2] = Q0.Z;
 						RefPoseSlot.Rotation[3] = Q0.W;
 
-						// Pack Q1 into Translation + Scale[0]
+						// Pack Q1 into Translation + Scale
 						RefPoseSlot.Translation[0] = Q1.X;
 						RefPoseSlot.Translation[1] = Q1.Y;
 						RefPoseSlot.Translation[2] = Q1.Z;
-						RefPoseSlot.Scale[0] = Q1.W;
-						RefPoseSlot.Scale[1] = 0.0f;
-						RefPoseSlot.Scale[2] = 0.0f;
+						RefPoseSlot.Scale          = Q1.W;
 					}
 				}
 			}
@@ -1624,20 +1624,20 @@ void FAnimSequenceTransformProvider::ProvideCPUTransforms(FSkinningTransformProv
 						break;
 						case EBoneTranslationRetargetingMode::AnimationScaled:
 						{
-							MeshLocalTransform.ScaleTranslation(RetargetTransform.GetScale3D());
+							MeshLocalTransform.ScaleTranslation(RetargetTransform.GetScale3D().X);
 						}
 						break;
 						case EBoneTranslationRetargetingMode::AnimationRelative:
 						{
 							MeshLocalTransform.SetRotation(MeshLocalTransform.GetRotation() * RetargetTransform.GetRotation());
 							MeshLocalTransform.SetTranslation(MeshLocalTransform.GetTranslation() + RetargetTransform.GetTranslation());
-							MeshLocalTransform.SetScale3D(MeshLocalTransform.GetScale3D() * RetargetTransform.GetScale3D());
+							MeshLocalTransform.SetScale3D(FVector3f(MeshLocalTransform.GetScale3D().X * RetargetTransform.GetScale3D().X));
 							MeshLocalTransform.NormalizeRotation();
 						}
 						break;
 						case EBoneTranslationRetargetingMode::OrientAndScale:
 						{
-							MeshLocalTransform.SetTranslation(RetargetTransform.GetRotation().RotateVector(MeshLocalTransform.GetTranslation()) * RetargetTransform.GetScale3D());
+							MeshLocalTransform.SetTranslation(RetargetTransform.GetRotation().RotateVector(MeshLocalTransform.GetTranslation()) * RetargetTransform.GetScale3D().X);
 							break;
 						}
 					}
@@ -1832,4 +1832,4 @@ void FAnimSequenceTransformProvider::ProvideCPUTransforms(FSkinningTransformProv
 		INC_DWORD_STAT_BY(STAT_AnimSequenceTransformProvider_CurrentTransformUpdates, NumCurrentTransformUpdates);
 		INC_DWORD_STAT_BY(STAT_AnimSequenceTransformProvider_PreviousTransformUpdates, NumPreviousTransformUpdates);
 	});
-}
\ No newline at end of file
+}
diff --git a/Engine/Source/Runtime/Renderer/Private/Skinning/AnimSequenceTransformProvider.h b/Engine/Source/Runtime/Renderer/Private/Skinning/AnimSequenceTransformProvider.h
--- a/Engine/Source/Runtime/Renderer/Private/Skinning/AnimSequenceTransformProvider.h
+++ b/Engine/Source/Runtime/Renderer/Private/Skinning/AnimSequenceTransformProvider.h
@@ -69,7 +69,7 @@ private:
 		{
 			SetRotation(Xfm.GetRotation());
 			SetTranslation(Xfm.GetTranslation());
-			SetScale(Xfm.GetScale3D());
+			SetScale(Xfm.GetScale3D().X);
 		}
 
 		void SetRotation(FQuat4f Q)
@@ -87,16 +87,14 @@ private:
 			Translation[2] = P.Z;
 		}
 
-		void SetScale(FVector3f S)
+		void SetScale(float S)
 		{
-			Scale[0] = S.X;
-			Scale[1] = S.Y;
-			Scale[2] = S.Z;
+			Scale = S;
 		}
 
 		float Rotation[4];
 		float Translation[3];
-		float Scale[3];
+		float Scale;
 	};
 
 	using FSequenceRegistryKey = const UAnimSequence*;
