Build-fix to previous CL UE-227651

#rnx

#Synced-CL 50155647

diff --git a/Engine/Shaders/Private/MaterialTemplate.ush b/Engine/Shaders/Private/MaterialTemplate.ush
--- a/Engine/Shaders/Private/MaterialTemplate.ush
+++ b/Engine/Shaders/Private/MaterialTemplate.ush
@@ -722,16 +722,19 @@ struct FMaterialVertexParameters
 
 	half4 VertexColor;
 
+#if ENABLE_NEW_HLSL_GENERATOR && NUM_TEX_COORD_INTERPOLATORS
 	// NUM_MATERIAL_TEXCOORDS_VERTEX means the number of actual UV TexCoords used by the shader
 	// NUM_TEX_COORD_INTERPOLATORS instead is the total interpolator registers to use: this is the value above plus all the user interpolators expressions in the material
 	// Using NUM_TEX_COORD_INTERPOLATORS here so that we have a single array ot hold all interpolated data (UVs plus user).
 	// Note: FMaterialVertex/PixelParameters structs are "local" to the material
 	// Note: The code that maps to these registers is generated by ComputeUserVertexInterpolatorAllocation() in the new translator
-#if NUM_TEX_COORD_INTERPOLATORS
 	float2 TexCoords[NUM_TEX_COORD_INTERPOLATORS];
-	#if FEATURE_LEVEL == FEATURE_LEVEL_ES3_1
-		float2 TexCoordOffset; // Offset for UV localization for large UV values
-	#endif
+#elif NUM_MATERIAL_TEXCOORDS_VERTEX
+	float2 TexCoords[NUM_MATERIAL_TEXCOORDS_VERTEX];
+#endif
+
+#if NUM_MATERIAL_TEXCOORDS_VERTEX && FEATURE_LEVEL == FEATURE_LEVEL_ES3_1
+	float2 TexCoordOffset; // Offset for UV localization for large UV values
 #endif
 
 	/** Per-particle properties. Only valid for particle vertex factories. */
