Substrate - Toon - Moved from diffuse/spc colro to Base color / metallic / specular model. Added support for PatternUVs in [0, 1].

#tests PC editor dx12
#rnx

#Synced-CL 50246125

diff --git a/Engine/Shaders/Private/DBufferDecalShared.ush b/Engine/Shaders/Private/DBufferDecalShared.ush
--- a/Engine/Shaders/Private/DBufferDecalShared.ush
+++ b/Engine/Shaders/Private/DBufferDecalShared.ush
@@ -354,14 +354,13 @@ void ApplyDBufferData(in FSubstrateDBuffer SubstrateDBufferData, inout FSubstrat
 			case SUBSTRATE_BSDF_TYPE_TOON:
 			{
 			#if SUBSTRATE_DBUFFER_RESPONSE_BASECOLOR
-				TOON_DIFFUSEALBEDO(BSDF) = TOON_DIFFUSEALBEDO(BSDF) * SubstrateDBufferData.OneMinusCoverage_BaseColor + SubstrateDBufferData.BaseColor;
+				TOON_BASECOLOR(BSDF) = TOON_BASECOLOR(BSDF) * SubstrateDBufferData.OneMinusCoverage_BaseColor + SubstrateDBufferData.BaseColor;
 			#endif
 
 			#if SUBSTRATE_DBUFFER_RESPONSE_ROUGHNESS
 				const float DBufferAnisotropy = 0.0f;
 				TOON_ROUGHNESS(BSDF)		= TOON_ROUGHNESS(BSDF) * SubstrateDBufferData.OneMinusCoverage_Roughness		+ SubstrateDBufferData.Roughness;
-				TOON_SPECULARCOLOR(BSDF)	= TOON_SPECULARCOLOR(BSDF)  * SubstrateDBufferData.OneMinusCoverage_Roughness	+ SubstrateDBufferData.Specular;
-				TOON_ANISOTROPY(BSDF)		= TOON_ANISOTROPY(BSDF)  * SubstrateDBufferData.OneMinusCoverage_Roughness		+ DBufferAnisotropy;
+				TOON_SPECULAR(BSDF)			= TOON_SPECULAR(BSDF)  * SubstrateDBufferData.OneMinusCoverage_Roughness		+ SubstrateDBufferData.Specular;
 			#endif
 				break;
 			}
diff --git a/Engine/Shaders/Private/PackUnpack.ush b/Engine/Shaders/Private/PackUnpack.ush
--- a/Engine/Shaders/Private/PackUnpack.ush
+++ b/Engine/Shaders/Private/PackUnpack.ush
@@ -149,6 +149,20 @@ float UnpackR6(uint In)
 	return float(In & 0x3F) * (1.0f / 63.0f);
 }
 
+//
+//	float <-> u5
+//
+
+uint PackR5(float Value)
+{
+	return uint(saturate(Value) * 31.0f);
+}
+
+float UnpackR5(uint In)
+{
+	return float(In & 0x1F) * (1.0f / 31.0f);
+}
+
 //
 //	float4 <-> u32 as u8.u8.u8.u8
 //
@@ -436,6 +450,7 @@ float DitherXbits(float Value, float Dither, float QuantizationScale)
 	return DitherIfNonBlack(saturate(Value), DitheredXbits);
 }
 
+#define Dither5bits(Value, Dither) DitherXbits(Value, Dither, 1.f/31.f)
 #define Dither6bits(Value, Dither) DitherXbits(Value, Dither, 1.f/63.f)
 #define Dither7bits(Value, Dither) DitherXbits(Value, Dither, 1.f/127.f)
 #define Dither8bits(Value, Dither) DitherXbits(Value, Dither, 1.f/255.f)
\ No newline at end of file
diff --git a/Engine/Shaders/Private/PostProcessGBufferHints.usf b/Engine/Shaders/Private/PostProcessGBufferHints.usf
--- a/Engine/Shaders/Private/PostProcessGBufferHints.usf
+++ b/Engine/Shaders/Private/PostProcessGBufferHints.usf
@@ -262,14 +262,12 @@ void MainCS(uint3 DispatchThreadId : SV_DispatchThreadID)
 	#if SUBSTRATE_ENABLED
 		else if (Data.ShadingModelID == SHADINGMODELID_SUBSTRATE_TOON)
 		{
-			float3 SpecularColor = 0.0f;
-			float Anisotropy = 0.0f;
+			float2 PatternUVs = 0.0f;
 			float ToonProfileId = 0.0f;
-			UnpackToonCustomData(PackRGBA8(Data.CustomData), SpecularColor, Anisotropy, ToonProfileId);
+			UnpackToonCustomData(PackRGBA8(Data.CustomData), PatternUVs, ToonProfileId);
 
-			Print(Ctx, TEXT("Specular Color     : "), FontWhite); Print(Ctx, SpecularColor); Newline(Ctx, RectMax);
+			Print(Ctx, TEXT("Pattern UVs        : "), FontWhite); Print(Ctx, PatternUVs); Newline(Ctx, RectMax);
 			Print(Ctx, TEXT("Toon Profil ID     : "), FontWhite); Print(Ctx, ToonProfileId); Newline(Ctx, RectMax);
-			Print(Ctx, TEXT("Anisotropy         : "), FontWhite); Print(Ctx, Anisotropy); Newline(Ctx, RectMax);
 		}
 	#endif
 		PrintBackgroundRect(Ctx, RectMin, RectMax);
diff --git a/Engine/Shaders/Private/ShadowProjectionPixelShader.usf b/Engine/Shaders/Private/ShadowProjectionPixelShader.usf
--- a/Engine/Shaders/Private/ShadowProjectionPixelShader.usf
+++ b/Engine/Shaders/Private/ShadowProjectionPixelShader.usf
@@ -325,10 +325,9 @@ void Main(
 				{
 					bApplyPCFOverBlurCorrection = false;
 
-					float3 SpecularColor = 0.0f;
-					float Anisotropy = 0.0f;
+					float2 PatternUVs = 0.0f;
 					float ToonProfileId = 0.0f;
-					UnpackToonCustomData(PackRGBA8(GBufferData.CustomData), SpecularColor, Anisotropy, ToonProfileId);
+					UnpackToonCustomData(PackRGBA8(GBufferData.CustomData), PatternUVs, ToonProfileId);
 
 					// Density is extinction as 1/cm (not 1/m) to match with the transmittance shadow evaluation from ManualPCF below.
 					Density = GetToonProfileShadowSettings(ToonProfileId).ShadowExtinctionCoefficientInInvCentimeter;
diff --git a/Engine/Shaders/Private/Substrate/Substrate.ush b/Engine/Shaders/Private/Substrate/Substrate.ush
--- a/Engine/Shaders/Private/Substrate/Substrate.ush
+++ b/Engine/Shaders/Private/Substrate/Substrate.ush
@@ -714,11 +714,12 @@ struct FSubstrateBSDF
 ///////////////////////////////////////////////////////////////////////////////
 // Toon BSDF members
 // TOON_PROFILEID is in [0, 255]
-#define TOON_DIFFUSEALBEDO(X)			X.VGPRs[0].xyz
-#define TOON_ROUGHNESS(X)				X.VGPRs[0].w
-#define TOON_SPECULARCOLOR(X)			X.VGPRs[1].xyz
-#define TOON_ANISOTROPY(X)				X.VGPRs[1].w
-#define TOON_PROFILEID(X)				X.VGPRs[2].x
+#define TOON_BASECOLOR(X)				X.VGPRs[0].xyz
+#define TOON_PROFILEID(X)				X.VGPRs[0].w
+#define TOON_METALLIC(X)				X.VGPRs[1].x
+#define TOON_SPECULAR(X)				X.VGPRs[1].y
+#define TOON_ROUGHNESS(X)				X.VGPRs[1].z
+#define TOON_PATTERNUVS(X)				X.VGPRs[2].xy
 
 ///////////////////////////////////////////////////////////////////////////////
 // Forward declarations
@@ -1979,13 +1980,14 @@ FSubstrateData GetSubstrateSingleLayerWaterBSDF(
 
 FSubstrateData GetSubstrateToonBSDF(
 	float ToonProfileId, // in [0, 255]
-	float3 DiffuseColor, float3 SpecularColor, float Roughness, float Anisotropy, 
+	float3 BaseColor, float Metallic, float Specular, float Roughness,
 	float3 Emissive,
+	float2 PatternUVs,
 	uint SharedLocalBasisIndex, inout uint SharedLocalBasisTypes)
 {
 	FSubstrateData SubstrateData = GetInitialisedSubstrateData();
 
-#if MATERIAL_USES_ANISOTROPY && SUBSTRATE_FASTPATH==0
+#if 0 // MATERIAL_USES_ANISOTROPY && SUBSTRATE_FASTPATH==0
 	const uint bHasAnisotropy = Anisotropy != 0 ? 1 : 0;
 	if (bHasAnisotropy)
 	{
@@ -1993,7 +1995,7 @@ FSubstrateData GetSubstrateToonBSDF(
 	}
 #else
 	const uint bHasAnisotropy = 0;
-	Anisotropy = 0.0f;
+//	Anisotropy = 0.0f;
 #endif
 
 #if (SUBSTRATE_FASTPATH==0 || MATERIALBLENDING_ANY_TRANSLUCENT) // !Fast path or rendering translucent materials
@@ -2005,19 +2007,18 @@ FSubstrateData GetSubstrateToonBSDF(
 	BSDF_SETSHAREDLOCALBASISID		(SubstrateData.InlinedBSDF, SharedLocalBasisIndex);
 	BSDF_SETEMISSIVE				(SubstrateData.InlinedBSDF, Emissive);
 	TOON_PROFILEID					(SubstrateData.InlinedBSDF) = ToonProfileId;
-	TOON_DIFFUSEALBEDO				(SubstrateData.InlinedBSDF) = DiffuseColor;
-	TOON_SPECULARCOLOR				(SubstrateData.InlinedBSDF) = SpecularColor;
+	TOON_BASECOLOR					(SubstrateData.InlinedBSDF) = BaseColor;
+	TOON_METALLIC					(SubstrateData.InlinedBSDF) = Metallic;
+	TOON_SPECULAR					(SubstrateData.InlinedBSDF) = Specular;
 	TOON_ROUGHNESS					(SubstrateData.InlinedBSDF) = Roughness;
-#if MATERIAL_USES_ANISOTROPY
-	TOON_ANISOTROPY					(SubstrateData.InlinedBSDF) = Anisotropy;
-#endif
+	TOON_PATTERNUVS					(SubstrateData.InlinedBSDF) = PatternUVs;
 
 #if SUBSTRATE_INLINE_SHADING
 	SubstrateData.InlinedBSDF.Coverage	= 1.0f;
 #endif
 
 #if USE_DEVELOPMENT_SHADERS
-	SubstrateData.PreviewColor			= DiffuseColor;
+	SubstrateData.PreviewColor			= BaseColor;
 #endif
 
 	return SubstrateData;
@@ -3573,10 +3574,11 @@ void FSubstrateBSDF::SubstrateSanitizeBSDF()
 			// We do not clamp emissive luminance upper bound. This is done later in the base pass within pre-exposed luminance space.
 			BSDF_SETEMISSIVE(this, max(BSDF_GETEMISSIVE(this), 0.0));
 
-			TOON_DIFFUSEALBEDO(this) = SanitizeDiffuseAlbedo(TOON_DIFFUSEALBEDO(this));
-			TOON_SPECULARCOLOR(this) = SanitizeF0(TOON_SPECULARCOLOR(this));
+			TOON_BASECOLOR(this) = SanitizeDiffuseAlbedo(TOON_BASECOLOR(this));
+			TOON_METALLIC(this) = clamp(TOON_METALLIC(this), 0.0, 1.0);
+			TOON_SPECULAR(this) = clamp(TOON_SPECULAR(this), 0.0, 1.0);
 			TOON_ROUGHNESS(this) = clamp(TOON_ROUGHNESS(this), 0.0, 1.0);	// We do not use SanitizeRoughness which clamps to View.MinRoughness that is higher than 0.
-			TOON_ANISOTROPY(this) = SanitizeAnisotropy(TOON_ANISOTROPY(this));
+			TOON_PATTERNUVS(this) = clamp(TOON_PATTERNUVS(this), 0.0.xx, 1.0.xx);
 			break;
 		}
 	}
@@ -3750,7 +3752,7 @@ float3 SubstrateGetBSDFBaseColor(in FSubstrateBSDF BSDF)
 	case SUBSTRATE_BSDF_TYPE_HAIR: 				return HAIR_BASECOLOR(BSDF);
 	case SUBSTRATE_BSDF_TYPE_EYE: 				return EYE_DIFFUSEALBEDO(BSDF);
 	case SUBSTRATE_BSDF_TYPE_SINGLELAYERWATER:	return SLW_BASECOLOR(BSDF);
-	case SUBSTRATE_BSDF_TYPE_TOON:				return TOON_DIFFUSEALBEDO(BSDF);
+	case SUBSTRATE_BSDF_TYPE_TOON:				return TOON_BASECOLOR(BSDF);
 	}
 	return 0.f;
 }
@@ -3763,7 +3765,7 @@ float3 SubstrateGetBSDFDiffuseColor(in FSubstrateBSDF BSDF)
 	case SUBSTRATE_BSDF_TYPE_HAIR: 				return HAIR_BASECOLOR(BSDF);
 	case SUBSTRATE_BSDF_TYPE_EYE: 				return EYE_DIFFUSEALBEDO(BSDF);
 	case SUBSTRATE_BSDF_TYPE_SINGLELAYERWATER:	return SLW_BASECOLOR(BSDF);
-	case SUBSTRATE_BSDF_TYPE_TOON:				return TOON_DIFFUSEALBEDO(BSDF);
+	case SUBSTRATE_BSDF_TYPE_TOON:				return TOON_BASECOLOR(BSDF);
 	}
 	return 0.f;
 }
@@ -3776,7 +3778,7 @@ void SubstrateSetBSDFDiffuseColor(inout FSubstrateBSDF BSDF, float3 InColor)
 	case SUBSTRATE_BSDF_TYPE_HAIR: 				HAIR_BASECOLOR(BSDF) = InColor; break;
 	case SUBSTRATE_BSDF_TYPE_EYE: 				EYE_DIFFUSEALBEDO(BSDF) = InColor; break;
 	case SUBSTRATE_BSDF_TYPE_SINGLELAYERWATER:	SLW_BASECOLOR(BSDF) = InColor; break;
-	case SUBSTRATE_BSDF_TYPE_TOON:				TOON_DIFFUSEALBEDO(BSDF) = InColor; break;
+	case SUBSTRATE_BSDF_TYPE_TOON:				TOON_BASECOLOR(BSDF) = InColor; break;
 	}
 }
 
@@ -3788,7 +3790,7 @@ float3 SubstrateGetBSDFSpecularF0(in FSubstrateBSDF BSDF)
 	case SUBSTRATE_BSDF_TYPE_HAIR: 				return DielectricSpecularToF0(HAIR_SPECULAR(BSDF));
 	case SUBSTRATE_BSDF_TYPE_EYE: 				return EYE_F0(BSDF);
 	case SUBSTRATE_BSDF_TYPE_SINGLELAYERWATER:	return DielectricSpecularToF0(SLW_SPECULAR(BSDF));
-	case SUBSTRATE_BSDF_TYPE_TOON:				return TOON_SPECULARCOLOR(BSDF);
+	case SUBSTRATE_BSDF_TYPE_TOON:				return DielectricSpecularToF0(TOON_SPECULAR(BSDF));
 	}
 	return 0.f;
 }
@@ -3813,7 +3815,7 @@ void SubstrateSetBSDFSpecularF0(in FSubstrateBSDF BSDF, float3 F0)
 	case SUBSTRATE_BSDF_TYPE_HAIR: 				HAIR_SPECULAR(BSDF) = F0RGBToDielectricSpecular(F0); break;
 	case SUBSTRATE_BSDF_TYPE_EYE: 				/* EYE_F0 us hardcoded */ break;
 	case SUBSTRATE_BSDF_TYPE_SINGLELAYERWATER:	SLW_SPECULAR(BSDF) = F0RGBToDielectricSpecular(F0); break;
-	case SUBSTRATE_BSDF_TYPE_TOON:				TOON_SPECULARCOLOR(BSDF) = F0; break;
+	case SUBSTRATE_BSDF_TYPE_TOON:				TOON_SPECULAR(BSDF) = F0RGBToDielectricSpecular(F0); break;
 	}
 }
 
@@ -3873,6 +3875,7 @@ float SubstrateGetBSDFSpecular(in FSubstrateBSDF BSDF)
 	case SUBSTRATE_BSDF_TYPE_HAIR: 				return HAIR_SPECULAR(BSDF);
 	case SUBSTRATE_BSDF_TYPE_EYE: 				return F0ToDielectricSpecular(EYE_F0(BSDF).x);
 	case SUBSTRATE_BSDF_TYPE_SINGLELAYERWATER:	return SLW_SPECULAR(BSDF);
+	case SUBSTRATE_BSDF_TYPE_TOON:				return TOON_SPECULAR(BSDF);
 	}
 	return 0.f;
 }
@@ -3881,10 +3884,11 @@ float SubstrateGetBSDFMetallic(in FSubstrateBSDF BSDF)
 {
 	switch (BSDF_GETTYPE(BSDF))
 	{
-	case SUBSTRATE_BSDF_TYPE_SLAB: 			return F0RGBToMetallic(SLAB_F0(BSDF)); // Metallic recovers is not accurate, and only used for legacy conversion purpose (e.g., PostProcessMaterial texture access)
-	case SUBSTRATE_BSDF_TYPE_HAIR: 			return 0;
+	case SUBSTRATE_BSDF_TYPE_SLAB: 				return F0RGBToMetallic(SLAB_F0(BSDF)); // Metallic recovers is not accurate, and only used for legacy conversion purpose (e.g., PostProcessMaterial texture access)
+	case SUBSTRATE_BSDF_TYPE_HAIR: 				return 0;
 	case SUBSTRATE_BSDF_TYPE_EYE: 				return 0;
-	case SUBSTRATE_BSDF_TYPE_SINGLELAYERWATER: return 0;
+	case SUBSTRATE_BSDF_TYPE_SINGLELAYERWATER:	return 0;
+	case SUBSTRATE_BSDF_TYPE_TOON:				return TOON_METALLIC(BSDF);
 	}
 	return 0.f;
 }
@@ -3934,7 +3938,7 @@ float SubstrateGetBSDFAnisotropy(in FSubstrateBSDF BSDF)
 	case SUBSTRATE_BSDF_TYPE_HAIR:				return 0;
 	case SUBSTRATE_BSDF_TYPE_EYE:				return 0;
 	case SUBSTRATE_BSDF_TYPE_SINGLELAYERWATER:	return 0;
-	case SUBSTRATE_BSDF_TYPE_TOON:				return BSDF_GETHASANISOTROPY(BSDF) ? TOON_ANISOTROPY(BSDF) : 0.0f;
+	case SUBSTRATE_BSDF_TYPE_TOON:				return 0;
 	}
 	return 0.f;
 }
@@ -4144,6 +4148,33 @@ float3 UnpackR7G7B6Gamma2(uint rgb)
 	return SubstrateSrgbToLinear(Out);
 }
 
+uint PackR6G6B5Gamma2(float3 rgb, float Dither)
+{
+	// We execute a dither that can go in positive/negative direction to avoid surfaces to looks overall brighter.
+	const float CenteredDither = Dither - 0.5f;
+	const float Dither6bits = CenteredDither * 1.0f / 63.0f;
+	const float Dither5bits = CenteredDither * 1.0f / 31.0f;
+	const float3 Dither665 = float3(Dither6bits, Dither6bits, Dither5bits);
+
+	// Dither value to hide more aggressive bit reduction.
+	const float3 dithered_rgb = DitherIfNonBlack(SubstrateLinearToSrgb(saturate(rgb)), Dither665);
+
+	uint r = (uint(dithered_rgb.r * 63.0f) << 0);
+	uint g = (uint(dithered_rgb.g * 63.0f) << 6);
+	uint b = (uint(dithered_rgb.b * 31.0f) << 12);
+	return r | g | b;
+}
+
+float3 UnpackR6G6B5Gamma2(uint rgb)
+{
+	float3 Out;
+	Out.r = float((rgb >> 0)  & 0x3F) * (1.0f / 63.0f);
+	Out.g = float((rgb >> 6)  & 0x3F) * (1.0f / 63.0f);
+	Out.b = float((rgb >> 12) & 0x1F) * (1.0f / 31.0f);
+
+	return SubstrateSrgbToLinear(Out);
+}
+
 uint PackR7G7B6Linear(float3 rgb, float Dither)
 {
 	// We execute a dither that can go in positive/negative direction to avoid surfaces to looks overall brighter.
@@ -4171,31 +4202,25 @@ float3 UnpackR7G7B6Linear(uint rgb)
 	return Out;
 }
 
-uint PackToonCustomData(in float3 SpecCol, in float Aniso, in float ProfileId, in float Dither)
+uint PackToonCustomData(in float2 UVs, in float ProfileId, in float Dither)
 {
-	// We execute a dither that can go in positive/negative direction to avoid surfaces to looks overall brighter.
-	const float CenteredDither = Dither - 0.5f;
-	const float Dither8bits = CenteredDither * 1.0f / 255.0f;
-	const float Dither7bits = CenteredDither * 1.0f / 127.0f;
-	const float Dither6bits = CenteredDither * 1.0f / 63.0f;
-	const float Dither5bits = CenteredDither * 1.0f / 31.0f;
-
-	SpecCol = DitherIfNonBlack(SubstrateLinearToSrgb(saturate(SpecCol)), float3(Dither5bits, Dither6bits, Dither5bits));
-	Aniso = DitherIfNonBlack(saturate(Aniso), Dither8bits);
+	uint U12bits = saturate(UVs.x) * 4095.0f;
+	uint V12bits = saturate(UVs.y) * 4095.0f;
 
-	uint Packed =	(uint(SpecCol.r *  31.0f) & 0x1F) <<  0;	// Spec color R 5 bits
-	Packed |=		(uint(SpecCol.g *  63.0f) & 0x3F) <<  5;	// Spec color G 6 bits
-	Packed |=		(uint(SpecCol.b *  31.0f) & 0x1F) << 11;	// Spec Color B 5 bits
-	Packed |=		(uint(Aniso     * 255.0f) & 0xFF) << 16;	// Aniso        8 bits
-	Packed |=		(uint(ProfileId         ) & 0xFF) << 24;	// ProfileId    8 bits
+	uint Packed	 = U12bits;
+	Packed		|= V12bits << 12;
+	Packed		|= (uint(ProfileId) & 0xFF) << 24;
 
 	return Packed;
 }
 
-void UnpackToonCustomData(in uint PackedToonCustomData, inout float3 SpecCol, inout float Aniso, inout float ProfileId)
+void UnpackToonCustomData(in uint PackedToonCustomData, inout float2 UVs, inout float ProfileId)
 {
-	SpecCol		= SubstrateSrgbToLinear(float3(BitFieldExtractU32(PackedToonCustomData, 5, 0), BitFieldExtractU32(PackedToonCustomData, 6, 5), BitFieldExtractU32(PackedToonCustomData, 5, 11)) * float3(1.0f / 31.0f, 1.0f / 63.0f, 1.0f / 31.0f));
-	Aniso		= float(UnpackB8(PackedToonCustomData));
+	uint U12bits = float(BitFieldExtractU32(PackedToonCustomData, 12, 0));
+	uint V12bits = float(BitFieldExtractU32(PackedToonCustomData, 12, 12));
+
+	UVs = float2(U12bits, V12bits) * (1.0f / 4095.0f);
+
 	ProfileId	= float(BitFieldExtractU32(PackedToonCustomData, 8, 24)); // Stays in [0, 255]
 }
 
@@ -4214,6 +4239,11 @@ uint PackR6(float Value, float Dither)
 	return uint(Dither6bits(Value, Dither) * 63.0f);
 }
 
+uint PackR5(float Value, float Dither)
+{
+	return uint(Dither5bits(Value, Dither) * 31.0f);
+}
+
 // Simple Octahedral from "A Survey of Efficient Representations for Independent Unit Vectors" http://jcgt.org/published/0003/02/01/
 float SignNotZero(in float k)
 {
@@ -4896,11 +4926,12 @@ FSubstrateBSDF UnpackSubstrateBSDFIn(FSubstrateMaterialContainer SubstrateBuffer
 			#if SUBSTRATE_DEFERRED_SHADING
 			if (SubstrateHeader.GetComplexPathMode() == COMPLEXPATH_MODE_TOON)
 			{
-				float4 DiffuseRoughness		= UnpackColorGamma2ToLinearAlphaLinear(DataX);
-				TOON_DIFFUSEALBEDO(OutBSDF)	= DiffuseRoughness.xyz;
-				TOON_ROUGHNESS(OutBSDF)		= DiffuseRoughness.w;
-
-				UnpackToonCustomData(DataY, TOON_SPECULARCOLOR(OutBSDF), TOON_ANISOTROPY(OutBSDF), TOON_PROFILEID(OutBSDF));
+				TOON_BASECOLOR(OutBSDF)	= UnpackR6G6B5Gamma2(DataX);
+				TOON_SPECULAR(OutBSDF)	= UnpackR5(DataX >> 17);
+				TOON_METALLIC(OutBSDF)	= UnpackR5(DataX >> 22);
+				TOON_ROUGHNESS(OutBSDF)	= UnpackR5(DataX >> 27);
+		
+				UnpackToonCustomData(DataY, TOON_PATTERNUVS(OutBSDF), TOON_PROFILEID(OutBSDF));
 
 				BSDF_SETSSSTYPE(OutBSDF, SSS_TYPE_TRANSMITTANCE_SHADOW);
 
diff --git a/Engine/Shaders/Private/Substrate/SubstrateEvaluation.ush b/Engine/Shaders/Private/Substrate/SubstrateEvaluation.ush
--- a/Engine/Shaders/Private/Substrate/SubstrateEvaluation.ush
+++ b/Engine/Shaders/Private/Substrate/SubstrateEvaluation.ush
@@ -1315,19 +1315,18 @@ FSubstrateEvaluateResult SubstrateEvaluateBSDFCommon(FSubstrateBSDFContext BSDFC
 
 		case SUBSTRATE_BSDF_TYPE_TOON:
 		{
-			const float3 DiffuseColor		= TOON_DIFFUSEALBEDO(BSDFContext.BSDF);
-			const float3 SpecularColor		= TOON_SPECULARCOLOR(BSDFContext.BSDF);
+			const float3 DiffuseColor		= ComputeDiffuseAlbedo(TOON_BASECOLOR(BSDFContext.BSDF), TOON_METALLIC(BSDFContext.BSDF));
+			const float3 SpecularColor		= ComputeF0(TOON_SPECULAR(BSDFContext.BSDF), TOON_BASECOLOR(BSDFContext.BSDF), TOON_METALLIC(BSDFContext.BSDF));
 			const float  Roughness			= TOON_ROUGHNESS(BSDFContext.BSDF);
-			const float  Anisotropy			= TOON_ANISOTROPY(BSDFContext.BSDF);
 
 			const float ToonProfileId		= TOON_PROFILEID(BSDFContext.BSDF);
 
 			FToonProfileEvaluateSettings ToonEvalSettings = GetToonProfileEvaluateSettings(ToonProfileId);
 
 			const float2 OffsetFraction = float2(0.0f, 0.0f);
-			const float2 ScreenUV = (BSDFContext.PixelCoord - View.ViewRectMinAndSize.xy) * View.ViewSizeAndInvSize.zw;
-			float ToonProfileDiffuseDither  = EvaluateToonProfileDiffuseDither(ToonProfileId, ScreenUV, ToonEvalSettings);
-			float ToonProfileSpecularDither = EvaluateToonProfileSpecularDither(ToonProfileId, ScreenUV, ToonEvalSettings);
+			const float2 ToonPatternUV = TOON_PATTERNUVS(BSDFContext.BSDF);// (BSDFContext.PixelCoord - View.ViewRectMinAndSize.xy)* View.ViewSizeAndInvSize.zw;
+			float ToonProfileDiffuseDither  = EvaluateToonProfileDiffuseDither(ToonProfileId, ToonPatternUV, ToonEvalSettings);
+			float ToonProfileSpecularDither = EvaluateToonProfileSpecularDither(ToonProfileId, ToonPatternUV, ToonEvalSettings);
 
 			// TransmissionThickness in fact already stores transmittance.
 			float TransmittanceShadow = ShadowTerms.TransmittanceOrOpticalThickness;
@@ -1605,8 +1604,7 @@ FSubstrateImportanceSampleResult SubstrateImportanceSampleBSDF(FSubstrateBSDFCon
 
 		case SUBSTRATE_BSDF_TYPE_TOON:
 		{
-			const float3 DiffuseColor		= TOON_DIFFUSEALBEDO(BSDFContext.BSDF);
-			const float3 SpecularColor		= TOON_SPECULARCOLOR(BSDFContext.BSDF);
+			const float3 DiffuseColor		= TOON_BASECOLOR(BSDFContext.BSDF);
 
 			float4 ImportanceSample	= CosineSampleHemisphere(E);
 			Out.L					= mul(ImportanceSample.xyz, BSDFContext.TangentBasis);
@@ -2041,8 +2039,7 @@ FSubstrateEnvLightResult SubstrateEvaluateForEnvLight(FSubstrateBSDFContext BSDF
 
 		case SUBSTRATE_BSDF_TYPE_TOON:
 		{
-			const float3 DiffuseColor						= TOON_DIFFUSEALBEDO(BSDFContext.BSDF);
-			const float3 SpecularColor						= TOON_SPECULARCOLOR(BSDFContext.BSDF);
+			const float3 DiffuseColor						= ComputeDiffuseAlbedo(TOON_BASECOLOR(BSDFContext.BSDF), TOON_METALLIC(BSDFContext.BSDF));
 
 			SubstrateEnvLightResult.DiffuseNormal			= BSDFContext.N;
 			SubstrateEnvLightResult.bPostProcessSubsurface	= false;
diff --git a/Engine/Shaders/Private/Substrate/SubstrateExport.ush b/Engine/Shaders/Private/Substrate/SubstrateExport.ush
--- a/Engine/Shaders/Private/Substrate/SubstrateExport.ush
+++ b/Engine/Shaders/Private/Substrate/SubstrateExport.ush
@@ -880,8 +880,16 @@ void PackSubstrateOut(
 						break;
 						case SUBSTRATE_BSDF_TYPE_TOON:
 						{
-							SUBSTRATE_STORE_UINT1(PackColorLinearToGamma2AlphaLinear(float4(TOON_DIFFUSEALBEDO(BSDF), TOON_ROUGHNESS(BSDF))));
-							SUBSTRATE_STORE_UINT1(PackToonCustomData(TOON_SPECULARCOLOR(BSDF), TOON_ANISOTROPY(BSDF), TOON_PROFILEID(BSDF), Dither));
+							uint Data0 = 0;
+							Data0 |= PackR6G6B5Gamma2(TOON_BASECOLOR(BSDF), Dither);
+							Data0 |= PackR5(TOON_SPECULAR(BSDF),  Dither) << 17;
+							Data0 |= PackR5(TOON_METALLIC(BSDF),  Dither) << 22;
+							Data0 |= PackR5(TOON_ROUGHNESS(BSDF), Dither) << 27;
+							SUBSTRATE_STORE_UINT1(Data0);
+
+							uint Data1 = PackToonCustomData(TOON_PATTERNUVS(BSDF), TOON_PROFILEID(BSDF), Dither);
+							SUBSTRATE_STORE_UINT1(Data1);
+
 							SSSDataType = SSS_TYPE_TRANSMITTANCE_SHADOW;
 
 							const float ShadowMFPInMeters = GetToonProfileShadowSettings(TOON_PROFILEID(BSDF)).ShadowMFPInMeters;
@@ -1257,12 +1265,13 @@ FExportResult SubstrateMaterialExportOut(
 			{
 				Out.ShadingModelID = SHADINGMODELID_SUBSTRATE_TOON;
 
-				// DiffuseColor in BaseColor
-				// Roughness in Roughness
-				Out.Metallic = 0.0; // Cannot reuse metallic as it maybe written by decal
-				Out.Specular = 1.0;	// Cannot reuse specular as it can be used to apply decal on specular color
+				// Data stored using common Getxxx functions above:
+				//	- BaseColor
+				//	- Roughness
+				//	- Metallic
+				//	- Specular
 
-				uint PackedToonCustomData = PackToonCustomData(TOON_SPECULARCOLOR(CurrentBSDF), TOON_ANISOTROPY(CurrentBSDF), TOON_PROFILEID(CurrentBSDF), Dither);
+				uint PackedToonCustomData = PackToonCustomData(TOON_PATTERNUVS(CurrentBSDF), TOON_PROFILEID(CurrentBSDF), Dither);
 				Out.CustomData = UnpackRGBA8(PackedToonCustomData);
 			}
 			else
diff --git a/Engine/Shaders/Private/Substrate/SubstrateRead.ush b/Engine/Shaders/Private/Substrate/SubstrateRead.ush
--- a/Engine/Shaders/Private/Substrate/SubstrateRead.ush
+++ b/Engine/Shaders/Private/Substrate/SubstrateRead.ush
@@ -117,15 +117,15 @@ FSubstrateGBufferBSDF SubstrateReadGBufferBSDF(in FScreenSpaceData ScreenSpaceDa
 		else if (GBuffer.ShadingModelID == SHADINGMODELID_SUBSTRATE_TOON)
 		{
 			// special read from custom data
-			float3 SpecularColor= 0.0f;
-			float Anisotropy    = 0.0f;
-			float ToonProfileId	= 0.0f;
-			UnpackToonCustomData(PackRGBA8(GBuffer.CustomData), SpecularColor, Anisotropy, ToonProfileId);
+			float2 PatternUVs = 0.0f;
+			float ToonProfileId = 0.0f;
+			UnpackToonCustomData(PackRGBA8(GBuffer.CustomData), PatternUVs, ToonProfileId);
 
 			BSDF = GetSubstrateToonBSDF(
 				ToonProfileId,
-				GBuffer.BaseColor, SpecularColor, GBuffer.Roughness, Anisotropy,
+				GBuffer.BaseColor, GBuffer.Metallic, GBuffer.Specular, GBuffer.Roughness,
 				Emissive,
+				PatternUVs,
 				SharedLocalBasisIndex, SharedLocalBasisTypes).InlinedBSDF;
 		}
 		else
diff --git a/Engine/Shaders/Shared/SubstrateVisualizeDefinitions.h b/Engine/Shaders/Shared/SubstrateVisualizeDefinitions.h
--- a/Engine/Shaders/Shared/SubstrateVisualizeDefinitions.h
+++ b/Engine/Shaders/Shared/SubstrateVisualizeDefinitions.h
@@ -358,10 +358,9 @@ FSubstrateSubsurfaceHeader SubstrateLoadSubsurfaceHeader(inout FSubstrateAddress
 	{
 		Mode = SSS_TYPE_TRANSMITTANCE_SHADOW;
 
-		float3 SpecularColor = 0.0f;
-		float Anisotropy = 0.0f;
+		float2 PatternUVs = 0.0f;
 		float ToonProfileId = 0.0f;
-		UnpackToonCustomData(PackRGBA8(GBufferData.CustomData), SpecularColor, Anisotropy, ToonProfileId);
+		UnpackToonCustomData(PackRGBA8(GBufferData.CustomData), PatternUVs, ToonProfileId);
 
 		SubstrateSubSurfaceHeaderSetTransmittanceShadowMFPInMeters(Out, GetToonProfileShadowSettings(ToonProfileId).ShadowMFPInMeters.xxx);
 	}
@@ -691,15 +690,20 @@ FSubstratePixelDebugData ConvertToSerializableSubstratePixelDebugData(uint2 InCo
 
 		case SUBSTRATE_BSDF_TYPE_TOON:
 		{
-			Data.Closures[i].DiffuseR		= TOON_DIFFUSEALBEDO(BSDF).r;
-			Data.Closures[i].DiffuseG		= TOON_DIFFUSEALBEDO(BSDF).g;
-			Data.Closures[i].DiffuseB		= TOON_DIFFUSEALBEDO(BSDF).b;
-			Data.Closures[i].F0R			= TOON_SPECULARCOLOR(BSDF).r;
-			Data.Closures[i].F0G			= TOON_SPECULARCOLOR(BSDF).g;
-			Data.Closures[i].F0B			= TOON_SPECULARCOLOR(BSDF).b;
+			const float3 DiffuseColor = ComputeDiffuseAlbedo(TOON_BASECOLOR(BSDF), TOON_METALLIC(BSDF));
+			const float3 SpecularColor = ComputeF0(TOON_SPECULAR(BSDF), TOON_BASECOLOR(BSDF), TOON_METALLIC(BSDF));
+
+			Data.Closures[i].DiffuseR		= DiffuseColor.r;
+			Data.Closures[i].DiffuseG		= DiffuseColor.g;
+			Data.Closures[i].DiffuseB		= DiffuseColor.b;
+			Data.Closures[i].F0R			= SpecularColor.r;
+			Data.Closures[i].F0G			= SpecularColor.g;
+			Data.Closures[i].F0B			= SpecularColor.b;
 			Data.Closures[i].Roughness		= TOON_ROUGHNESS(BSDF);
-			Data.Closures[i].Anisotropy		= TOON_ANISOTROPY(BSDF);
 			Data.Closures[i].SSSPRofileID	= TOON_PROFILEID(BSDF);
+			
+			Data.Closures[i].GlintUVDDXx	= TOON_PATTERNUVS(BSDF).x;
+			Data.Closures[i].GlintUVDDXy	= TOON_PATTERNUVS(BSDF).y;
 
 			if (Data.Closures[i].SSSType == SSS_TYPE_TRANSMITTANCE_SHADOW)
 			{
diff --git a/Engine/Source/Runtime/Engine/Private/Materials/HLSLMaterialTranslator.cpp b/Engine/Source/Runtime/Engine/Private/Materials/HLSLMaterialTranslator.cpp
--- a/Engine/Source/Runtime/Engine/Private/Materials/HLSLMaterialTranslator.cpp
+++ b/Engine/Source/Runtime/Engine/Private/Materials/HLSLMaterialTranslator.cpp
@@ -13431,26 +13431,29 @@ int32 FHLSLMaterialTranslator::SubstrateSingleLayerWaterBSDF(
 
 int32 FHLSLMaterialTranslator::SubstrateToonBSDF(
 	int32 ToonProfileId,
-	int32 DiffuseColor, int32 SpecularColor, int32 Roughness, int32 Anisotropy,
+	int32 BaseColor, int32 Metallic, int32 Specular, int32 Roughness,
 	int32 EmissiveColor, int32 Normal,
+	int32 PatternUVs,
 	const FString& SharedLocalBasisIndexMacro, FSubstrateOperator* PromoteToOperator)
 {
+	if (ToonProfileId <= INDEX_NONE)
+	{
+		Errorf(TEXT("Substrate error: invalid Toon profile specified (asset: %s [%s]: ).\r\n"), *Material->GetDebugName(), *Material->GetAssetPath().ToString());
+		return INDEX_NONE;
+	}
+
 	if (PromoteToOperator)
 	{
-		if (ToonProfileId <= INDEX_NONE)
-		{
-			Errorf(TEXT("Substrate error: invalid Toon profile specified (asset: %s [%s]: ).\r\n"), *Material->GetDebugName(), *Material->GetAssetPath().ToString());
-			return INDEX_NONE;
-		}
 		return AddCodeChunk(
-			MCT_Substrate, TEXT("Parameters.%s.PromoteParameterBlendedBSDFToOperator(GetSubstrateToonBSDF(%s, %s, %s, %s, %s, %s, %s, Parameters.%s.Types), %u, %u, %u, %u) /* Normal:%s */"),
+			MCT_Substrate, TEXT("Parameters.%s.PromoteParameterBlendedBSDFToOperator(GetSubstrateToonBSDF(%s, %s, %s, %s, %s, %s, %s, %s, Parameters.%s.Types), %u, %u, %u, %u) /* Normal:%s */"),
 			*Substrate::GetParametersSubstrateTreeName(SubstrateTranslatorData.CurrentSubstrateCompilationContext),
 			*SubstrateGetCastParameterCode(ToonProfileId, MCT_Float),
-			*SubstrateGetCastParameterCode(DiffuseColor, MCT_Float3),
-			*SubstrateGetCastParameterCode(SpecularColor, MCT_Float3),
+			*SubstrateGetCastParameterCode(BaseColor, MCT_Float3),
+			*SubstrateGetCastParameterCode(Metallic, MCT_Float),
+			*SubstrateGetCastParameterCode(Specular, MCT_Float),
 			*SubstrateGetCastParameterCode(Roughness, MCT_Float),
-			*SubstrateGetCastParameterCode(Anisotropy, MCT_Float),
 			*SubstrateGetCastParameterCode(EmissiveColor, MCT_Float3),
+			*SubstrateGetCastParameterCode(PatternUVs, MCT_Float2),
 			*SharedLocalBasisIndexMacro,
 			*Substrate::GetParametersSharedLocalBasesName(SubstrateTranslatorData.CurrentSubstrateCompilationContext),
 			PromoteToOperator->Index,
@@ -13462,13 +13465,14 @@ int32 FHLSLMaterialTranslator::SubstrateToonBSDF(
 	}
 
 	return AddCodeChunk(
-		MCT_Substrate, TEXT("GetSubstrateToonBSDF(%s, %s, %s, %s, %s, %s, %s, Parameters.%s.Types)"),
+		MCT_Substrate, TEXT("GetSubstrateToonBSDF(%s, %s, %s, %s, %s, %s, %s, %s, Parameters.%s.Types)"),
 		*SubstrateGetCastParameterCode(ToonProfileId, MCT_Float),
-		*SubstrateGetCastParameterCode(DiffuseColor, MCT_Float3),
-		*SubstrateGetCastParameterCode(SpecularColor, MCT_Float3),
+		*SubstrateGetCastParameterCode(BaseColor, MCT_Float3),
+		*SubstrateGetCastParameterCode(Metallic, MCT_Float),
+		*SubstrateGetCastParameterCode(Specular, MCT_Float),
 		*SubstrateGetCastParameterCode(Roughness, MCT_Float),
-		*SubstrateGetCastParameterCode(Anisotropy, MCT_Float),
 		*SubstrateGetCastParameterCode(EmissiveColor, MCT_Float3),
+		*SubstrateGetCastParameterCode(PatternUVs, MCT_Float2),
 		*SharedLocalBasisIndexMacro,
 		*Substrate::GetParametersSharedLocalBasesName(SubstrateTranslatorData.CurrentSubstrateCompilationContext)
 	);
diff --git a/Engine/Source/Runtime/Engine/Private/Materials/HLSLMaterialTranslator.h b/Engine/Source/Runtime/Engine/Private/Materials/HLSLMaterialTranslator.h
--- a/Engine/Source/Runtime/Engine/Private/Materials/HLSLMaterialTranslator.h
+++ b/Engine/Source/Runtime/Engine/Private/Materials/HLSLMaterialTranslator.h
@@ -1288,8 +1288,9 @@ protected:
 		int32 ColorScaleBehindWater, int32 Normal, const FString& SharedLocalBasisIndexMacro, FSubstrateOperator* PromoteToOperator) override;
 	virtual int32 SubstrateToonBSDF(
 		int32 ToonProfileId,
-		int32 DiffuseColor, int32 SpecularColor, int32 Roughness, int32 Anisotropy,
-		int32 EmissiveColor, int32 Normal, 
+		int32 BaseColor, int32 Metallic, int32 Specular, int32 Roughness,
+		int32 EmissiveColor, int32 Normal,
+		int32 PatternUVs,
 		const FString& SharedLocalBasisIndexMacro, FSubstrateOperator* PromoteToOperator) override;
 	virtual int32 SubstrateHorizontalMixing(int32 Background, int32 Foreground, int32 Mix, int OperatorIndex, uint32 MaxDistanceFromLeaves) override;
 	virtual int32 SubstrateHorizontalMixingParameterBlending(int32 Background, int32 Foreground, int32 HorizontalMixCodeChunk, int32 NormalMixCodeChunk, const FString& SharedLocalBasisIndexMacro, int32 BackgroundBSDFNormalCodeChunk, int32 ForegroundBSDFNormalCodeChunk, FSubstrateOperator* PromoteToOperator) override;
diff --git a/Engine/Source/Runtime/Engine/Private/Materials/MaterialExpressionSubstrate.cpp b/Engine/Source/Runtime/Engine/Private/Materials/MaterialExpressionSubstrate.cpp
--- a/Engine/Source/Runtime/Engine/Private/Materials/MaterialExpressionSubstrate.cpp
+++ b/Engine/Source/Runtime/Engine/Private/Materials/MaterialExpressionSubstrate.cpp
@@ -2512,12 +2512,13 @@ int32 UMaterialExpressionSubstrateToonBSDF::Compile(class FMaterialCompiler* Com
 
 	int32 OutputCodeChunk = Compiler->SubstrateToonBSDF(
 		ToonProfileIdCodeChunk,
-		CompileWithDefaultFloat3(Compiler, DiffuseColor, 0.0f, 0.0f, 0.0f),
-		CompileWithDefaultFloat3(Compiler, SpecularColor, 0.0f, 0.0f, 0.0f),
+		CompileWithDefaultFloat3(Compiler, BaseColor, 0.18f, 0.18f, 0.18f),
+		CompileWithDefaultFloat1(Compiler, Metallic, 0.0f),
+		CompileWithDefaultFloat1(Compiler, Specular, 0.5f),
 		CompileWithDefaultFloat1(Compiler, Roughness, 0.5f),
-		CompileWithDefaultFloat1(Compiler, Anisotropy, 0.0f),
 		CompileWithDefaultFloat3(Compiler, EmissiveColor, 0.0f, 0.0f, 0.0f),
 		NormalCodeChunk,
+		CompileWithDefaultFloat2(Compiler, PatternUVs, 0.0f, 0.0f),
 		Compiler->GetSubstrateSharedLocalBasisIndexMacro(NewRegisteredSharedLocalBasis),
 		!SubstrateOperator.bUseParameterBlending || (SubstrateOperator.bUseParameterBlending && SubstrateOperator.bRootOfParameterBlendingSubTree) ? &SubstrateOperator : nullptr);
 
@@ -2541,15 +2542,17 @@ EMaterialValueType UMaterialExpressionSubstrateToonBSDF::GetInputValueType(int32
 	case 0:
 		return MCT_Float3; // BaseColor
 	case 1:
-		return MCT_Float3; // SpecularColor
+		return MCT_Float1; // Metallic
 	case 2:
-		return MCT_Float1; // Roughness
+		return MCT_Float1; // Specular
 	case 3:
-		return MCT_Float1; // Anisotropy
+		return MCT_Float1; // Roughness
 	case 4:
 		return MCT_Float3; // Normal
 	case 5:
 		return MCT_Float3; // Emissive Color
+	case 6:
+		return MCT_Float2; // Pattern UVs
 	}
 
 	check(false);
@@ -2580,10 +2583,12 @@ FSubstrateOperator* UMaterialExpressionSubstrateToonBSDF::SubstrateGenerateMater
 	SubstrateOperator.ThicknessIndex = SubstrateTranslatorData.SubstrateThicknessStackGetThicknessIndex();
 	SubstrateOperator.bBSDFWritesEmissive = EmissiveColor.IsConnected();
 	SubstrateOperator.BSDFFeatures |= ESubstrateBsdfFeature::Toon;
-	if (Anisotropy.IsConnected())
-	{
-		SubstrateOperator.BSDFFeatures |= ESubstrateBsdfFeature::Anisotropy;
-	}
+
+//	if ( SUBSTRATE_TODO fetch data from the profile to set anisotopy if value is !0.
+//	{
+//		SubstrateOperator.BSDFFeatures |= ESubstrateBsdfFeature::Anisotropy;
+//	}
+
 	return &SubstrateOperator;
 }
 #endif // WITH_EDITOR
diff --git a/Engine/Source/Runtime/Engine/Public/MaterialCompiler.h b/Engine/Source/Runtime/Engine/Public/MaterialCompiler.h
--- a/Engine/Source/Runtime/Engine/Public/MaterialCompiler.h
+++ b/Engine/Source/Runtime/Engine/Public/MaterialCompiler.h
@@ -753,8 +753,9 @@ public:
 		int32 WaterAlbedo, int32 WaterExtinction, int32 WaterPhaseG, int32 ColorScaleBehindWater, int32 Normal, const FString& SharedLocalBasisIndexMacro, FSubstrateOperator* PromoteToOperator) = 0;
 	virtual int32 SubstrateToonBSDF(
 		int32 ToonProfileId,
-		int32 DiffuseColor, int32 SpecularColor, int32 Roughness, int32 Anisotropy,
+		int32 BaseColor, int32 Metallic, int32 Specular, int32 Roughness,
 		int32 EmissiveColor, int32 Normal,
+		int32 PatternUVs,
 		const FString& SharedLocalBasisIndexMacro, FSubstrateOperator* PromoteToOperator) = 0;
 	virtual int32 SubstrateHorizontalMixing(int32 Background, int32 Foreground, int32 Mix, int OperatorIndex, uint32 MaxDistanceFromLeaves) = 0;
 	virtual int32 SubstrateHorizontalMixingParameterBlending(int32 Background, int32 Foreground, int32 HorizontalMixCodeChunk, int32 NormalMixCodeChunk, const FString& SharedLocalBasisIndexMacro, int32 BackgroundBSDFNormalCodeChunk, int32 ForegroundBSDFNormalCodeChunk, FSubstrateOperator* PromoteToOperator) = 0;
@@ -1557,14 +1558,16 @@ public:
 
 	virtual int32 SubstrateToonBSDF(
 		int32 ToonProfileId,
-		int32 DiffuseColor, int32 SpecularColor, int32 Roughness, int32 Anisotropy,
+		int32 BaseColor, int32 Metallic, int32 Specular, int32 Roughness,
 		int32 EmissiveColor, int32 Normal,
+		int32 PatternUVs,
 		const FString& SharedLocalBasisIndexMacro, FSubstrateOperator* PromoteToOperator) override
 	{
 		return Compiler->SubstrateToonBSDF(
 			ToonProfileId,
-			DiffuseColor, SpecularColor, Roughness, Anisotropy, 
+			BaseColor, Metallic, Specular, Roughness,
 			EmissiveColor, Normal,
+			PatternUVs,
 			SharedLocalBasisIndexMacro, PromoteToOperator);
 	}
 
diff --git a/Engine/Source/Runtime/Engine/Public/Materials/MaterialExpressionSubstrate.h b/Engine/Source/Runtime/Engine/Public/Materials/MaterialExpressionSubstrate.h
--- a/Engine/Source/Runtime/Engine/Public/Materials/MaterialExpressionSubstrate.h
+++ b/Engine/Source/Runtime/Engine/Public/Materials/MaterialExpressionSubstrate.h
@@ -871,28 +871,28 @@ class UMaterialExpressionSubstrateToonBSDF : public UMaterialExpressionSubstrate
 	GENERATED_UCLASS_BODY()
 
 	/**
-	 * Surface base color. (type = float3, unit = unitless, defaults to white)
+	 * Surface base color. (type = float3, unit = unitless, defaults to 0.18 grey)
 	 */
 	UPROPERTY()
-	FExpressionInput DiffuseColor;
+	FExpressionInput BaseColor;
 
 	/**
-	 * Surface specular color. (type = float3, unit = unitless, defaults to white)
+	 * Surface metallic value to transition from a dielectric to a conductor. (type = float, unit = unitless, defaults to 0.0)
 	 */
 	UPROPERTY()
-	FExpressionInput SpecularColor;
+	FExpressionInput Metallic;
 
 	/**
-	 * Controls how rough the Material is. Roughness of 0 (smooth) is a mirror reflection and 1 (rough) is completely matte or diffuse. When using anisotropy, it is the roughness used along the Tangent axis. (type = float, unit = unitless, defaults to 0.5)
+	 * Surface specular value. (type = float, unit = unitless, defaults to 0.5)
 	 */
 	UPROPERTY()
-	FExpressionInput Roughness;
+	FExpressionInput Specular;
 
 	/**
-	 * Controls the anisotropy factor of the roughness. Positive value elongates the specular lobe along the Tangent vector, Negative value elongates the specular lobe along the perpendicular of the Tangent. (type = float, unit = unitless).
+	 * Controls how rough the Material is. Roughness of 0 (smooth) is a mirror reflection and 1 (rough) is completely matte or diffuse. When using anisotropy, it is the roughness used along the Tangent axis. (type = float, unit = unitless, defaults to 0.5)
 	 */
 	UPROPERTY()
-	FExpressionInput Anisotropy;
+	FExpressionInput Roughness;
 
 	/**
 	 * The normal of the surface (type = float3, unit = unitless, defaults to +Z vector)
@@ -906,6 +906,12 @@ class UMaterialExpressionSubstrateToonBSDF : public UMaterialExpressionSubstrate
 	UPROPERTY()
 	FExpressionInput EmissiveColor;
 
+	/**
+	 * UVs used to render the diffuse, specular or any other artistic pattern. It can only store data in [0,1].
+	 */
+	UPROPERTY()
+	FExpressionInput PatternUVs;
+
 	/** Toon profile used to control the appearance of this BSDF. */
 	UPROPERTY(EditAnywhere, BlueprintReadOnly, Category = Material, meta = (DisplayName = "Toon Profile"))
 	TObjectPtr<class UToonProfile> ToonProfile;
diff --git a/Engine/Source/Runtime/Renderer/Private/Substrate/SubstrateVisualize.cpp b/Engine/Source/Runtime/Renderer/Private/Substrate/SubstrateVisualize.cpp
--- a/Engine/Source/Runtime/Renderer/Private/Substrate/SubstrateVisualize.cpp
+++ b/Engine/Source/Runtime/Renderer/Private/Substrate/SubstrateVisualize.cpp
@@ -884,8 +884,8 @@ void AddProcessAndPrintSubstrateMaterialPropertiesPasses(FRDGBuilder& GraphBuild
 						PrintFSmall(DrawPosX + DataOffset, DrawPosY, Closure.Roughness);
 						NewLine();
 
-						PrintS(DrawPosX + 0.0f, DrawPosY, TEXT("Anisotropy"), Orange);
-						PrintFSmall(DrawPosX + DataOffset, DrawPosY, Closure.Anisotropy);
+						PrintS(DrawPosX + 0.0f, DrawPosY, TEXT("Pattern UVs"), Orange);
+						PrintF2DAdapt(DrawPosX + DataOffset, DrawPosY, Data.Closures[i].GlintUVDDXx, Data.Closures[i].GlintUVDDXy);
 						NewLine();
 
 						if (Closure.SSSType == SSS_TYPE_TRANSMITTANCE_SHADOW)
