PR #14339: Fix FastLog_N for correct logarithm calculation

#Synced-CL 50298075

diff --git a/Engine/Shaders/Private/FastMath.ush b/Engine/Shaders/Private/FastMath.ush
--- a/Engine/Shaders/Private/FastMath.ush
+++ b/Engine/Shaders/Private/FastMath.ush
@@ -16,8 +16,8 @@ FastExp_N(float3)
 FastExp_N(float4)
 
 // Reference: http://www.humus.name/Articles/Persson_LowLevelThinking.pdf p.27
-// Warning: This is a coarse approximation
-#define FastLog_N(T) T FastLog(T x) { return log2(0.693147f * x); }
+// Corrected version. The factor is incorrectly placed in the reference.
+#define FastLog_N(T) T FastLog(T x) { return 0.693147f*log2(x); }
 FastLog_N(float)
 FastLog_N(float2)
 FastLog_N(float3)
