PR #14337: Fix: Corrected inconsistent bitmask in ConvertTangentUnormToSnorm16

#Synced-CL 50022451

diff --git a/Engine/Shaders/Private/Common.ush b/Engine/Shaders/Private/Common.ush
--- a/Engine/Shaders/Private/Common.ush
+++ b/Engine/Shaders/Private/Common.ush
@@ -997,20 +997,20 @@ float2 ConvertTangentUnormToSnorm16(float2 Input)
 {
 	int2 IntVal = int2(round(Input * 65535.0f));
 	//negate
-	IntVal = select(IntVal > 32767, IntVal | 0xFFFFFF80, IntVal);
+	IntVal = select(IntVal > 32767, IntVal | 0xFFFF8000, IntVal);
 	return clamp(IntVal / 32767.0f, -1, 1);
 }
 float3 ConvertTangentUnormToSnorm16(float3 Input)
 {
 	int3 IntVal = int3(round(Input * 65535.0f));
-	IntVal = select(IntVal > 32767, IntVal | 0xFFFFFF80, IntVal);
+	IntVal = select(IntVal > 32767, IntVal | 0xFFFF8000, IntVal);
 	return clamp(IntVal / 32767.0f, -1, 1);
 }
 float4 ConvertTangentUnormToSnorm16(float4 Input)
 {
 	int4 IntVal = int4(round(Input * 65535.0f));
 	//negate
-	IntVal = select(IntVal > 32767, IntVal | 0xFFFFFF80, IntVal);
+	IntVal = select(IntVal > 32767, IntVal | 0xFFFF8000, IntVal);
 	return clamp(IntVal / 32767.0f, -1, 1);
 }
 float ConvertTangentSnormToUnorm8(float Input)
