Path Tracing: Specialize raygen kernels for camera rays

Reduce the overhead of the SER version  by skipping the reorder call completely instead of using a runtime branch. This also allows removing a few other elements from the main trace loop and simplify the logic for secondary bounces.

Additionally, break up the visualization kernel (PrimaryRays mode) into per-bounce dispatches, like the reference algorithm. This significantly reduces the cost of using multiple bounces.

Remove the experimental path guiding code for now, it impacts the path state size which should be benchmarked individually.

Change defaults for CompactionDepth to 6 - this allows a tail kernel to handle the small number of paths that reach high depths and improves performance on complex scenes.

Change defaults for the visualization mode to use the bounce count from PPV and include emissive objects by default by aligning with the main cvars, and remove the visualization specific cvars

Add the missing light source visibility logic to the visualization mode for parity and to allow skylights to be visualized properly.

#Synced-CL 50321622

diff --git a/Engine/Shaders/Private/PathTracing/PathTracing.usf b/Engine/Shaders/Private/PathTracing/PathTracing.usf
--- a/Engine/Shaders/Private/PathTracing/PathTracing.usf
+++ b/Engine/Shaders/Private/PathTracing/PathTracing.usf
@@ -36,16 +36,19 @@ int2 ComputeTextureIndex(uint2 DispatchIdx)
 	return DispatchIdx + TileTextureOffset;
 }
 
+#if PATH_TRACER_KERNEL_TYPE != PATH_TRACER_KERNEL_CAMERA
 int FirstBounce;
+#endif
 
 RWStructuredBuffer<FPathTracingPackedPathState> PathStateData;
 RWBuffer<uint> ActivePaths;
 RWBuffer<uint> NumPathStates;
 
-#if PATH_TRACER_USE_COMPACTION
+#if PATH_TRACER_KERNEL_TYPE != PATH_TRACER_KERNEL_TAIL
 RWBuffer<uint> NextActivePaths;
 #endif
 
+#if PATH_TRACER_KERNEL_TYPE != PATH_TRACER_KERNEL_CAMERA
 FPathState LoadPathStateData(uint Index)
 {
 	FPathTracingPackedPathState Data = PathStateData[Index];
@@ -64,8 +67,9 @@ FPathState LoadPathStateData(uint Index)
 	Output.PackedScatterPhase = Data.PackedScatterPhase;
 	return Output;
 }
+#endif
 
-#if PATH_TRACER_USE_COMPACTION
+#if PATH_TRACER_KERNEL_TYPE != PATH_TRACER_KERNEL_TAIL
 void StorePathStateData(FPathState PathState, uint Index)
 {
 	FPathTracingPackedPathState Data;
@@ -89,56 +93,59 @@ RAY_TRACING_ENTRY_RAYGEN(PathTracingMainRG)
 	const uint2 DispatchDim = DispatchRaysDimensions().xy;
 
 	uint LinearPixelIndex = DispatchIdx.x + DispatchDim.x * DispatchIdx.y;
-	uint PathIndex = DispatchIdx.x + 65536 * DispatchIdx.y;
-#if PATH_TRACER_USE_COMPACTION == 1
+
+#if PATH_TRACER_KERNEL_TYPE == PATH_TRACER_KERNEL_CAMERA
+	const int FirstBounce = 0;
+#endif
+
+#if PATH_TRACER_KERNEL_TYPE != PATH_TRACER_KERNEL_TAIL
 	if (LinearPixelIndex == 0)
 	{
 		// write other dimensions for indirect dispatch of the next bounce
-		NumPathStates[3 * FirstBounce + 4] = 1;
-		NumPathStates[3 * FirstBounce + 5] = 1;
+		NumPathStates[DISPATCH_INDIRECT_UINT_COUNT * (FirstBounce + 1) + 1] = 1;
+		NumPathStates[DISPATCH_INDIRECT_UINT_COUNT * (FirstBounce + 1) + 2] = 1;
 	}
 #endif
-	FPathState PathState;
-	if (FirstBounce == 0)
-	{
+
+#if PATH_TRACER_KERNEL_TYPE == PATH_TRACER_KERNEL_CAMERA
+	uint PathIndex = DispatchIdx.x + 65536 * DispatchIdx.y;
 #if PATH_TRACER_USE_ADAPTIVE_SAMPLING
-		if (Iteration > 0)
-		{
-			if (LinearPixelIndex >= NumPathStates[0])
-			{
-				return; // nothing left to do on this thread
-			}
-			PathIndex = ActivePaths[LinearPixelIndex];
-			DispatchIdx = uint2(PathIndex & 65535, PathIndex >> 16);
-			LinearPixelIndex = DispatchIdx.x + ScanlineWidth * DispatchIdx.y;
-		}
-#endif
-		PathState = CreatePathState(ComputePixelIndex(DispatchIdx), 
-									ComputeTextureIndex(DispatchIdx));
-	}
-	else
+	if (Iteration > 0)
 	{
-		// compare to number of active paths from the previous bounce
-		if (LinearPixelIndex >= NumPathStates[3 * FirstBounce + 0])
+		if (LinearPixelIndex >= NumPathStates[0])
 		{
 			return; // nothing left to do on this thread
 		}
 		PathIndex = ActivePaths[LinearPixelIndex];
-		PathState = LoadPathStateData((PathIndex & 65535) + ScanlineWidth * (PathIndex >> 16));
+		DispatchIdx = uint2(PathIndex & 65535, PathIndex >> 16);
+		LinearPixelIndex = DispatchIdx.x + ScanlineWidth * DispatchIdx.y;
+	}
+#endif
+	FPathState PathState = CreatePathState(ComputePixelIndex(DispatchIdx),  ComputeTextureIndex(DispatchIdx));
+#else // PATH_TRACER_KERNEL_TYPE == PATH_TRACER_KERNEL_INNER
+	// compare to number of active paths from the previous bounce
+	if (LinearPixelIndex >= NumPathStates[DISPATCH_INDIRECT_UINT_COUNT * FirstBounce + 0])
+	{
+		return; // nothing left to do on this thread
 	}
-#if PATH_TRACER_USE_COMPACTION == 1
+	const uint PathIndex = ActivePaths[LinearPixelIndex];
+	FPathState PathState = LoadPathStateData((PathIndex & 65535) + ScanlineWidth * (PathIndex >> 16));
+#endif
+
+#if PATH_TRACER_KERNEL_TYPE != PATH_TRACER_KERNEL_TAIL
 	const bool KeepGoing = PathTracingKernel(PathState, FirstBounce);
 	const uint2 TextureIndex = ComputeTextureIndex(uint2(PathIndex & 65535, PathIndex >> 16));
-	if (FirstBounce == 0)
+#if PATH_TRACER_KERNEL_TYPE == PATH_TRACER_KERNEL_CAMERA
 	{
 		// Camera ray must write depth right away because we don't want to have to carry DepthZ in the PackedPathState
 		PathState.WriteDepth(TextureIndex);
 	}
+#endif
 	if (KeepGoing)
 	{
 	    // Add the path to the set of valid paths for the next bounce
 		uint PathStateIndex;
-		InterlockedAdd(NumPathStates[3 * FirstBounce + 3], 1, PathStateIndex);
+		InterlockedAdd(NumPathStates[DISPATCH_INDIRECT_UINT_COUNT * (FirstBounce + 1)], 1, PathStateIndex);
 		NextActivePaths[PathStateIndex] = PathIndex;
 		StorePathStateData(PathState, (PathIndex & 65535) + ScanlineWidth * (PathIndex >> 16));
 	}
@@ -148,7 +155,7 @@ RAY_TRACING_ENTRY_RAYGEN(PathTracingMainRG)
 		// Accumulate radiance and update pixel variance
 		PathState.WritePixel(TextureIndex);
 	}
-#else // PATH_TRACER_USE_COMPACTION == 0
+#else // PATH_TRACER_KERNEL_TYPE == PATH_TRACER_KERNEL_TAIL
 	for (int Bounce = FirstBounce; Bounce <= MaxBounces; Bounce++)
 	{
 		if (!PathTracingKernel(PathState, Bounce))
@@ -161,9 +168,5 @@ RAY_TRACING_ENTRY_RAYGEN(PathTracingMainRG)
 	// Accumulate radiance and update pixel variance
 	const uint2 TextureIndex = ComputeTextureIndex(uint2(PathIndex & 65535, PathIndex >> 16));
 	PathState.WritePixel(TextureIndex);
-	if (FirstBounce == 0)
-	{
-		PathState.WriteDepth(TextureIndex);
-	}
 #endif
 }
diff --git a/Engine/Shaders/Private/PathTracing/PathTracingCommon.ush b/Engine/Shaders/Private/PathTracing/PathTracingCommon.ush
--- a/Engine/Shaders/Private/PathTracing/PathTracingCommon.ush
+++ b/Engine/Shaders/Private/PathTracing/PathTracingCommon.ush
@@ -50,6 +50,10 @@ HLSL_STATIC_ASSERT(SHADINGMODELID_SOLID_GLASS <= SHADINGMODELID_MASK, "");
 
 #define PATH_TRACING_ABSORPTION_SCALE  0.01        // inverse distance at which we hit the transmission color (Beer's law)
 
+// Ignore hair strands if the ray origin is beyond this radius. This is meant to workaround a performance cliff that can occur on some hardware.
+// In ordinary scenes, the most likely occurance of this problem is when combining hair and volumetric atmosphere as rays from several Km can travel back towards the hair and cause long delays or even TDRs
+#define PATH_TRACING_SKIP_HAIR_DISTANCE      1e6
+
 #define SUBSTRATE_PATH_TRACING_SPECULAR_PROFILE	(SUBSTRATE_ENABLED && PLATFORM_ENABLES_SUBSTRATE_SPECULAR_PROFILE)
 #define SUBSTRATE_PATH_TRACING_GLINTS 			(SUBSTRATE_ENABLED && PLATFORM_ENABLES_SUBSTRATE_GLINTS)
 
diff --git a/Engine/Shaders/Private/PathTracing/PathTracingCore.ush b/Engine/Shaders/Private/PathTracing/PathTracingCore.ush
--- a/Engine/Shaders/Private/PathTracing/PathTracingCore.ush
+++ b/Engine/Shaders/Private/PathTracing/PathTracingCore.ush
@@ -5,15 +5,9 @@
 #define RANDSEQ_VERSION 1
 #define RANDSEQ_SFC_TEXTURE 1
 
-#define RANDSEQ_UNROLL_SOBOL PATH_TRACER_USE_COMPACTION // unrolling sobol sampler is best when we can guarantee high occupancy
-
 #define USE_PATH_TRACING_LIGHT_GRID	1
 #define USE_RAY_TRACING_DECAL_GRID	1
 
-// Ignore hair strands if the ray origin is beyond this radius. This is meant to workaround a performance cliff that can occur on some hardware.
-// In ordinary scenes, the most likely occurance of this problem is when combining hair and volumetric atmosphere as rays from several Km can travel back towards the hair and cause long delays or even TDRs
-#define PATH_TRACING_SKIP_HAIR_DISTANCE      1e6
-
 // Debug printing is enabled from config file with r.PathTracing.Debug
 #ifndef DEBUG_ENABLE
 #define DEBUG_ENABLE 0
@@ -115,9 +109,13 @@ struct FPathState
 	                                // 2 bytes // phase g (could be quantized further)
 	                                // 80 bytes total (see FPathTracingPackedPathState)
 
-	// Temporary parameters not stored				
-	float DepthZ;                   // Camera ray depth
+	// Temporary parameters not stored
+#if PATH_TRACER_KERNEL_TYPE == PATH_TRACER_KERNEL_CAMERA
+	float DepthZ;                   // Camera ray depth (don't need it otherwise)
+#endif
+#if DEBUG_ENABLE && PATH_TRACER_KERNEL_TYPE == PATH_TRACER_KERNEL_CAMERA
 	uint2 PixelCoord;               // Pixel index/coord
+#endif
 
 	float3 GetAlbedo()
 	{
@@ -313,20 +311,27 @@ struct FPathState
 		NormalTexture[TextureIndex].xyz = lerp(OldNormal.xyz,Normal.xyz, Blend);
 	}
 
+#if PATH_TRACER_KERNEL_TYPE == PATH_TRACER_KERNEL_CAMERA
 	void WriteDepth(uint2 TextureIndex)
 	{
 		float OldDepth = Iteration > 0 ? DepthTexture[TextureIndex] : 0;
 		DepthTexture[TextureIndex] = max(OldDepth, ConvertToDeviceZ(DepthZ));
 	}
+#endif
+
 };
 
-void AccumulateRadiance(inout float3 TotalRadiance, float3 PathRadiance, bool bIsCameraRay)
+void AccumulateRadiance(inout float3 TotalRadiance, float3 PathRadiance)
 {
 	PathRadiance *= BaseExposure;
-	// User asked for path contributions to be clamped to reduce fireflies.
-	// This puts an upper bound on variance within the pixel at the expense of bias
 	float MaxPathRadiance = max3(PathRadiance.x, PathRadiance.y, PathRadiance.z);
-	float MaxIntensity = bIsCameraRay ? MaxHalfFloat : MaxPathIntensity;
+#if PATH_TRACER_KERNEL_TYPE == PATH_TRACER_KERNEL_CAMERA
+	float MaxIntensity = MaxHalfFloat; // camera rays, clamp at a fairly high level for half-float final storage into SceneColor buffer
+#else
+	// User asked for path contributions to be clamped to reduce fireflies from indirect paths
+	// This puts an upper bound on variance within the pixel at the expense of bias
+	float MaxIntensity = MaxPathIntensity;
+#endif
 	if (MaxPathRadiance > MaxIntensity)
 	{
 		// adjust sample while keeping color
@@ -468,9 +473,13 @@ void EvaluateDecals(FPathState PathState, float HitT, inout FDecalShaderPayload
 
 FPackedPathTracingPayload TraceTransparentRay(inout FPathState PathState, int Bounce, inout FVolumeSegment VolumeSegment)
 {
-	const bool bIsCameraRay = Bounce == 0;
 	const bool bLastBounce = Bounce == MaxBounces;
-	const uint RayFlags = bIsCameraRay ? RAY_FLAG_CULL_BACK_FACING_TRIANGLES : 0;
+#if PATH_TRACER_KERNEL_TYPE == PATH_TRACER_KERNEL_CAMERA
+	// backface culling on camera rays only to match raster but still allow refraction and indirect to work properly
+	const uint RayFlags = RAY_FLAG_CULL_BACK_FACING_TRIANGLES;
+#else
+	const uint RayFlags = 0;
+#endif
 	const bool bUseCameraRayType = !PathState.HasMadeContributionScatter();
 	uint InstanceInclusionMask = bUseCameraRayType ? PATHTRACER_MASK_CAMERA : PATHTRACER_MASK_INDIRECT;
 	if (length2(PathState.Ray.Origin) < Pow2(PATH_TRACING_SKIP_HAIR_DISTANCE))
@@ -512,10 +521,9 @@ FPackedPathTracingPayload TraceTransparentRay(inout FPathState PathState, int Bo
 			PathState.Ray.GetNativeDesc(),
 			OpaquePayload,
 			Hit);
-		if (!bIsCameraRay)
-		{
-			NvReorderThread(Hit);
-		}
+	#if PATH_TRACER_KERNEL_TYPE != PATH_TRACER_KERNEL_CAMERA
+		NvReorderThread(Hit);
+	#endif
 		NvInvokeHitObject(TLAS, Hit, OpaquePayload);
 	}
 #else
@@ -605,10 +613,9 @@ FPackedPathTracingPayload TraceTransparentRay(inout FPathState PathState, int Bo
 				PathState.Ray.GetNativeDesc(),
 				PackedPayload,
 				Hit);
-			if (!bIsCameraRay)
-			{
-				NvReorderThread(Hit);
-			}
+		#if PATH_TRACER_KERNEL_TYPE != PATH_TRACER_KERNEL_CAMERA
+			NvReorderThread(Hit);
+		#endif
 			NvInvokeHitObject(TLAS, Hit, PackedPayload);
 		}
 #else
@@ -637,15 +644,16 @@ FPackedPathTracingPayload TraceTransparentRay(inout FPathState PathState, int Bo
 		{
 			// we didn't hit any real geometry, but we did hit the volume's blocker geometry
 			// create a virtual hit with it here
-			PackedPayload = VolumeGetBlockerHit(PathState.Ray.Origin, PathState.Ray.Direction, VolumeIsectList.BlockerHitT, bIsCameraRay);
+			PackedPayload = VolumeGetBlockerHit(PathState.Ray.Origin, PathState.Ray.Direction, VolumeIsectList.BlockerHitT, PATH_TRACER_KERNEL_TYPE == PATH_TRACER_KERNEL_CAMERA);
 		}
 
 		// Loop over lights to capture their contribution
 		// #dxr_todo: if we have lots of lights, having some hierarchical structure would be better ....
 		// number of directly visible lights for the first bounce
-		// after the first bounce, we don't need to include lights in the trace call
+#if PATH_TRACER_KERNEL_TYPE == PATH_TRACER_KERNEL_CAMERA
+		// Run this for camera rays only, after the first bounce, we don't need to include lights in the trace call
 		// because NEE handled it for us
-		for (uint LightId = 0, NumVisibleLights = bIsCameraRay ? SceneVisibleLightCount : 0; LightId < NumVisibleLights; ++LightId)
+		for (uint LightId = 0, NumVisibleLights = SceneVisibleLightCount; LightId < NumVisibleLights; ++LightId)
 		{
 			FRayDesc LightRay = PathState.Ray;
 			LightRay.TMax = PackedPayload.IsMiss() ? PathState.Ray.TMax : PackedPayload.HitT;
@@ -672,6 +680,7 @@ FPackedPathTracingPayload TraceTransparentRay(inout FPathState PathState, int Bo
 				PackedPayload = PackPathTracingPayload(LightHitPayload);
 			}
 		}
+#endif
 
 		if (any(PathState.GetSigmaT() > 0))
 		{
@@ -792,11 +801,10 @@ FPackedPathTracingPayload TraceTransparentRay(inout FPathState PathState, int Bo
 					PathState.GetNormal()
 				);
 				HitSample = CloudPayload.GetVolumetricCallableShaderRISContext();
-				AccumulateRadiance(PathState.Radiance, CloudPayload.GetVolumetricCallableShaderOutputRadiance(), bIsCameraRay);
-				if (bIsCameraRay)
-				{
-					PathState.Alpha += CloudPayload.GetVolumetricCallableShaderOutputAlpha();
-				}
+				AccumulateRadiance(PathState.Radiance, CloudPayload.GetVolumetricCallableShaderOutputRadiance());
+			#if PATH_TRACER_KERNEL_TYPE == PATH_TRACER_KERNEL_CAMERA
+				PathState.Alpha += CloudPayload.GetVolumetricCallableShaderOutputAlpha();
+			#endif
 				if (CloudPayload.HitT > 0)
 				{
 					// If we registered a hit, transfer it to the currently sampled payload
@@ -833,7 +841,7 @@ FPackedPathTracingPayload TraceTransparentRay(inout FPathState PathState, int Bo
 				//  - we can't use analytic transmittance (need ray marching for throughput or emission)
 				//  - we are a camera ray and some volumes are flagged as holdouts
 				//  - this is not the last bounce, so we need to select a point for scattering
-				if (!bUseAnalyticTransmittance || (bIsCameraRay && bHaveHoldouts) || !bLastBounce)
+				if (!bUseAnalyticTransmittance || (PATH_TRACER_KERNEL_TYPE == PATH_TRACER_KERNEL_CAMERA && bHaveHoldouts) || !bLastBounce)
 				{
 					// remember the start of our interval before ray marching messes with it
 					float StartVolumeTMin = Interval.VolumeTMin;
@@ -845,7 +853,7 @@ FPackedPathTracingPayload TraceTransparentRay(inout FPathState PathState, int Bo
 					// TODO: This biases the result! Is there a better way to limit the number of steps?
 					for (int Step = 0; Step < MaxRaymarchSteps; Step++)
 					{
-						FVolumeTrackingResult TrackingResult = VolumeSampleDistance(PathThroughput, PathState.Ray.Origin, PathState.Ray.Direction, Interval, bIsCameraRay, PathState.RandSequence);
+						FVolumeTrackingResult TrackingResult = VolumeSampleDistance(PathThroughput, PathState.Ray.Origin, PathState.Ray.Direction, Interval, PATH_TRACER_KERNEL_TYPE == PATH_TRACER_KERNEL_CAMERA, PathState.RandSequence);
 						if (TrackingResult.Distance < 0)
 						{
 							break;
@@ -860,17 +868,16 @@ FPackedPathTracingPayload TraceTransparentRay(inout FPathState PathState, int Bo
 							float3 Ro = PathState.Ray.Origin;
 							float3 Rd = PathState.Ray.Direction;
 							float3 TranslatedWorldPos = Ro + Interval.VolumeTMin * Rd;
-							FVolumeShadedResult Result = VolumeGetDensity(TranslatedWorldPos, Interval.VolumeMask, bIsCameraRay);
+							FVolumeShadedResult Result = VolumeGetDensity(TranslatedWorldPos, Interval.VolumeMask, PATH_TRACER_KERNEL_TYPE == PATH_TRACER_KERNEL_CAMERA);
 
 							// clamp to make sure we never exceed the majorant (should not be the case, but need to avoid any possible numerical issues)
 							float3 SigmaT = min(Result.SigmaT, TrackingResult.SigmaBar);
 							float3 SigmaN = TrackingResult.SigmaBar - SigmaT;
 
-							if (bIsCameraRay)
-							{
-								float3 SigmaH = min(Result.SigmaH, SigmaT);
-								PathState.Alpha += Luminance(PathThroughput * (SigmaT - SigmaH));
-							}
+						#if PATH_TRACER_KERNEL_TYPE == PATH_TRACER_KERNEL_CAMERA
+							float3 SigmaH = min(Result.SigmaH, SigmaT);
+							PathState.Alpha += Luminance(PathThroughput * (SigmaT - SigmaH));
+						#endif
 
 							// if this is not the last bounce, consider the possibility of scattering to give us medium hits
 							if (!bLastBounce)
@@ -908,7 +915,7 @@ FPackedPathTracingPayload TraceTransparentRay(inout FPathState PathState, int Bo
 						}
 					}
 					// accumulate all at once so clamping applies to entire volume contribution
-					AccumulateRadiance(PathState.Radiance, VolumeRadiance, bIsCameraRay);
+					AccumulateRadiance(PathState.Radiance, VolumeRadiance);
 					// restore interval start now that we are finished with ray marching
 					Interval.VolumeTMin = StartVolumeTMin;
 
@@ -939,22 +946,23 @@ FPackedPathTracingPayload TraceTransparentRay(inout FPathState PathState, int Bo
 		if (PackedPayload.IsMiss())
 		{
 			// Ray didn't hit any real geometry
-			if (bIsCameraRay)
-			{
-				PathState.Alpha += Luminance(PathState.PathThroughput) * BackgroundAlpha;
-			}
+		#if PATH_TRACER_KERNEL_TYPE == PATH_TRACER_KERNEL_CAMERA
+			PathState.Alpha += Luminance(PathState.PathThroughput) * BackgroundAlpha;
+		#endif
 			break;
 		}
 
-		if (!PackedPayload.IsHoldout() && bIsCameraRay)
+		#if PATH_TRACER_KERNEL_TYPE == PATH_TRACER_KERNEL_CAMERA
+		if (!PackedPayload.IsHoldout())
 		{
 			PathState.Alpha += Luminance(PathState.PathThroughput * (1.0 - PackedPayload.UnpackTransparencyColor()));
 		}
+		#endif
 
 		// add in surface emission
 		if (PathState.ShouldAccumulateEmissive())
 		{
-			AccumulateRadiance(PathState.Radiance, PathState.PathThroughput * PackedPayload.UnpackRadiance(), bIsCameraRay);
+			AccumulateRadiance(PathState.Radiance, PathState.PathThroughput * PackedPayload.UnpackRadiance());
 		}
 
 		if (PackedPayload.HitT == RAY_DEFAULT_T_MAX)
@@ -963,12 +971,14 @@ FPackedPathTracingPayload TraceTransparentRay(inout FPathState PathState, int Bo
 			break;
 		}
 
-		if (bIsCameraRay && PackedPayload.ShouldOutputDepth() && PathState.DepthZ == 0.0)
+	#if PATH_TRACER_KERNEL_TYPE == PATH_TRACER_KERNEL_CAMERA
+		if (PackedPayload.ShouldOutputDepth() && PathState.DepthZ == 0.0)
 		{
 			const float3 ViewZ = View.ViewToTranslatedWorld[2].xyz;
 			const float3 TranslatedWorldPos = PathState.Ray.Origin + PackedPayload.HitT * PathState.Ray.Direction;
 			PathState.DepthZ = dot(TranslatedWorldPos, ViewZ);
 		}
+	#endif
 
 		if (!bLastBounce)
 		{
@@ -1686,7 +1696,7 @@ FPathState CreatePathState(int2 PixelIndex, int2 TextureIndex)
 	Output.SetSigmaSAndPhase(float3(StartingMediumData[3], StartingMediumData[4], StartingMediumData[5]), StartingMediumData[6]);
 	Output.SetFirstScatterType(PATHTRACER_SCATTER_CAMERA);
 
-	#if DEBUG_ENABLE
+	#if DEBUG_ENABLE && PATH_TRACER_KERNEL_TYPE == PATH_TRACER_KERNEL_CAMERA
 	Output.PixelCoord = LaunchIndex;
 	#endif
 
@@ -1705,8 +1715,6 @@ bool PathTracingKernel(inout FPathState PathState, int Bounce)
 	return false;
 #endif
 
-	const bool bIsCameraRay = Bounce == 0;
-
 	FVolumeSegment VolumeSegment = CreateEmptyVolumeSegment();
 	FPackedPathTracingPayload PackedPayload = TraceTransparentRay(
 		PathState,
@@ -1808,7 +1816,7 @@ bool PathTracingKernel(inout FPathState PathState, int Bounce)
 							}
 							LightContrib *= TraceTransparentVisibilityRay(LightRay, 1.0, PathState.GetPathRoughness(), MissShaderIndex, bCastShadows, bCloudCastShadows, CloudFactor, PathState.RandSequence);
 
-							AccumulateRadiance(PathState.Radiance, LightContrib, bIsCameraRay);
+							AccumulateRadiance(PathState.Radiance, LightContrib);
 						}
 					}
 				}
@@ -1833,7 +1841,7 @@ bool PathTracingKernel(inout FPathState PathState, int Bounce)
 							}
 							MaterialRay.TMax = LightResult.HitT;
 							LightContrib *= TraceTransparentVisibilityRay(MaterialRay, 1.0, PathState.GetPathRoughness(), MissShaderIndex, bCastShadows, bCloudCastShadows, CloudFactor, PathState.RandSequence);
-							AccumulateRadiance(PathState.Radiance, LightContrib, bIsCameraRay);
+							AccumulateRadiance(PathState.Radiance, LightContrib);
 						}
 					}
 				}
@@ -1885,7 +1893,7 @@ bool PathTracingKernel(inout FPathState PathState, int Bounce)
 	FMaterialSample MaterialSample = SampleMaterial(-PathState.Ray.Direction, Payload, RandSample.xyz);
 
 	// Debug Printing
-	#if DEBUG_ENABLE
+	#if DEBUG_ENABLE && PATH_TRACER_KERNEL_TYPE == PATH_TRACER_KERNEL_CAMERA
 	FShaderPrintContext Ctx = InitShaderPrintContextAtCursor(PathState.PixelCoord, uint2(50, 50));
 	if (Ctx.bIsActive)
 	{
@@ -1985,7 +1993,7 @@ bool PathTracingKernel(inout FPathState PathState, int Bounce)
 					const uint MissShaderIndex = GetLightMissShaderIndex(LightId);
 					LightContrib *= TraceTransparentVisibilityRay(LightRay, AvgRoughness, PathState.GetPathRoughness(), MissShaderIndex, bCastShadows, bCloudCastShadows, CloudFactor, PathState.RandSequence);
 
-					AccumulateRadiance(PathState.Radiance, LightContrib, bIsCameraRay);
+					AccumulateRadiance(PathState.Radiance, LightContrib);
 				}
 			}
 		}
@@ -2146,10 +2154,9 @@ bool PathTracingKernel(inout FPathState PathState, int Bounce)
 				const uint MissShaderIndex = GetLightMissShaderIndex(LightId);
 				const float ShadowRayRoughness = ApproximateCaustics ? PathState.GetPathRoughness() : 0.0;
 				LightContrib *= TraceTransparentVisibilityRay(LightRay, ShadowRayRoughness, PreviousPathRoughness, MissShaderIndex, bCastShadows, bCloudCastShadows, CloudFactor, PathState.RandSequence);
-				AccumulateRadiance(PathState.Radiance, LightContrib, bIsCameraRay);
+				AccumulateRadiance(PathState.Radiance, LightContrib);
 			}
 		}
-
 	}
 
 #if PATH_TRACER_USE_CLOUD_SHADER
diff --git a/Engine/Shaders/Private/PathTracing/PathTracingDebug.usf b/Engine/Shaders/Private/PathTracing/PathTracingDebug.usf
--- a/Engine/Shaders/Private/PathTracing/PathTracingDebug.usf
+++ b/Engine/Shaders/Private/PathTracing/PathTracingDebug.usf
@@ -42,10 +42,7 @@ RWTexture2D<float>  RWSceneRoughness;
 RWTexture2D<float>  RWSceneLinearDepth;
 RWTexture2D<float>  RWSceneDepth;
 RaytracingAccelerationStructure TLAS;
-
-RWStructuredBuffer<FMarkovChainState> MarkovChainStateGrid;
-uint MarkovChainStateGridSize;
-uint UsePathGuiding;
+uint SceneVisibleLightCount;
 
 int DebugMode;
 int SamplesPerPixel;
@@ -53,6 +50,9 @@ int SampleIndex;
 float SampleBlendFactor;
 int MaxBounces;
 int EnableEmissive;
+#if PATH_TRACER_KERNEL_TYPE != PATH_TRACER_KERNEL_CAMERA
+int FirstBounce;
+#endif
 
 uint4 GridHash(uint4 P)
 {
@@ -117,6 +117,255 @@ FPathTracingResult MakeResult(float3 Radiance, float Alpha, float Z)
 	return Result;
 }
 
+// TODO: Expose these like for the normal path tracer
+static const uint EnabledDirectLightingContributions   = 15u;
+static const uint EnabledIndirectLightingContributions = 15u;
+
+struct FPathState
+{                                   	// packed size
+	FRandomSequence RandSequence;   	//  8 bytes
+	float3 Radiance;                	// 12 bytes
+	float Alpha;                   		//  4 bytes
+	uint3 PackedDiffSpecColor;			// 12 bytes
+	uint2 PackedNormalRoughness;		//  8 bytes
+	FRayDesc Ray;                   	// 12+12 = 24 bytes (TMin/TMax do not need to be stored)
+	float3 PathThroughput;          	// 12 bytes
+	uint2 PackedRoughnessSigma;     	// 2 bytes // half precision (PathRoughness)
+	// 6 bytes // half precision  (SigmaT)
+	// 3 bits - packed in sign bits of SigmaT (FirstScatterType)
+	uint2 PackedScatterPhase;       	// 6 bytes // half precision (SigmaS)
+										// 2 bytes // phase g (could be quantized further)
+
+	// Temporary parameters not stored
+	#if PATH_TRACER_KERNEL_TYPE == PATH_TRACER_KERNEL_CAMERA
+	float DepthZ;                   // Camera ray depth (don't need it otherwise)
+	#endif
+
+	float3 GetDiffuseColor()
+	{
+		return float3(UnpackFloat2FromUInt(PackedDiffSpecColor.x),
+					  UnpackFloat2FromUInt(PackedDiffSpecColor.y).x);
+	}
+
+	float3 GetSpecularColor()
+	{
+		return float3(UnpackFloat2FromUInt(PackedDiffSpecColor.y).y,
+					  UnpackFloat2FromUInt(PackedDiffSpecColor.z));
+	}
+
+	float3 GetNormal()
+	{
+		return float3(UnpackFloat2FromUInt(PackedNormalRoughness.x),
+		              UnpackFloat2FromUInt(PackedNormalRoughness.y).x);
+	}
+
+	float GetRoughness()
+	{
+		return UnpackFloat2FromUInt(PackedNormalRoughness.y).y;
+	}
+
+	void SetDiffuseSpecularColor(float3 Diffuse, float3 Specular)
+	{
+		PackedDiffSpecColor.x = PackFloat2ToUInt(Diffuse.xy);
+		PackedDiffSpecColor.y = PackFloat2ToUInt(float2(Diffuse.z, Specular.x));
+		PackedDiffSpecColor.z = PackFloat2ToUInt(Specular.yz);
+	}
+
+	void SetNormal(float3 Normal)
+	{
+		PackedNormalRoughness.x = PackFloat2ToUInt(Normal.xy);
+		PackedNormalRoughness.y = PackFloat2ToUInt(float2(Normal.z, GetRoughness()));
+	}
+
+	void SetRoughness(float Roughness)
+	{
+		PackedNormalRoughness.y = PackFloat2ToUInt(float2(GetNormal().z, Roughness));
+	}
+
+	uint GetFirstScatterType()
+	{
+		// Extract SigmaT sign bits
+		uint a = (PackedRoughnessSigma.x >> 29) & 4u;
+		uint b = (PackedRoughnessSigma.y >> 30) & 2u;
+		uint c = (PackedRoughnessSigma.y >> 15) & 1u;
+		return a | b | c;
+	}
+
+	void SetFirstScatterType(uint ScatterType)
+	{
+		uint a = ScatterType & 4u;
+		uint b = ScatterType & 2u;
+		uint c = ScatterType & 1u;
+		// erase old value
+		PackedRoughnessSigma.x &= 0x7fffffffu;
+		PackedRoughnessSigma.y &= 0x7fff7fffu;
+		// set new value
+		PackedRoughnessSigma.x |= a << 29;
+		PackedRoughnessSigma.y |= b << 30;
+		PackedRoughnessSigma.y |= c << 15;
+	}
+
+	float GetPathRoughness()
+	{
+		return f16tof32(PackedRoughnessSigma.x);
+	}
+
+	void SetPathRoughness(float PathRoughness)
+	{
+		PackedRoughnessSigma.x &= 0xffff0000u;
+		PackedRoughnessSigma.x |= f32tof16(PathRoughness);
+	}
+	
+	float3 GetSigmaT()
+	{
+		return abs(float3(f16tof32(PackedRoughnessSigma.x >> 16),
+		                  f16tof32(PackedRoughnessSigma.y),
+		                  f16tof32(PackedRoughnessSigma.y >> 16)));
+	}
+
+	void SetSigmaT(float3 SigmaT)
+	{
+		// erase old value (keeping sign bits)
+		PackedRoughnessSigma.x &= 0x8000ffffu;
+		PackedRoughnessSigma.y &= 0x80008000u;
+		SigmaT = abs(SigmaT); // make sure there are no sign bits
+		// write new value
+		PackedRoughnessSigma.x |= f32tof16(SigmaT.x) << 16;
+		PackedRoughnessSigma.y |= f32tof16(SigmaT.y);
+		PackedRoughnessSigma.y |= f32tof16(SigmaT.z) << 16;
+	}
+
+	float4 GetSigmaSAndPhase()
+	{
+		return float4(UnpackFloat2FromUInt(PackedScatterPhase.x), UnpackFloat2FromUInt(PackedScatterPhase.y));
+	}
+
+	void SetSigmaSAndPhase(float3 SigmaS, float PhaseG)
+	{
+		PackedScatterPhase.x = PackFloat2ToUInt(SigmaS.xy);
+		PackedScatterPhase.y = PackFloat2ToUInt(float2(SigmaS.z, PhaseG));
+	}
+
+	bool HasMadeContributionScatter()
+	{
+		// The first scattering type that matters for the PATHTRACER_CONTRIBUTION_* flags is the first non-refractive one
+		// We choose to track camera and pure refracted paths together so that objects behind glass can be treated the same as directly visible objects.
+		uint FirstScatterType = GetFirstScatterType();
+		return !(FirstScatterType == PATHTRACER_SCATTER_CAMERA ||
+			     FirstScatterType == PATHTRACER_SCATTER_REFRACT);
+	}
+
+	bool HasMadeDiffuseOrSpecularScatter()
+	{
+		uint FirstScatterType = GetFirstScatterType();
+		return FirstScatterType == PATHTRACER_SCATTER_DIFFUSE || FirstScatterType == PATHTRACER_SCATTER_SPECULAR;
+	}
+
+	bool UpdateScatterType(uint ScatterType)
+	{
+		// Keep track of the first non-refractive scatter type, and leave it "locked" beyond this
+		// This ensures a particular pixel can only contribute to one lighting component.
+		// The reason to not split refractive events is that you typically want to treat refracted paths the same as camera paths.
+		// For example, when extracting just the diffuse component of a character wearing glasses, you would want the directly visible diffuse
+		// and refracted diffuse to be treated the same.
+		if (!HasMadeContributionScatter())
+		{
+			SetFirstScatterType(ScatterType);
+		}
+		switch (GetFirstScatterType())
+		{
+			// In these cases, only keep tracing if the corresponding lighting component is enabled
+			case PATHTRACER_SCATTER_DIFFUSE:  return (EnabledIndirectLightingContributions & PATHTRACER_CONTRIBUTION_DIFFUSE ) != 0;
+			case PATHTRACER_SCATTER_SPECULAR: return (EnabledIndirectLightingContributions & PATHTRACER_CONTRIBUTION_SPECULAR) != 0;
+			case PATHTRACER_SCATTER_VOLUME:   return (EnabledIndirectLightingContributions & PATHTRACER_CONTRIBUTION_VOLUME  ) != 0;
+		}
+		// In all other cases, keep tracing, we haven't made a final decision on what this path is yet
+		return true;
+	}
+
+	bool ShouldAccumulateEmissive()
+	{
+		uint EnabledLightingContributions = HasMadeContributionScatter() ? EnabledIndirectLightingContributions : EnabledDirectLightingContributions;
+		return (EnabledLightingContributions & PATHTRACER_CONTRIBUTION_EMISSIVE) != 0;
+	}
+
+	// The following functions allow us to scale down the contributions for lighting components before we have scattered
+	bool ShouldAccumulateDiffuse()
+	{
+		return (HasMadeContributionScatter() || (EnabledDirectLightingContributions & PATHTRACER_CONTRIBUTION_DIFFUSE) != 0);
+	}
+
+	bool ShouldAccumulateSpecular()
+	{
+		return (HasMadeContributionScatter() || (EnabledDirectLightingContributions & PATHTRACER_CONTRIBUTION_SPECULAR) != 0);
+	}
+
+	bool ShouldAccumulateVolume()
+	{
+		return (HasMadeContributionScatter() || (EnabledDirectLightingContributions & PATHTRACER_CONTRIBUTION_VOLUME) != 0);
+	}
+
+	float2 GetDiffuseSpecularScale(bool bIsVolumeSample)
+	{
+		float2 Result = 0.0;
+		if (bIsVolumeSample)
+		{
+			if (ShouldAccumulateVolume())
+			{
+				Result.x = 1.0;
+			}
+		}
+		else
+		{
+			if (ShouldAccumulateDiffuse())
+			{
+				Result.x = 1.0;
+			}
+			if (ShouldAccumulateSpecular())
+			{
+				Result.y = 1.0;
+			}
+		}
+		return Result;
+	}
+
+	void WritePixel(uint2 TextureIndex)
+	{
+		// Avoid reading the old pixel on the first sample on the off-chance there is a NaN/Inf pixel ...
+		float4 OldPixel    = SampleIndex > 0 ? RWSceneColor[TextureIndex] : 0;
+		float4 OldDiffuse  = SampleIndex > 0 ? RWSceneDiffuseColor[TextureIndex]   : 0;
+		float4 OldSpecular = SampleIndex > 0 ? RWSceneSpecularColor[TextureIndex]   : 0;
+		float4 OldNormal   = SampleIndex > 0 ? RWSceneNormal[TextureIndex]   : 0;
+		float OldRoughness = SampleIndex > 0 ? RWSceneRoughness[TextureIndex] : 0;
+
+		const float4 PixelValue = float4(Radiance * View.PreExposure, Alpha);
+
+		const float3 Diffuse = GetDiffuseColor();
+		const float3 Specular = GetSpecularColor();
+		const float3 Normal = GetNormal();
+		const float Roughness = GetRoughness();
+
+		RWSceneColor[TextureIndex] = lerp(OldPixel, PixelValue, SampleBlendFactor);
+		RWSceneDiffuseColor[TextureIndex].xyz = lerp(OldDiffuse.xyz, Diffuse, SampleBlendFactor);
+		RWSceneSpecularColor[TextureIndex].xyz = lerp(OldSpecular.xyz, Specular, SampleBlendFactor);
+		RWSceneNormal[TextureIndex] = float4(lerp(OldNormal, float4(Normal.xyz, Roughness), SampleBlendFactor));
+		RWSceneRoughness[TextureIndex] = lerp(OldRoughness, Roughness, SampleBlendFactor);
+	}
+
+	#if PATH_TRACER_KERNEL_TYPE == PATH_TRACER_KERNEL_CAMERA
+	void WriteDepth(uint2 TextureIndex)
+	{
+		float OldLinearDepth = SampleIndex > 0 ? RWSceneLinearDepth[TextureIndex] : FLOAT_MAX;
+		float OldDepth = SampleIndex > 0 ? RWSceneDepth[TextureIndex] : 0;
+		RWSceneLinearDepth[TextureIndex] = min(OldLinearDepth, DepthZ);
+		RWSceneDepth[TextureIndex] = max(OldDepth, ConvertToDeviceZ(DepthZ));
+	}
+	#endif
+};
+
+
+
+
 float3 GetLightingChannelMaskDebugColor(int LightingChannelMask)
 {
 	float r = LightingChannelMask & 1 ? 1.0 : 0.1;
@@ -133,90 +382,29 @@ float3 AnisoDebugColor(float InAniso)
 	return float3(0.0, pow(AniG, 2.2), pow(AniB, 2.2));
 }
 
-float3 StateGetPos(FMarkovChainState State)
-{
-	return State.WeightSum > 0.0 ? State.WeightedTarget / State.WeightSum : State.WeightedTarget;
-}
-
-float StateGetCos(FMarkovChainState State)
-{
-	return State.WeightSum > 0 ? State.WeightedCosine / State.WeightSum : 0.0;
-}
-
-float3 StateGetDir(FMarkovChainState State, float3 AbsoluteWorldPos)
-{
-	return normalize(StateGetPos(State) - AbsoluteWorldPos);
-}
-
-float4 StateGetVMFLobe(FMarkovChainState State, float3 AbsoluteWorldPos)
-{
-	float3 StateP = StateGetPos(State);
-	float  StateR = StateGetCos(State);
-	float Np = 0.2 / length2(AbsoluteWorldPos - StateP); // Prior (see Section 4, step 2)
-	float r = MISWeightBalanced(State.N * State.N, Np) * StateR;
-	r = min(r, 0.99999994); // clamp to avoid singularity at r=1
-    float Kappa = r * (3.0 - r * r) / (1.0 - r * r); // Banerjee et al. 2005 - eq 4.4
-	return float4(normalize(StateP - AbsoluteWorldPos), Kappa);
-}
-
-// Based on Yusuke Tokuyoshi, A Numerically Stable Implementation of the von Mises–Fisher Distribution on S2
-float log1p(const float x)
-{
-	return abs(x) < 0.01 ? sign(x) * mad(-0.5f * x, x, x) : log(1 + x);
-}
-float expm1(const float x)
-{
-	return abs(x) < 0.03 ? sign(x) * mad(0.5f * x, x, x) : exp(x) - 1;
-}
-
-float x_over_expm1(const float x)
-{
-	return abs(x) < 0.03 ? rcp(0.5 * x + 1) : x * rcp(exp(x) - 1);
-}
-
-float VMFPdf(const float3 W, const float4 LobeData) {
-	const float3 Mu = LobeData.xyz;
-	const float Kappa = LobeData.w;
-	const float3 D = W - Mu;
-	return exp(-0.5f * Kappa * dot(D, D)) * x_over_expm1(-2.0f * Kappa) * (1.0 / (4.0 * PI));
-}
-
-float3 VMFSample(const float4 LobeData, const float2 Rand) {
-	const float3 Mu = LobeData.xyz;
-	const float Kappa = LobeData.w;
-	const float Phi = (2 * PI) * Rand.x;
-    const float r = Kappa > ( 1.192092896e-07F / 4.0f) ? log1p(Rand.y * expm1(-2.0f * Kappa)) / Kappa : -2.0f * Rand.y;
-    const float SinTheta = sqrt(-mad(r, r, 2.0f * r));
-	const float CosTheta = 1 + r;
-    const float3 Dir = float3(cos(Phi) * SinTheta, sin(Phi) * SinTheta, CosTheta);
-	return TangentToWorld(Dir, Mu);
-}
-
-int GetDominantSide(float3 N)
-{
-	if (abs(N.x) > abs(N.y) && abs(N.x) > abs(N.z))
-		return (N.x < 0);
-	if (abs(N.y) > abs(N.z))
-		return 2 + (N.y < 0);
-	return 4 + (N.z < 0);
-}
-
-
 #if PATH_TRACER_PRIMARY_RAYS
 
-FPackedPathTracingPayload TraceTransparentRay(FRayDesc Ray, int ScatterType, int Bounce, float PathRoughness, float3 SigmaT, inout float3 PathThroughput, inout FRandomSequence RandSequence, inout FPathTracingResult Result, inout float3 LocalRadiance)
+FPackedPathTracingPayload TraceTransparentRay(inout FPathState PathState, int Bounce)
 {
-	const uint RayFlags = Bounce == 0 ? RAY_FLAG_CULL_BACK_FACING_TRIANGLES : 0;
+#if PATH_TRACER_KERNEL_TYPE == PATH_TRACER_KERNEL_CAMERA
+	const uint RayFlags = RAY_FLAG_CULL_BACK_FACING_TRIANGLES;
+#else
+	const uint RayFlags = 0;
+#endif
 	const uint MissShaderIndex = 0;
-	uint InstanceInclusionMask = (ScatterType == PATHTRACER_SCATTER_CAMERA || ScatterType == PATHTRACER_SCATTER_REFRACT) ? PATHTRACER_MASK_CAMERA : PATHTRACER_MASK_INDIRECT;
-
-	const float2 RandSample = RandSequence.Get2D();
+	const bool bUseCameraRayType = !PathState.HasMadeContributionScatter();
+	uint InstanceInclusionMask = bUseCameraRayType ? PATHTRACER_MASK_CAMERA : PATHTRACER_MASK_INDIRECT;
+	if (length2(PathState.Ray.Origin) < Pow2(PATH_TRACING_SKIP_HAIR_DISTANCE))
+	{
+		InstanceInclusionMask |= bUseCameraRayType ? PATHTRACER_MASK_HAIR_CAMERA : PATHTRACER_MASK_HAIR_INDIRECT;
+	}
+	const float2 RandSample = PathState.RandSequence.Get2D();
 
 	FRISContext HitSample = InitRISContext(RandSample.x);
 
 	float3 ResultThroughput = 0;
 	FPackedPathTracingPayload ResultPayload = (FPackedPathTracingPayload) 0;
-	FPackedPathTracingPayload OpaquePackedPayload = InitPathTracingPayload(ScatterType, PathRoughness);
+	FPackedPathTracingPayload OpaquePackedPayload = InitPathTracingPayload(PathState.GetFirstScatterType(), PathState.GetPathRoughness());
 	OpaquePackedPayload.SetDBufferA(float4(0, 0, 0, 1));
 	OpaquePackedPayload.SetDBufferB(float4(0, 0, 0, 1));
 	OpaquePackedPayload.SetDBufferC(float4(0, 0, 0, 1));
@@ -233,13 +421,12 @@ FPackedPathTracingPayload TraceTransparentRay(FRayDesc Ray, int ScatterType, int
 			RAY_TRACING_SHADER_SLOT_MATERIAL,
 			RAY_TRACING_NUM_SHADER_SLOTS,
 			MissShaderIndex,
-			Ray.GetNativeDesc(),
+			PathState.Ray.GetNativeDesc(),
 			OpaquePackedPayload,
 			Hit);
-		if (Bounce > 0)
-		{
-			NvReorderThread(Hit);
-		}
+	#if PATH_TRACER_KERNEL_TYPE != PATH_TRACER_KERNEL_CAMERA
+		NvReorderThread(Hit);
+	#endif
 		NvInvokeHitObject(TLAS, Hit, OpaquePackedPayload);
 	}
 #else
@@ -250,7 +437,7 @@ FPackedPathTracingPayload TraceTransparentRay(FRayDesc Ray, int ScatterType, int
 		RAY_TRACING_SHADER_SLOT_MATERIAL,
 		RAY_TRACING_NUM_SHADER_SLOTS,
 		MissShaderIndex,
-		Ray.GetNativeDesc(),
+		PathState.Ray.GetNativeDesc(),
 		OpaquePackedPayload);
 #endif
 	if (OpaquePackedPayload.IsHit())
@@ -258,15 +445,15 @@ FPackedPathTracingPayload TraceTransparentRay(FRayDesc Ray, int ScatterType, int
 		// TODO: If hit, also process decals, etc ...
 
 		// Now that we know where the opaque hit is, trace the remaining portion of the ray prior to this
-		Ray.TMax = asfloat(asuint(OpaquePackedPayload.HitT) - 1);
+		PathState.Ray.TMax = asfloat(asuint(OpaquePackedPayload.HitT) - 1);
 	}
 
-	InstanceInclusionMask = (ScatterType == PATHTRACER_SCATTER_CAMERA || ScatterType == PATHTRACER_SCATTER_REFRACT) ? PATHTRACER_MASK_CAMERA_TRANSLUCENT : PATHTRACER_MASK_INDIRECT_TRANSLUCENT;
+	InstanceInclusionMask = bUseCameraRayType ? PATHTRACER_MASK_CAMERA_TRANSLUCENT : PATHTRACER_MASK_INDIRECT_TRANSLUCENT;
 
 	for (;;)
 	{
 #if PATH_TRACER_INCLUDE_TRANSLUCENT
-		FPackedPathTracingPayload PackedPayload = InitPathTracingPayload(ScatterType, PathRoughness);
+		FPackedPathTracingPayload PackedPayload = InitPathTracingPayload(PathState.GetFirstScatterType(), PathState.GetPathRoughness());
 		PackedPayload.SetDBufferA(float4(0, 0, 0, 1));
 		PackedPayload.SetDBufferB(float4(0, 0, 0, 1));
 		PackedPayload.SetDBufferC(float4(0, 0, 0, 1));
@@ -274,7 +461,7 @@ FPackedPathTracingPayload TraceTransparentRay(FRayDesc Ray, int ScatterType, int
 		if (OpaquePackedPayload.IsHit())
 		{
 			const float3 ViewZ = View.ViewToTranslatedWorld[2].xyz;
-			const float3 TranslatedWorldPos = Ray.Origin + OpaquePackedPayload.HitT * Ray.Direction;
+			const float3 TranslatedWorldPos = PathState.Ray.Origin + OpaquePackedPayload.HitT * PathState.Ray.Direction;
 			PackedPayload.SetFlag(PATH_TRACING_PAYLOAD_INPUT_FLAG_HAS_SCENE_DEPTH);
 			PackedPayload.SetSceneDepth(dot(TranslatedWorldPos, ViewZ));
 		}
@@ -289,13 +476,12 @@ FPackedPathTracingPayload TraceTransparentRay(FRayDesc Ray, int ScatterType, int
 				RAY_TRACING_SHADER_SLOT_MATERIAL,
 				RAY_TRACING_NUM_SHADER_SLOTS,
 				MissShaderIndex,
-				Ray.GetNativeDesc(),
+				PathState.Ray.GetNativeDesc(),
 				PackedPayload,
 				Hit);
-			if (Bounce > 0)
-			{
-				NvReorderThread(Hit);
-			}
+		#if PATH_TRACER_KERNEL_TYPE != PATH_TRACER_KERNEL_CAMERA
+			NvReorderThread(Hit);
+		#endif
 			NvInvokeHitObject(TLAS, Hit, PackedPayload);
 		}
 #else
@@ -307,7 +493,7 @@ FPackedPathTracingPayload TraceTransparentRay(FRayDesc Ray, int ScatterType, int
 			RAY_TRACING_SHADER_SLOT_MATERIAL,
 			RAY_TRACING_NUM_SHADER_SLOTS,
 			MissShaderIndex,
-			Ray.GetNativeDesc(),
+			PathState.Ray.GetNativeDesc(),
 			PackedPayload);
 #endif
 		if (PackedPayload.IsMiss())
@@ -315,11 +501,49 @@ FPackedPathTracingPayload TraceTransparentRay(FRayDesc Ray, int ScatterType, int
 			// If we miss the translucent geometry, the next hit is the opaque one
 			PackedPayload = OpaquePackedPayload;
 		}
+
+		// Loop over lights to capture their contribution
+		// #dxr_todo: if we have lots of lights, having some hierarchical structure would be better ....
+		// number of directly visible lights for the first bounce
+#if PATH_TRACER_KERNEL_TYPE == PATH_TRACER_KERNEL_CAMERA
+		// Run this for camera rays only, after the first bounce, we don't need to include lights in the trace call
+		// because NEE handled it for us
+		for (uint LightId = 0, NumVisibleLights = SceneVisibleLightCount; LightId < NumVisibleLights; ++LightId)
+		{
+			FRayDesc LightRay = PathState.Ray;
+			LightRay.TMax = PackedPayload.IsMiss() ? PathState.Ray.TMax : PackedPayload.HitT;
+			FLightHit LightHit = TraceLight(LightRay, LightId);
+			if (LightHit.IsHit())
+			{
+				// create a virtual surface hit so we only need to worry about volume transmission in one place
+				// NOTE: returning only a single light hit here causes lights at infinity to occlude each other
+				//       this is most likely what artists would want (imagine a scene with two suns for example)
+				//       but is not consistent with how reflections behave and the sorting order will be arbitrary
+				FPathTracingPayload LightHitPayload = (FPathTracingPayload)0; // clear all fields
+				LightHitPayload.HitT = LightHit.HitT;
+				LightHitPayload.Radiance = LightHit.Radiance;
+				if ((SceneLights[LightId].Flags & PATHTRACER_FLAG_TYPE_MASK) == PATHTRACING_LIGHT_SKY)
+				{
+					// UE-222576: cancel out scale factor that is baked into the texture for camera rays only
+					LightHitPayload.Radiance /= GetColor(LightId);
+				}
+				LightHitPayload.ShadingModelID = SHADINGMODELID_UNLIT;
+				// light hits are fully transparent to match how the light loop works (the light surface does not cast a shadow)
+				LightHitPayload.BSDFOpacity = 1.0;
+				//LightHitPayload.TransparencyColor = LightHit.HitT == RAY_DEFAULT_T_MAX ? 1.0 - BackgroundAlpha : 1.0;
+				LightHitPayload.TransparencyColor = 1;
+				LightHitPayload.SetFrontFace();
+				PackedPayload = PackPathTracingPayload(LightHitPayload);
+			}
+		}
+#endif
+
 #else // PATH_TRACER_INCLUDE_TRANSLUCENT
 
 #define PackedPayload OpaquePackedPayload
 
 #endif // PATH_TRACER_INCLUDE_TRANSLUCENT
+
 		// If miss - exit
 		if (PackedPayload.IsMiss())
 		{
@@ -328,44 +552,50 @@ FPackedPathTracingPayload TraceTransparentRay(FRayDesc Ray, int ScatterType, int
 			break;
 		}
 
-		if (Bounce == 0 && !PackedPayload.IsHoldout())
+	#if PATH_TRACER_KERNEL_TYPE == PATH_TRACER_KERNEL_CAMERA
+		if (!PackedPayload.IsHoldout())
 		{
-			Result.Radiance.a += Luminance(PathThroughput * (1.0 - PackedPayload.UnpackTransparencyColor()));
+			PathState.Alpha += Luminance(PathState.PathThroughput * (1.0 - PackedPayload.UnpackTransparencyColor()));
 		}
+	#endif
 
 		// Apply Beer's law
-		PathThroughput *= select(SigmaT > 0.0, exp(-SigmaT * (PackedPayload.HitT - Ray.TMin)), 1.0);
+		float3 SigmaT = PathState.GetSigmaT();
+		PathState.PathThroughput *= select(SigmaT > 0.0, exp(-SigmaT * (PackedPayload.HitT - PathState.Ray.TMin)), 1.0);
 
-		if (Bounce == 0 || EnableEmissive)
+		if (PATH_TRACER_KERNEL_TYPE == PATH_TRACER_KERNEL_CAMERA || EnableEmissive)
 		{
-			LocalRadiance += PathThroughput * PackedPayload.UnpackRadiance();
+			PathState.Radiance += PathState.PathThroughput * PackedPayload.UnpackRadiance();
 		}
 
-		if (Bounce == 0 && PackedPayload.ShouldOutputDepth() && Result.Z == 0.0)
+	#if PATH_TRACER_KERNEL_TYPE == PATH_TRACER_KERNEL_CAMERA
+		if (PackedPayload.ShouldOutputDepth() && PathState.DepthZ == 0.0)
 		{
 			const float3 ViewZ = View.ViewToTranslatedWorld[2].xyz;
-			const float3 TranslatedWorldPos = Ray.Origin + PackedPayload.HitT * Ray.Direction;
-			Result.LinearDepth = dot(TranslatedWorldPos, ViewZ);
-			Result.Z = ConvertToDeviceZ(Result.LinearDepth);
-			Result.WorldNormal = PackedPayload.UnpackWorldNormal();
+			const float3 TranslatedWorldPos = PathState.Ray.Origin + PackedPayload.HitT * PathState.Ray.Direction;
+			PathState.DepthZ = dot(TranslatedWorldPos, ViewZ);
+			PathState.SetNormal(PackedPayload.UnpackWorldNormal());
 		}
+	#endif
 
 		if (Bounce < MaxBounces)
 		{
-			FPathTracingPayload HitPayload = UnpackPathTracingPayload(PackedPayload, Ray);
+			FPathTracingPayload HitPayload = UnpackPathTracingPayload(PackedPayload, PathState.Ray);
 			AdjustPayloadAfterUnpack(HitPayload, true);
-			float3 Contrib = PathThroughput * EstimateMaterialAlbedo(HitPayload);
+			float3 Contrib = PathState.PathThroughput * EstimateMaterialAlbedo(HitPayload);
 
 			// only allow albedo tracking through a very narrow range of roughness values
 			const float DenoiserMinRoughness = 0.000;
 			const float DenoiserMaxRoughness = 0.004;
-			float DenoiserContrib = 1 - saturate((PathRoughness - DenoiserMinRoughness) / (DenoiserMaxRoughness - DenoiserMinRoughness));
-			if (DenoiserContrib > 0 && Bounce == 0)
+			float DenoiserContrib = 1 - saturate((PathState.GetPathRoughness() - DenoiserMinRoughness) / (DenoiserMaxRoughness - DenoiserMinRoughness));
+			DenoiserContrib *= HitPayload.BSDFOpacity;
+			// TODO: track denoiser signal through reflections/refractions
+			if (DenoiserContrib > 0 && PATH_TRACER_KERNEL_TYPE == PATH_TRACER_KERNEL_CAMERA)
 			{
-				float3 DenoiserThroughput = PathThroughput * DenoiserContrib * HitPayload.BSDFOpacity;
+				float3 DenoiserThroughput = PathState.PathThroughput * DenoiserContrib;
 
 				// TEMP: If the material has a valid SSS lobe, just merge it into diffuse for now
-				FSSSRandomWalkInfo SSS = GetMaterialSSSInfo(HitPayload, Ray.Direction);
+				FSSSRandomWalkInfo SSS = GetMaterialSSSInfo(HitPayload, PathState.Ray.Direction);
 				if (any(SSS.Weight > 0))
 				{
 					HitPayload.DiffuseColor += SSS.Color;
@@ -374,28 +604,35 @@ FPackedPathTracingPayload TraceTransparentRay(FRayDesc Ray, int ScatterType, int
 				// Accumulate G-Buffer like data for the denoiser
 				float3 RoughnessData = HitPayload.RoughnessData;
 				float Roughness = lerp(RoughnessData.x, RoughnessData.y, RoughnessData.z); // mix both lobes to get an average value
-				DenoiserThroughput *= HitPayload.WeightV * HitPayload.BSDFOpacity * lerp(1.0, HitPayload.TransmittanceN, HitPayload.CoverageAboveAlongN);
+				DenoiserThroughput *= HitPayload.WeightV * lerp(1.0, HitPayload.TransmittanceN, HitPayload.CoverageAboveAlongN);
 
 				float NormThroughput = Luminance(DenoiserThroughput);
 				float3 DiffuseColor = HitPayload.DiffuseColor;
 				float3 SpecularColor = HitPayload.SpecularColor;
 
-				float MinRoughness = max(DenoiserMinRoughness, PathRoughness);
+				float MinRoughness = max(DenoiserMinRoughness, PathState.GetPathRoughness());
 				float DenoiserRoughnessWeight = saturate((Roughness - MinRoughness) / (DenoiserMaxRoughness - MinRoughness));
 #if 1
 				// Remapping of NoV seems to be required by DLSS
 				// TODO: Could use the energy conservation routines and tweak diffuse as well
-				float NoV = saturate(dot(-Ray.Direction, HitPayload.WorldNormal));
+				float NoV = saturate(dot(-PathState.Ray.Direction, HitPayload.WorldNormal));
 				float x = NoV * NoV;
 				NoV = 0.75 * x / (x + (1 - NoV) * (1 - NoV));
 				float3 F = EnvBRDF(HitPayload.SpecularColor, Roughness, NoV);
 #else
 				float3 F = HitPayload.SpecularColor;
 #endif
-				Result.DiffuseColor += DenoiserThroughput * HitPayload.DiffuseColor;
-				Result.SpecularColor += DenoiserThroughput * F * DenoiserRoughnessWeight;
-				Result.Roughness += NormThroughput * Roughness;
-				//Result.WorldNormal += NormThroughput * HitPayload.WorldNormal;
+
+				// Accumulate contributions to denoiser channels
+				float3 AccumDiffColor = PathState.GetDiffuseColor();
+				float3 AccumSpecColor = PathState.GetSpecularColor();
+				float AccumRoughness = PathState.GetRoughness();
+				AccumDiffColor += DenoiserThroughput * HitPayload.DiffuseColor;
+				AccumSpecColor += DenoiserThroughput * F * DenoiserRoughnessWeight;
+				AccumRoughness += NormThroughput * Roughness;
+				
+				PathState.SetDiffuseSpecularColor(AccumDiffColor, AccumSpecColor);
+				PathState.SetRoughness(AccumRoughness);
 			}
 
 			float SelectionWeight = max3(Contrib.x, Contrib.y, Contrib.z);
@@ -403,109 +640,81 @@ FPackedPathTracingPayload TraceTransparentRay(FRayDesc Ray, int ScatterType, int
 			{
 				// stash this hit as a candidate return value
 				ResultPayload = PackedPayload;
-				ResultThroughput = PathThroughput / SelectionWeight;
+				ResultThroughput = PathState.PathThroughput / SelectionWeight;
 			}
 		}
 
 		// account for local transparency change
-		PathThroughput *= PackedPayload.UnpackTransparencyColor();
+		PathState.PathThroughput *= PackedPayload.UnpackTransparencyColor();
 
 		// prepare next step around the loop (if any throughput is left)
 		// retrace the exact same ray with TMin one ulp past the hit we just found
-		Ray.TMin = ComputeNewTMin(Ray.Origin, Ray.Direction, PackedPayload.HitT);
-		if (!(Ray.TMin < Ray.TMax) || all(PathThroughput == 0))
+		PathState.Ray.TMin = ComputeNewTMin(PathState.Ray.Origin, PathState.Ray.Direction, PackedPayload.HitT);
+		if (!(PathState.Ray.TMin < PathState.Ray.TMax) || all(PathState.PathThroughput == 0))
 		{
 			break;
 		}
 	}
-	PathThroughput = ResultThroughput * HitSample.GetNormalization();
+	PathState.PathThroughput = ResultThroughput * HitSample.GetNormalization();
 	return ResultPayload;
 }
 
 
 
-FPathTracingResult PathTracingDebugCore(FRayDesc Ray, uint2 LaunchIndex, int SampleIndex)
+bool PathTracingDebugCore(inout FPathState PathState, int Bounce)
 {
-	FPathTracingResult Result = (FPathTracingResult) 0;
-	float3 PathThroughput = 1;
-	float PathRoughness = 0;
-	float3 SigmaT = 0; // TODO: initialize from pre-pass for rays starting inside glass/water
-	uint ScatterType = PATHTRACER_SCATTER_CAMERA;
-	FRandomSequence RandSequence = RandomSequenceCreate(uint3(LaunchIndex, View.StateFrameIndex), SampleIndex, SamplesPerPixel);
-
-	// State about the markov chain selection that needs to be kept from one bounce to the next
-	float ScoresSum = 0;
-	float3 PreviousN = 0;
-	FMarkovChainState ChosenState = (FMarkovChainState)0;
-	// MCPG parameters
-#define N_MC  5 // From paper, seems to be the smallest value which produces good results, some subtle improvements from higher values at the cost of performance
-	const float AdaptiveGridProb = 0.7; // Mentioned in text - TODO: could rely on adaptive grid only with a tweak to level distribution?
-	const float StaticGridWidth = 25.6; // matches level 24 of adaptive grid
-	const float LevelScale = 3.0; 		// 3 levels per octave
-	const float MinWidth = 0.1; 		// 1mm - smallest possible voxel size
-	const float TanFactor = 0.0087;		// ~2*tan(0.25 degrees)
-	const int 	NMax = 1024;			// Blend factors [See Dittebrandt et al.]
-	const float AlphaMin = 1.0 / 100.0; // Allow blending over 100 frames, higher than original publication because the chains get reset during BSDF sampling
-	const float MaterialProb = 0.15f; 	// Chance of BSDF sampling (TODO: adapt this as a function of roughness)
-
-	LOOP
-	for (int Bounce = 0; Bounce <= MaxBounces; Bounce++)
-	{
-		float3 LocalRadiance = 0;
-		FPackedPathTracingPayload PackedPayload = TraceTransparentRay(Ray, ScatterType, Bounce, PathRoughness, SigmaT, PathThroughput, RandSequence, Result, LocalRadiance);
+	FPackedPathTracingPayload PackedPayload = TraceTransparentRay(PathState, Bounce);
 
-		if (PackedPayload.IsMiss())
-		{
-			// We didn't hit anything, done
-			break;
-		}
+	if (PackedPayload.IsMiss())
+	{
+		// We didn't hit anything, done
+		return false;
+	}
 
-		const float3 TranslatedWorldPos = Ray.Origin + PackedPayload.HitT * Ray.Direction;
-		const float3 WorldNormal = PackedPayload.UnpackWorldNormal();
-		const uint PrimitiveLightingChannelMask = PackedPayload.GetPrimitiveLightingChannelMask();
+	const float3 TranslatedWorldPos = PathState.Ray.Origin + PackedPayload.HitT * PathState.Ray.Direction;
+	const float3 WorldNormal = PackedPayload.UnpackWorldNormal();
+	const uint PrimitiveLightingChannelMask = PackedPayload.GetPrimitiveLightingChannelMask();
 
-		const bool bCameraRay = Bounce == 0;
-		const bool bLastBounce = Bounce == MaxBounces;
+	const bool bLastBounce = Bounce == MaxBounces;
 
-		FLightLoopCount LightLoopCount = LightGridLookup(TranslatedWorldPos);
+	FLightLoopCount LightLoopCount = LightGridLookup(TranslatedWorldPos);
 
-		// Compute direct lighting
-		FPathTracingPayload Payload = UnpackPathTracingPayload(PackedPayload, Ray);
-		AdjustPayloadAfterUnpack(Payload, true);
-		// TEMP: If the material has a valid SSS lobe, just merge it into diffuse for now
-		FSSSRandomWalkInfo SSS = GetMaterialSSSInfo(Payload, Ray.Direction);
-		if (any(SSS.Weight > 0))
-		{
-			Payload.DiffuseColor += SSS.Color;
-		}
+	// Compute direct lighting
+	FPathTracingPayload Payload = UnpackPathTracingPayload(PackedPayload, PathState.Ray);
+	AdjustPayloadAfterUnpack(Payload, true);
+	// TEMP: If the material has a valid SSS lobe, just merge it into diffuse for now
+	FSSSRandomWalkInfo SSS = GetMaterialSSSInfo(Payload, PathState.Ray.Direction);
+	if (any(SSS.Weight > 0))
+	{
+		Payload.DiffuseColor += SSS.Color;
+	}
 
-		// Generate main random number for direct+indirect sampling
-		float4 RandSample = RandSequence.Get4D();
+	// Generate main random number for direct+indirect sampling
+	float4 RandSample = PathState.RandSequence.Get4D();
 
 #if LIGHT_LOOP_METHOD == 0
-		float LightPickingCdf[RAY_TRACING_LIGHT_COUNT_MAXIMUM];
-		float LightPickingCdfSum = 0;
+	float LightPickingCdf[RAY_TRACING_LIGHT_COUNT_MAXIMUM];
+	float LightPickingCdfSum = 0;
 
-		for (uint Index = 0, Num = LightLoopCount.NumLights; Index < Num; ++Index)
-		{
-			uint LightIndex = GetLightId(Index, LightLoopCount);
-			float LightEstimate = EstimateLight(LightIndex, TranslatedWorldPos, WorldNormal, PrimitiveLightingChannelMask, false);
-			LightPickingCdfSum += LightEstimate;
-			LightPickingCdf[Index] = LightPickingCdfSum;
-		}
+	for (uint Index = 0, Num = LightLoopCount.NumLights; Index < Num; ++Index)
+	{
+		uint LightIndex = GetLightId(Index, LightLoopCount);
+		float LightEstimate = EstimateLight(LightIndex, TranslatedWorldPos, WorldNormal, PrimitiveLightingChannelMask, false);
+		LightPickingCdfSum += LightEstimate;
+		LightPickingCdf[Index] = LightPickingCdfSum;
+	}
 
-		if (LightPickingCdfSum > 0)
-		{
-			// init worked
-			int LightId;
-			float LightPickPdf = 0;
-			
-			FMaterialSample MaterialSample = SampleMaterial(-Ray.Direction, Payload, RandSample.xyz);
-
-			SelectLight(RandSample.x * LightPickingCdfSum, LightLoopCount.NumLights, LightPickingCdf, LightId, LightPickPdf);
-			LightPickPdf /= LightPickingCdfSum;
-			float2 LightRandSample = RandSample.yz;
-			LightId = GetLightId(LightId, LightLoopCount);
+	FMaterialSample MaterialSample = SampleMaterial(-PathState.Ray.Direction, Payload, RandSample.xyz);
+	if (LightPickingCdfSum > 0)
+	{
+		// init worked
+		int LightId;
+		float LightPickPdf = 0;
+
+		SelectLight(RandSample.x * LightPickingCdfSum, LightLoopCount.NumLights, LightPickingCdf, LightId, LightPickPdf);
+		LightPickPdf /= LightPickingCdfSum;
+		float2 LightRandSample = RandSample.yz;
+		LightId = GetLightId(LightId, LightLoopCount);
 #else // LIGHT_LOOP_METHOD != 0
 		float2 LightRandSample = RandSample.xy;
 
@@ -516,7 +725,7 @@ FPathTracingResult PathTracingDebugCore(FRayDesc Ray, uint2 LaunchIndex, int Sam
 		float ShadowRayRoughness = 0;
 		uint ShadowRayLightId;
 #endif
-		FMaterialSample MaterialSample = SampleMaterial(-Ray.Direction, Payload, RandSample.xyz);
+		FMaterialSample MaterialSample = SampleMaterial(-PathState.Ray.Direction, Payload, RandSample.xyz);
 		for (uint Index = 0, Num = LightLoopCount.NumLights; Index < Num; ++Index)
 		{
 			uint LightId = GetLightId(Index, LightLoopCount);
@@ -533,11 +742,12 @@ FPathTracingResult PathTracingDebugCore(FRayDesc Ray, uint2 LaunchIndex, int Sam
 #if 1
 				// Evaluate real material
 				float2 LightDiffuseSpecularScale = float2(GetLightDiffuseScale(LightId), GetLightSpecularScale(LightId));
-				FMaterialEval MaterialEval = EvalMaterial(-Ray.Direction, LightSample.Direction, Payload, LightDiffuseSpecularScale);
-				float3 LightContrib = PathThroughput * LightSample.RadianceOverPdf * MaterialEval.Weight * MaterialEval.Pdf;
+				FMaterialEval MaterialEval = EvalMaterial(-PathState.Ray.Direction, LightSample.Direction, Payload, LightDiffuseSpecularScale);
+				float3 LightContrib = PathState.PathThroughput * LightSample.RadianceOverPdf * MaterialEval.Weight * MaterialEval.Pdf;
+				uint ScatterType = PathState.GetFirstScatterType();
 				if (ScatterType == PATHTRACER_SCATTER_DIFFUSE || ScatterType == PATHTRACER_SCATTER_SPECULAR)
 				{
-					LightContrib *= GetLightIndirectScale(LightId, PathRoughness);
+					LightContrib *= GetLightIndirectScale(LightId, PathState.GetPathRoughness());
 				}
 
 #if USE_BSDF_SAMPLING
@@ -573,7 +783,7 @@ FPathTracingResult PathTracingDebugCore(FRayDesc Ray, uint2 LaunchIndex, int Sam
 					// Trace shadow ray immediately
 					const bool bCastShadows = CastsShadow(LightId);
 					const uint MissShaderIndex = GetLightMissShaderIndex(LightId);
-					LocalRadiance += LightContrib * TraceShadowRay(ShadowRay, GetAverageRoughness(Payload), PathRoughness, MissShaderIndex, bCastShadows);
+					PathState.Radiance += LightContrib * TraceShadowRay(ShadowRay, GetAverageRoughness(Payload), PathState.GetPathRoughness(), MissShaderIndex, bCastShadows);
 #endif
 				}
 			}
@@ -593,11 +803,12 @@ FPathTracingResult PathTracingDebugCore(FRayDesc Ray, uint2 LaunchIndex, int Sam
 				if (LightResult.IsHit())
 				{
 					ShadowRay.TMax = LightResult.HitT;
-					float3 LightContrib = PathThroughput * MaterialSample.Weight * LightResult.Radiance;
+					float3 LightContrib = PathState.PathThroughput * MaterialSample.Weight * LightResult.Radiance;
 					LightContrib *= MaterialSample.ScatterType == PATHTRACER_SCATTER_DIFFUSE ? GetLightDiffuseScale(LightId) : GetLightSpecularScale(LightId);
+					uint ScatterType = PathState.GetFirstScatterType();
 					if (ScatterType == PATHTRACER_SCATTER_DIFFUSE || ScatterType == PATHTRACER_SCATTER_SPECULAR)
 					{
-						LightContrib *= GetLightIndirectScale(LightId, PathRoughness);
+						LightContrib *= GetLightIndirectScale(LightId, PathState.GetPathRoughness());
 					}
 					LightContrib *= MISWeightPower(MaterialSample.Pdf, LightResult.Pdf * LightPickPdf);
 #if LIGHT_LOOP_METHOD == 2
@@ -613,7 +824,7 @@ FPathTracingResult PathTracingDebugCore(FRayDesc Ray, uint2 LaunchIndex, int Sam
 #else
 					const bool bCastShadows = CastsShadow(LightId);
 					const uint MissShaderIndex = GetLightMissShaderIndex(LightId);
-					LocalRadiance += LightContrib * TraceShadowRay(ShadowRay, MaterialSample.Roughness, PathRoughness, MissShaderIndex, bCastShadows);
+					PathState.Radiance += LightContrib * TraceShadowRay(ShadowRay, MaterialSample.Roughness, PathState.GetPathRoughness(), MissShaderIndex, bCastShadows);
 #endif
 				}
 			}
@@ -626,225 +837,83 @@ FPathTracingResult PathTracingDebugCore(FRayDesc Ray, uint2 LaunchIndex, int Sam
 			const bool bCastShadows = CastsShadow(ShadowRayLightId);
 			const uint MissShaderIndex = GetLightMissShaderIndex(ShadowRayLightId);
 			ShadowRayContrib *= TraceSample.GetNormalization();
-			LocalRadiance += ShadowRayContrib * TraceShadowRay(ShadowRaySample, ShadowRayRoughness, PathRoughness, MissShaderIndex, bCastShadows);
+			PathState.Radiance += ShadowRayContrib * TraceShadowRay(ShadowRaySample, ShadowRayRoughness, PathState.GetPathRoughness(), MissShaderIndex, bCastShadows);
 		}
 #endif
 
-		Result.Radiance.rgb += LocalRadiance;
-
-		if (UsePathGuiding && Bounce > 0)
-		{
-			// Algorithm 3: State Transition and Maximum Likelihood Estimation
-			// Writes to grid for the parent path vertex
-			const float3 Contrib = LocalRadiance;
-			const float MCf = dot(Contrib, float3(0.25, 0.5, 0.25));
-			if (MCf > 0 && RandSequence.Get1D() * ScoresSum < MCf * N_MC)
-			{
-				const float3 AbsoluteWorldPosition = DFDemote(DFFastSubtract(Ray.Origin, ResolvedView.PreViewTranslation));
-				const float Distance = length(Ray.Origin);
-
-				// hit point is the "target" from the point of view of the previous hit
-				float3 TargetTranslatedWorldPos = TranslatedWorldPos;
-				float3 TargetAbsoluteWorldPosition = DFDemote(DFFastSubtract(TargetTranslatedWorldPos, ResolvedView.PreViewTranslation));
-				float3 CurrentStateDirection = StateGetDir(ChosenState, AbsoluteWorldPosition);
-				float Cosine = max(0.0, dot(normalize(TargetAbsoluteWorldPosition - AbsoluteWorldPosition), CurrentStateDirection));
-
-				ChosenState.N = min(ChosenState.N + 1, NMax);
-				const float Alpha = max(1.0 / ChosenState.N, AlphaMin);
-
-				float Weight = MCf;
-				ChosenState.WeightSum = lerp(ChosenState.WeightSum, Weight, Alpha);
-				ChosenState.WeightedTarget = lerp(ChosenState.WeightedTarget, Weight * TargetAbsoluteWorldPosition, Alpha);
-				ChosenState.WeightedCosine = lerp(ChosenState.WeightedCosine, Weight * Cosine, Alpha);
-
-				if (AdaptiveGridProb < 1)
-				{
-					// Save to Static Grid
-					float3 RandSample = RandSequence.Get3D();
-					int3 Coord = int3(floor(AbsoluteWorldPosition / StaticGridWidth + RandSample));
-					int Side = 0;
-					uint2 HashIndex = GridHash(int4(Coord, -6 + Side)).xy; // two-hashes for collision detection
-					ChosenState.Hash2 = HashIndex.y;
-					MarkovChainStateGrid[HashIndex.x & (MarkovChainStateGridSize - 1)] = ChosenState;
-				}
-				if (AdaptiveGridProb > 0)
-				{
-					// Save to adaptive grid
-					float4 RandSample = RandSequence.Get4D();
-					float TargetWidth = TanFactor * Distance;
-					float Level = round(LevelScale * log2(max(TargetWidth, MinWidth) / MinWidth)) + floor(-log2(1.0 - RandSample.w));
-					float CellWidth = MinWidth * exp2(Level / LevelScale);
-					int Side = GetDominantSide(PreviousN);
-					int LevelI = int(Level);
-					int3 Coord = int3(floor(AbsoluteWorldPosition / CellWidth + RandSample.xyz));
-					uint2 HashIndex = GridHash(int4(Coord, 6 * LevelI + Side)).xy; // two-hashes for collision detection
-					ChosenState.Hash2 = HashIndex.y;
-					MarkovChainStateGrid[HashIndex.x & (MarkovChainStateGridSize - 1)] = ChosenState;
-				}
-			}
-		}
-
-		// don't generate another ray on the last bounce or second to last bounce if emissive materials are disabled
-		if (bLastBounce || (Bounce + 1 == MaxBounces && EnableEmissive == 0))
-		{
-			break;
-		}
-
-		// sample new direction (potentially guided)
-
-		ScoresSum = 0; // reset
-		ChosenState = (FMarkovChainState)0;
-		float3 SampledDirection = 0;
-		float Scores[N_MC];
-		float4 VMFLobes[N_MC];
-		if (UsePathGuiding)
-		{
-			const float3 WorldSmoothNormal = PackedPayload.UnpackWorldSmoothNormal();
-			const float3 AbsoluteWorldPosition = DFDemote(DFFastSubtract(TranslatedWorldPos, ResolvedView.PreViewTranslation));
-			const float Distance = length(TranslatedWorldPos);
-
-			// Algorithm 1: Grid Resampling 
-			UNROLL
-			for (int Index = 0; Index < N_MC; Index++)
-			{
-				RandomSequence SplitRandSequence = RandSequence.Split(Index, N_MC);
-				float4 RandSample1 = SplitRandSequence.Get4D();
-				float4 RandSample2 = SplitRandSequence.Get4D();
-
-				bool bUseAdaptiveGrid = RandSample2.z < AdaptiveGridProb;
+	// don't generate another ray on the last bounce or second to last bounce if emissive materials are disabled
+	if (bLastBounce || (Bounce + 1 == MaxBounces && EnableEmissive == 0))
+	{
+		return false;
+	}
 
-				// static grid default
-				float CellWidth = StaticGridWidth;
-				int Side = 0;
-				int LevelI = -1;
-				if (bUseAdaptiveGrid)
-				{
-					// adaptive grid
-					float TargetWidth = TanFactor * Distance;
-					float Level = round(LevelScale * log2(max(TargetWidth, MinWidth) / MinWidth)) + floor(-log2(1.0 - RandSample1.w));
-					CellWidth = MinWidth * exp2(Level / LevelScale);
-					LevelI = int(Level);
-					Side = GetDominantSide(WorldSmoothNormal);
-				}
-				// stochastic rounding (pick one of the 8 possible neighbor grid cells randomly)
-				int3 Coord = int3(floor(AbsoluteWorldPosition / CellWidth + RandSample1.xyz));
-				
-				uint2 HashIndex = GridHash(int4(Coord, 6 * LevelI + Side)).xy; // two-hashes for collision detection
+	// sample new direction
 
-				FMarkovChainState State = MarkovChainStateGrid[HashIndex.x & (MarkovChainStateGridSize - 1)];
-					
-				// collision detection
-				State.WeightSum *= float(State.Hash2 == HashIndex.y);
-				if (bUseAdaptiveGrid)
-				{
-					// adaptive grid
-				}
-				else
-				{
-					// static grid
-					float3 StateP = State.WeightSum > 0.0 ? State.WeightedTarget / State.WeightSum : State.WeightedTarget;
-					State.WeightSum *= float(dot(WorldSmoothNormal, normalize(StateP - AbsoluteWorldPosition)) > 0.0);
-				}
+	// BSDF Sampling
+	float3 Throughput = MaterialSample.Weight;
+	float3 SampledDirection = MaterialSample.Direction;
+	if (!PathState.UpdateScatterType(MaterialSample.ScatterType))
+	{
+		return false;
+	}
+	// enlarge roughness based on the chosen lobe roughness
+	float PreviousPathRoughness = PathState.GetPathRoughness();
+	PathState.SetPathRoughness(max(PreviousPathRoughness, MaterialSample.Roughness));
 
-				Scores[Index] = State.WeightSum;
-				ScoresSum += State.WeightSum;
-				VMFLobes[Index] = StateGetVMFLobe(State, AbsoluteWorldPosition);
-				if (ScoresSum * RandSample2.w < Scores[Index])
-				{
-					ChosenState = State;
-					SampledDirection = VMFSample(VMFLobes[Index], RandSample2.xy);
-				}
-			}
-		}
 
-		// Algorithm 2: Sampling and Ray Casting
-		// TODO: Adjust SurfaceProb based on roughness
-		float3 Throughput = 0;
-		if (ScoresSum == 0 || RandSample.w < MaterialProb)
-		{
-			// BSDF Sampling
-			Throughput = MaterialSample.Weight;
-			SampledDirection = MaterialSample.Direction;
-			ScatterType = ScatterType == PATHTRACER_SCATTER_CAMERA || ScatterType == PATHTRACER_SCATTER_REFRACT ? MaterialSample.ScatterType : ScatterType; // change if first bounce
-			PathRoughness = max(MaterialSample.Roughness, PathRoughness);
-			ChosenState = (FMarkovChainState)0;
-		}
-		if (ScoresSum > 0)
+	// update the current extinction if we are crossing a boundary on glass or water
+	if (MaterialSample.PositionBiasSign < 0 && Payload.IsMaterialSolidGlass())
+	{
+		const float3 LocalSigmaT = Payload.GetExtinction();
+		float3 SigmaT = PathState.GetSigmaT();
+		if (Payload.IsFrontFace())
 		{
-			// MIS between guided lobes and BSDF sampling
-			float GuidedPdf = 0;
-			UNROLL
-			for (int Index = 0; Index < N_MC; Index++)
-			{
-				if (Scores[Index] > 0)
-				{
-					GuidedPdf += Scores[Index] * VMFPdf(SampledDirection, VMFLobes[Index]) / ScoresSum;
-				}
-			}
-			FMaterialEval MaterialEval = EvalMaterial(-Ray.Direction, SampledDirection, Payload, 1.0);
-			Throughput = MaterialEval.Weight * MISWeightBalanced(MaterialProb * MaterialEval.Pdf, GuidedPdf - GuidedPdf * MaterialProb) / MaterialProb;
-			MaterialSample.PositionBiasSign = sign(dot(SampledDirection, Payload.WorldGeoNormal));
-			ScatterType = PATHTRACER_SCATTER_DIFFUSE; // TODO: is there a good way to classify the scattering type for guided rays?
-			PathRoughness = 1.0;                      // TODO: return average roughness in EvalMaterial?
+			// entering
+			SigmaT += LocalSigmaT;
 		}
-
-		// update the current extinction if we are crossing a boundary on glass or water
-		if (MaterialSample.PositionBiasSign < 0 && Payload.IsMaterialSolidGlass())
+		else
 		{
-			const float3 LocalSigmaT = Payload.GetExtinction();
-			if (Payload.IsFrontFace())
-			{
-				// entering
-				SigmaT += LocalSigmaT;
-			}
-			else
-			{
-				// exiting
-				SigmaT -= LocalSigmaT;
-				SigmaT = max(SigmaT, 0);
-			}
+			// exiting
+			SigmaT -= LocalSigmaT;
+			SigmaT = max(SigmaT, 0);
 		}
+		PathState.SetSigmaT(SigmaT);
+	}
 
-		float3 NextPathThroughput = PathThroughput * Throughput;
-		// Russian roulette:
-		//   The probability of keeping the path should be roughly proportional to the weight at the current shade point,
-		//  but just using MaterialWeight would miss out on cases where the path throughput changes color (like in a cornell
-		//  box when bouncing between walls of different colors). So use the ratio of the brightest color channel in the
-		//  previous and next throughput.
-		//   The second tweak is to add a sqrt() around the probability to soften the termination probability (paths will last
-		//  a little longer). This allows paths to go a bit deeper than the naive heuristic while still allowing them to terminate
-		//  early. This makes RR effective from the very first bounce without needing to delay it.
-		float ContinuationProb = sqrt(saturate(max3(NextPathThroughput.x, NextPathThroughput.y, NextPathThroughput.z) / max3(PathThroughput.x, PathThroughput.y, PathThroughput.z)));
-		if (ContinuationProb < 1)
-		{
-			// If there is some chance we should terminate the ray, draw an extra random value
-			float RussianRouletteRand = RandSample.w;
-			if (RussianRouletteRand >= ContinuationProb)
-			{
-				// stochastically terminate the path
-				break;
-			}
-			PathThroughput = NextPathThroughput / ContinuationProb;
-		}
-		else
+	float3 PathThroughput = PathState.PathThroughput;
+	float3 NextPathThroughput = PathThroughput * Throughput;
+	// Russian roulette:
+	//   The probability of keeping the path should be roughly proportional to the weight at the current shade point,
+	//  but just using MaterialWeight would miss out on cases where the path throughput changes color (like in a cornell
+	//  box when bouncing between walls of different colors). So use the ratio of the brightest color channel in the
+	//  previous and next throughput.
+	//   The second tweak is to add a sqrt() around the probability to soften the termination probability (paths will last
+	//  a little longer). This allows paths to go a bit deeper than the naive heuristic while still allowing them to terminate
+	//  early. This makes RR effective from the very first bounce without needing to delay it.
+	float ContinuationProb = sqrt(saturate(max3(NextPathThroughput.x, NextPathThroughput.y, NextPathThroughput.z) / max3(PathThroughput.x, PathThroughput.y, PathThroughput.z)));
+	if (ContinuationProb < 1)
+	{
+		// If there is some chance we should terminate the ray, draw an extra random value
+		float RussianRouletteRand = RandSample.w;
+		if (RussianRouletteRand >= ContinuationProb)
 		{
-			PathThroughput = NextPathThroughput;
+			// stochastically terminate the path
+			return false;
 		}
-
-
-
-		// Setup next ray
-		PathThroughput *= Throughput;
-		PreviousN = PackedPayload.UnpackWorldSmoothNormal();
-		Ray.Origin = TranslatedWorldPos;
-		Ray.Direction = SampledDirection;
-		Ray.TMin = 0;
-		Ray.TMax = RAY_DEFAULT_T_MAX;
-
-		ApplyRayBias(Ray.Origin, PackedPayload.HitT, MaterialSample.PositionBiasSign * PackedPayload.UnpackWorldGeoNormal());
+		PathThroughput = NextPathThroughput / ContinuationProb;
 	}
-	Result.Radiance.rgb *= View.PreExposure;
-	return Result;
+	else
+	{
+		PathThroughput = NextPathThroughput;
+	}
+	// Setup next ray
+	PathState.PathThroughput = PathThroughput;
+	PathState.Ray.Origin = TranslatedWorldPos;
+	PathState.Ray.Direction = SampledDirection;
+	PathState.Ray.TMin = 0;
+	PathState.Ray.TMax = RAY_DEFAULT_T_MAX;
+	ApplyRayBias(PathState.Ray.Origin, PackedPayload.HitT, MaterialSample.PositionBiasSign * PackedPayload.UnpackWorldGeoNormal());
+	return true;
 }
 
 #else // PATH_TRACER_PRIMARY_RAYS
@@ -1131,6 +1200,56 @@ FPathTracingResult PathTracingDebugCore(FRayDesc Ray, uint2 LaunchIndex)
 
 #endif // !PATH_TRACER_PRIMARY_RAYS
 
+RWStructuredBuffer<FPathTracingDebugPackedPathState> PathStateData;
+RWBuffer<uint> ActivePaths;
+RWBuffer<uint> NumPathStates;
+
+#if PATH_TRACER_KERNEL_TYPE != PATH_TRACER_KERNEL_TAIL
+RWBuffer<uint> NextActivePaths;
+#endif
+
+#if PATH_TRACER_KERNEL_TYPE != PATH_TRACER_KERNEL_CAMERA
+FPathState LoadPathStateData(uint Index)
+{
+	FPathTracingDebugPackedPathState Data = PathStateData[Index];
+	FPathState Output = (FPathState)0;
+	Output.RandSequence.SampleIndex = Data.RandSeqSampleIndex;
+	Output.RandSequence.SampleSeed  = Data.RandSeqSampleSeed;
+	Output.Radiance = Data.Radiance;
+	Output.Alpha = Data.Alpha;
+	Output.PackedDiffSpecColor = Data.PackedDiffSpecColor;
+	Output.PackedNormalRoughness = Data.PackedNormalRoughness;
+	Output.Ray.Origin = Data.RayOrigin;
+	Output.Ray.Direction = Data.RayDirection;
+	Output.Ray.TMin = 0;
+	Output.Ray.TMax = RAY_DEFAULT_T_MAX;
+	Output.PathThroughput = Data.PathThroughput;
+	Output.PackedRoughnessSigma = Data.PackedRoughnessSigma;
+	Output.PackedScatterPhase = Data.PackedScatterPhase;
+	return Output;
+}
+#endif
+
+#if PATH_TRACER_KERNEL_TYPE != PATH_TRACER_KERNEL_TAIL
+void StorePathStateData(FPathState PathState, uint Index)
+{
+	FPathTracingDebugPackedPathState Data;
+	Data.RandSeqSampleIndex = PathState.RandSequence.SampleIndex;
+	Data.RandSeqSampleSeed  = PathState.RandSequence.SampleSeed;
+	Data.Radiance = PathState.Radiance;
+	Data.Alpha = PathState.Alpha;
+	Data.PackedDiffSpecColor = PathState.PackedDiffSpecColor;
+	Data.PackedNormalRoughness = PathState.PackedNormalRoughness;
+	Data.RayOrigin = PathState.Ray.Origin;
+	Data.RayDirection = PathState.Ray.Direction;
+	Data.PathThroughput = PathState.PathThroughput;
+	Data.PackedRoughnessSigma = PathState.PackedRoughnessSigma;
+	Data.PackedScatterPhase = PathState.PackedScatterPhase;
+	PathStateData[Index] = Data;
+}
+#endif
+
+
 RAY_TRACING_ENTRY_RAYGEN(PathTracingDebugRG)
 {
 	uint2 DispatchIdx = DispatchRaysIndex().xy;
@@ -1138,33 +1257,84 @@ RAY_TRACING_ENTRY_RAYGEN(PathTracingDebugRG)
 
 	ResolvedView = ResolveView();
 
-	uint2 LaunchIndex = DispatchIdx + View.ViewRectMin.xy;
-
+	uint LinearPixelIndex = DispatchIdx.x + DispatchDim.x * DispatchIdx.y;
 
+#if PATH_TRACER_KERNEL_TYPE == PATH_TRACER_KERNEL_CAMERA
+	uint2 LaunchIndex = DispatchIdx + View.ViewRectMin.xy;
 	const float2 AAJitter = 0.5; // trace through pixel center
 	const float2 ViewportUV = (LaunchIndex + AAJitter) * View.BufferSizeAndInvSize.zw;
 	const FRayDesc CameraRay = CreatePrimaryRay(ViewportUV);
+#endif
 
 #if PATH_TRACER_PRIMARY_RAYS
-	FPathTracingResult Result = PathTracingDebugCore(CameraRay, LaunchIndex, SampleIndex);
-	if (SampleIndex > 0)
+
+#if PATH_TRACER_KERNEL_TYPE == PATH_TRACER_KERNEL_CAMERA
+	const int FirstBounce = 0;
+	FPathState PathState = (FPathState) 0;
+	PathState.RandSequence = RandomSequenceCreate(uint3(LaunchIndex, View.StateFrameIndex), SampleIndex, SamplesPerPixel);
+	PathState.Ray = CameraRay;
+	// path state variables (these cary information between bounces)
+	PathState.PathThroughput = 1.0;
+	PathState.SetPathRoughness(0.0);
+	PathState.SetFirstScatterType(PATHTRACER_SCATTER_CAMERA);
+	const uint PathIndex = DispatchIdx.x + 65536 * DispatchIdx.y;
+#else
+
+	// compare to number of active paths from the previous bounce
+	if (LinearPixelIndex >= NumPathStates[DISPATCH_INDIRECT_UINT_COUNT * FirstBounce + 0])
 	{
-		// Blend with previous value
-		Result.Radiance = lerp(RWSceneColor[DispatchIdx], Result.Radiance, SampleBlendFactor);
-		Result.Z = max(RWSceneDepth[DispatchIdx], Result.Z); // Device Z is inverted so larger numbers are closer
-	
-		// G-Buffer like data for denoisers
-		// NOTE: averaging here is not strictly necessary, but will be useful for stochastic situations
-		Result.LinearDepth   = lerp(RWSceneLinearDepth[DispatchIdx]			, Result.LinearDepth	, SampleBlendFactor); // TODO: would it be better to use min? only makes a difference in stochastic cases
-		Result.DiffuseColor  = lerp(RWSceneDiffuseColor[DispatchIdx].xyz	, Result.DiffuseColor	, SampleBlendFactor);
-		Result.SpecularColor = lerp(RWSceneSpecularColor[DispatchIdx].xyz	, Result.SpecularColor	, SampleBlendFactor);
-		Result.WorldNormal   = lerp(RWSceneNormal[DispatchIdx].xyz			, Result.WorldNormal	, SampleBlendFactor); // TODO: should we normalize? only makes a difference in stochastic cases
-		Result.Roughness     = lerp(RWSceneRoughness[DispatchIdx]			, Result.Roughness		, SampleBlendFactor);
+		return; // nothing left to do on this thread
 	}
-#else
-	FPathTracingResult Result = PathTracingDebugCore(CameraRay, LaunchIndex);
+	const uint PathIndex = ActivePaths[LinearPixelIndex];
+	FPathState PathState = LoadPathStateData((PathIndex & 65535) + DispatchDim.x * (PathIndex >> 16));
+
+#endif
+	if (LinearPixelIndex == 0)
+	{
+		// write other dimensions for indirect dispatch of the next bounce
+		NumPathStates[DISPATCH_INDIRECT_UINT_COUNT * (FirstBounce + 1) + 1] = 1;
+		NumPathStates[DISPATCH_INDIRECT_UINT_COUNT * (FirstBounce + 1) + 2] = 1;
+	}
+
+#if PATH_TRACER_KERNEL_TYPE != PATH_TRACER_KERNEL_TAIL
+	const bool KeepGoing = PathTracingDebugCore(PathState, FirstBounce);
+	const uint2 TextureIndex = uint2(PathIndex & 65535, PathIndex >> 16);
+	#if PATH_TRACER_KERNEL_TYPE == PATH_TRACER_KERNEL_CAMERA
+	{
+		// Camera ray must write depth right away because we don't want to have to carry DepthZ in the PackedPathState
+		PathState.WriteDepth(TextureIndex);
+	}
+	#endif
+	if (KeepGoing)
+	{
+	    // Add the path to the set of valid paths for the next bounce
+		uint PathStateIndex;
+		InterlockedAdd(NumPathStates[DISPATCH_INDIRECT_UINT_COUNT * (FirstBounce + 1)], 1, PathStateIndex);
+		NextActivePaths[PathStateIndex] = PathIndex;
+		StorePathStateData(PathState, (PathIndex & 65535) + DispatchDim.x * (PathIndex >> 16));
+	}
+	else
+	{
+		// nothing left to do
+		// Accumulate radiance and update pixel variance
+		PathState.WritePixel(TextureIndex);
+	}
+#else // PATH_TRACER_KERNEL_TYPE == PATH_TRACER_KERNEL_TAIL
+	for (int Bounce = FirstBounce; Bounce <= MaxBounces; Bounce++)
+	{
+		if (!PathTracingDebugCore(PathState, Bounce))
+		{
+			// kernel had nothing more to do
+			break;
+		}
+	}
+	// Accumulate radiance and update pixel variance
+	const uint2 TextureIndex = uint2(PathIndex & 65535, PathIndex >> 16);
+	PathState.WritePixel(TextureIndex);
 #endif
 
+#else
+	FPathTracingResult Result = PathTracingDebugCore(CameraRay, LaunchIndex);
 	// Write to output textures
 	RWSceneColor[DispatchIdx] = Result.Radiance;	
 	RWSceneDepth[DispatchIdx] = Result.Z;
@@ -1173,4 +1343,5 @@ RAY_TRACING_ENTRY_RAYGEN(PathTracingDebugRG)
 	RWSceneNormal[DispatchIdx] = float4(Result.WorldNormal, Result.Roughness); // TODO: which copy of roughness is actually used?
 	RWSceneRoughness[DispatchIdx] = Result.Roughness;
 	RWSceneLinearDepth[DispatchIdx] = Result.LinearDepth;
+#endif
 }
diff --git a/Engine/Shaders/Shared/PathTracingDefinitions.h b/Engine/Shaders/Shared/PathTracingDefinitions.h
--- a/Engine/Shaders/Shared/PathTracingDefinitions.h
+++ b/Engine/Shaders/Shared/PathTracingDefinitions.h
@@ -6,6 +6,11 @@
 #define PATHTRACER_SAMPLER_DEFAULT			0
 #define PATHTRACER_SAMPLER_ERROR_DIFFUSION	1
 
+// Constants for 'PATH_TRACER_KERNEL_TYPE'
+#define PATH_TRACER_KERNEL_CAMERA 0   // Camera rays only
+#define PATH_TRACER_KERNEL_INNER  1   // Internal single-bounce
+#define PATH_TRACER_KERNEL_TAIL   2   // Internal tail bounces (from a fixed point up to the max)
+
 // Constants for the 'Flags' field of FPathTracingLight
 #define PATHTRACER_FLAG_LIGHTING_CHANNEL_MASK				(7 << 0)	// Which lighting channel is this light assigned to?
 #define PATHTRACER_FLAG_TRANSMISSION_MASK					(1 << 3)	// Does the light affect the transmission side?
diff --git a/Engine/Shaders/Shared/RayTracingTypes.h b/Engine/Shaders/Shared/RayTracingTypes.h
--- a/Engine/Shaders/Shared/RayTracingTypes.h
+++ b/Engine/Shaders/Shared/RayTracingTypes.h
@@ -79,6 +79,23 @@ struct FPathTracingPackedPathState
 };
 HLSL_STATIC_ASSERT(sizeof(FPathTracingPackedPathState) == 88, "Packed Path State size should be minimized");
 
+struct FPathTracingDebugPackedPathState
+{
+	uint      RandSeqSampleIndex;
+	uint      RandSeqSampleSeed;
+	float3    Radiance;
+	float     Alpha;
+	uint3     PackedDiffSpecColor;
+	uint2     PackedNormalRoughness;
+	float3    RayOrigin;
+	float3    RayDirection;
+	float3    PathThroughput;
+	uint2     PackedRoughnessSigma;
+	uint2     PackedScatterPhase;
+};
+HLSL_STATIC_ASSERT(sizeof(FPathTracingDebugPackedPathState) == 96, "Packed Path State size should be minimized");
+
+
 struct FRayTracingDecal
 {
 	float3 		TranslatedBoundMin;
@@ -89,24 +106,13 @@ struct FRayTracingDecal
 };
 HLSL_STATIC_ASSERT(sizeof(FRayTracingDecal) == 96, "Ray tracing decal structure should be aligned to 32 bytes for optimal access on the GPU");
 
-struct FMarkovChainState
-{
-    float3 		WeightedTarget;
-    float 		WeightSum;
-	float 		WeightedCosine;
-    uint 		N;
-    uint        Hash2; // second hash for resolving collisions
-	uint        Padding0; // for alignment
-};
-HLSL_STATIC_ASSERT(sizeof(FMarkovChainState) == 32, "");
-
 #ifdef __cplusplus
 } // namespace UE::HLSL
 
 using FRTLightingData = UE::HLSL::FRTLightingData;
 using FPathTracingLight = UE::HLSL::FPathTracingLight;
 using FPathTracingPackedPathState = UE::HLSL::FPathTracingPackedPathState;
+using FPathTracingDebugPackedPathState = UE::HLSL::FPathTracingDebugPackedPathState;
 using FRayTracingDecal = UE::HLSL::FRayTracingDecal;
-using FMarkovChainState = UE::HLSL::FMarkovChainState;
 
 #endif
diff --git a/Engine/Source/Runtime/Renderer/Private/PathTracing.cpp b/Engine/Source/Runtime/Renderer/Private/PathTracing.cpp
--- a/Engine/Source/Runtime/Renderer/Private/PathTracing.cpp
+++ b/Engine/Source/Runtime/Renderer/Private/PathTracing.cpp
@@ -126,20 +126,6 @@ FAutoConsoleVariableRef CVarPathTracingVisualize(
 	ECVF_RenderThreadSafe
 );
 
-TAutoConsoleVariable<int32> CVarPathTracingVisualizeMaxBounces(
-	TEXT("r.PathTracing.Visualize.MaxBounces"),
-	1,
-	TEXT("Number of light bounces for primary rays visualization mode"),
-	ECVF_RenderThreadSafe
-);
-
-TAutoConsoleVariable<bool> CVarPathTracingVisualizeEnableEmissive(
-	TEXT("r.PathTracing.Visualize.EnableEmissive"),
-	false,
-	TEXT("Indicates if emissive materials should contribute to scene lighting for primary rays visualization mode"),
-	ECVF_RenderThreadSafe
-);
-
 TAutoConsoleVariable<bool> CVarPathTracingVisualizeIncludeTranslucent(
 	TEXT("r.PathTracing.Visualize.IncludeTranslucent"),
 	true,
@@ -171,9 +157,8 @@ TAutoConsoleVariable<int32> CVarPathTracingVisualizeIndirectGuidingHashSize(
 
 TAutoConsoleVariable<int32> CVarPathTracingCompactionDepth(
 	TEXT("r.PathTracing.CompactionDepth"),
-	-1,
-	TEXT("Enables path compaction to improve GPU occupancy for the path tracer. The value sets the bounce up to which compaction will happen, beyond that point the path tracer will handle bounces within the dispatch. (default: -1, all bounces)\n")
-	TEXT("Requires r.PathTracing.Experimental=true to modify.\n"),
+	6,
+	TEXT("Enables path compaction to improve GPU occupancy for the path tracer. The value sets the bounce up to which compaction will happen, beyond that point the path tracer will handle bounces within the dispatch. Values <= 0 apply compaction to all bounces."),
 	ECVF_RenderThreadSafe
 );
 
@@ -1269,13 +1254,13 @@ class FPathTracingRG : public FGlobalShader
 	DECLARE_GLOBAL_SHADER(FPathTracingRG)
 	SHADER_USE_ROOT_PARAMETER_STRUCT(FPathTracingRG, FGlobalShader)
 
-	class FCompactionType : SHADER_PERMUTATION_BOOL("PATH_TRACER_USE_COMPACTION");
+	class FKernelType : SHADER_PERMUTATION_INT("PATH_TRACER_KERNEL_TYPE", 3);
 	class FAdaptiveSampling : SHADER_PERMUTATION_BOOL("PATH_TRACER_USE_ADAPTIVE_SAMPLING");
 	class FCloudShader : SHADER_PERMUTATION_BOOL("PATH_TRACER_USE_CLOUD_SHADER");
 	class FSubstrateComplexSpecialMaterial : SHADER_PERMUTATION_BOOL("PATH_TRACER_USE_SUBSTRATE_SPECIAL_COMPLEX_MATERIAL");
 	class FUseSER : SHADER_PERMUTATION_BOOL("PATH_TRACER_USE_SER");
 	class FNeedTMinWorkaround : SHADER_PERMUTATION_BOOL("NEED_TMIN_WORKAROUND");
-	using FPermutationDomain = TShaderPermutationDomain<FCompactionType, FAdaptiveSampling, FCloudShader, FSubstrateComplexSpecialMaterial, FUseSER, FNeedTMinWorkaround>;
+	using FPermutationDomain = TShaderPermutationDomain<FKernelType, FAdaptiveSampling, FCloudShader, FSubstrateComplexSpecialMaterial, FUseSER, FNeedTMinWorkaround>;
 
 	static bool ShouldCompilePermutation(const FGlobalShaderPermutationParameters& Parameters)
 	{
@@ -1288,10 +1273,6 @@ class FPathTracingRG : public FGlobalShader
 			{
 				return false;
 			}
-			if (!PermutationVector.Get<FCompactionType>())
-			{
-				return false;
-			}
 		}
 
 		if (PermutationVector.Get<FCloudShader>())
@@ -1421,7 +1402,8 @@ class FPathTracingDebugRG : public FGlobalShader
 	class FUseSER : SHADER_PERMUTATION_BOOL("PATH_TRACER_USE_SER");
 	class FUsePrimaryRays : SHADER_PERMUTATION_BOOL("PATH_TRACER_PRIMARY_RAYS");
 	class FIncludeTranslucent : SHADER_PERMUTATION_BOOL("PATH_TRACER_INCLUDE_TRANSLUCENT");
-	using FPermutationDomain = TShaderPermutationDomain<FUseSER, FUsePrimaryRays, FIncludeTranslucent>;
+	class FKernelType : SHADER_PERMUTATION_INT("PATH_TRACER_KERNEL_TYPE", 3);
+	using FPermutationDomain = TShaderPermutationDomain<FUseSER, FUsePrimaryRays, FIncludeTranslucent, FKernelType>;
 
 	static bool ShouldCompilePermutation(const FGlobalShaderPermutationParameters& Parameters)
 	{
@@ -1440,6 +1422,16 @@ class FPathTracingDebugRG : public FGlobalShader
 			}
 		}
 
+		// If we are one of the debug modes (i.e.: not "PrimaryRays" mode)
+		if (!PermutationVector.Get<FUsePrimaryRays>())
+		{
+			// then only the CAMERA kernel type applies
+			if (PermutationVector.Get<FKernelType>() != PATH_TRACER_KERNEL_CAMERA)
+			{
+				return false;
+			}
+		}
+
 		return PathTracing::ShouldCompilePathTracingShadersForProject(Parameters.Platform);
 	}
 
@@ -1469,7 +1461,6 @@ class FPathTracingDebugRG : public FGlobalShader
 		SHADER_PARAMETER_RDG_TEXTURE_UAV(RWTexture2D<float>,  RWSceneLinearDepth)
 		SHADER_PARAMETER_RDG_TEXTURE_UAV(RWTexture2D<float>,  RWSceneDepth)
 		SHADER_PARAMETER_RDG_BUFFER_SRV(RaytracingAccelerationStructure, TLAS)
-		SHADER_PARAMETER_RDG_BUFFER_UAV(RWStructuredBuffer<FMarkovChainState>, MarkovChainStateGrid)
 		SHADER_PARAMETER_STRUCT_REF(FViewUniformShaderParameters, ViewUniformBuffer)
 		SHADER_PARAMETER_RDG_TEXTURE(Texture2D<uint>, RandomSequenceSpaceFillingCurve)
 		SHADER_PARAMETER_TEXTURE(Texture2D, PreIntegratedGF)
@@ -1480,9 +1471,6 @@ class FPathTracingDebugRG : public FGlobalShader
 		SHADER_PARAMETER(float, SampleBlendFactor)
 		SHADER_PARAMETER(int32, MaxBounces)
 		SHADER_PARAMETER(int32, EnableEmissive)
-		SHADER_PARAMETER(uint32, MarkovChainStateGridSize)
-		SHADER_PARAMETER(uint32, UsePathGuiding)
-
 
 		// scene lights
 		SHADER_PARAMETER_RDG_BUFFER_SRV(StructuredBuffer<FPathTracingLight>, SceneLights)
@@ -1501,6 +1489,14 @@ class FPathTracingDebugRG : public FGlobalShader
 
 		// Skylight
 		SHADER_PARAMETER_STRUCT_INCLUDE(FPathTracingSkylight, SkylightParameters)
+
+		// extra parameters required for path compacting kernel
+		SHADER_PARAMETER(int32, FirstBounce)
+		SHADER_PARAMETER_RDG_BUFFER_UAV(RWStructuredBuffer<FPathTracingDebugPackedPathState>, PathStateData)
+		SHADER_PARAMETER_RDG_BUFFER_UAV(RWBuffer<uint>, ActivePaths)
+		SHADER_PARAMETER_RDG_BUFFER_UAV(RWBuffer<uint>, NextActivePaths)
+		SHADER_PARAMETER_RDG_BUFFER_UAV(RWBuffer<uint>, NumPathStates)
+
 	END_SHADER_PARAMETER_STRUCT()
 };
 IMPLEMENT_GLOBAL_SHADER(FPathTracingDebugRG, "/Engine/Private/PathTracing/PathTracingDebug.usf", "PathTracingDebugRG", SF_RayGen);
@@ -3274,7 +3270,7 @@ static bool GPUNeedsTMinWorkaround()
 	return IsRHIDeviceAMD();
 }
 
-static FPathTracingRG::FPermutationDomain GetPathTracingRGPermutation(const FViewInfo& View, const FScene& Scene, const bool bUseCompaction)
+static FPathTracingRG::FPermutationDomain GetPathTracingRGPermutation(const FViewInfo& View, const FScene& Scene, const bool bIsCameraRays, const bool bCompactKernel)
 {
 	const bool bUseExperimental = CVarPathTracingExperimental.GetValueOnRenderThread();
 	const bool bUseAdaptiveSampling = bUseExperimental && CVarPathTracingAdaptiveSampling.GetValueOnRenderThread() != 0;
@@ -3284,7 +3280,7 @@ static FPathTracingRG::FPermutationDomain GetPathTracingRGPermutation(const FVie
 	const bool bUseSER = GRHIGlobals.SupportsShaderExecutionReordering && CVarPathTracingShaderExecutionReordering.GetValueOnRenderThread();
 
 	FPathTracingRG::FPermutationDomain Out;
-	Out.Set<FPathTracingRG::FCompactionType>(bUseCompaction);
+	Out.Set<FPathTracingRG::FKernelType>(bIsCameraRays ? PATH_TRACER_KERNEL_CAMERA : (bCompactKernel ? PATH_TRACER_KERNEL_INNER : PATH_TRACER_KERNEL_TAIL));
 	Out.Set<FPathTracingRG::FAdaptiveSampling>(bUseAdaptiveSampling);
 	Out.Set<FPathTracingRG::FCloudShader>(bUseCloudShader);
 	Out.Set<FPathTracingRG::FSubstrateComplexSpecialMaterial>(bHasComplexSpecialRenderPath);
@@ -3293,6 +3289,18 @@ static FPathTracingRG::FPermutationDomain GetPathTracingRGPermutation(const FVie
 	return Out;
 }
 
+static FPathTracingDebugRG::FPermutationDomain GetPathTracingDebugRGPermutation(const FViewInfo& View, const FScene& Scene, const bool bIsCameraRays, const bool bCompactKernel, const int DebugMode)
+{
+	const bool bUsePrimaryRays = DebugMode == PATH_TRACER_DEBUG_VIZ_PRIMARY_RAYS;
+	const bool bUseSER = GRHIGlobals.SupportsShaderExecutionReordering && CVarPathTracingShaderExecutionReordering.GetValueOnRenderThread();
+	FPathTracingDebugRG::FPermutationDomain Out;
+	Out.Set<FPathTracingDebugRG::FKernelType>((bIsCameraRays || !bUsePrimaryRays) ? PATH_TRACER_KERNEL_CAMERA : (bCompactKernel ? PATH_TRACER_KERNEL_INNER : PATH_TRACER_KERNEL_TAIL));
+	Out.Set<FPathTracingDebugRG::FUseSER>(bUseSER && bUsePrimaryRays);
+	Out.Set<FPathTracingDebugRG::FUsePrimaryRays>(bUsePrimaryRays);
+	Out.Set<FPathTracingDebugRG::FIncludeTranslucent>(CVarPathTracingVisualizeIncludeTranslucent.GetValueOnRenderThread());
+	return Out;
+}
+
 void FDeferredShadingSceneRenderer::PreparePathTracing(const FViewInfo& View, const FScene& Scene, TArray<FRHIRayTracingShader*>& OutRayGenShaders)
 {
 	check(View.Family != nullptr);
@@ -3301,30 +3309,52 @@ void FDeferredShadingSceneRenderer::PreparePathTracing(const FViewInfo& View, co
 	if (ViewFamily.EngineShowFlags.PathTracing
 		&& PathTracing::ShouldCompilePathTracingShadersForProject(ViewFamily.GetShaderPlatform()))
 	{
+		const int CompactionDepth = CVarPathTracingCompactionDepth.GetValueOnRenderThread();
+
 		if (int DebugMode = GetPathTracingVisualizationMode(); DebugMode >= 0)
 		{
-			const bool bUsePrimaryRays = DebugMode == PATH_TRACER_DEBUG_VIZ_PRIMARY_RAYS;
-			const bool bUseSER = GRHIGlobals.SupportsShaderExecutionReordering && CVarPathTracingShaderExecutionReordering.GetValueOnRenderThread();
-			FPathTracingDebugRG::FPermutationDomain PermutationVector;
-			PermutationVector.Set<FPathTracingDebugRG::FUseSER>(bUseSER && bUsePrimaryRays);
-			PermutationVector.Set<FPathTracingDebugRG::FUsePrimaryRays>(bUsePrimaryRays);
-			PermutationVector.Set<FPathTracingDebugRG::FIncludeTranslucent>(CVarPathTracingVisualizeIncludeTranslucent.GetValueOnRenderThread());
-			auto RayGenShader = GetGlobalShaderMap(ViewFamily.GetShaderPlatform())->GetShader<FPathTracingDebugRG>(PermutationVector);
-			OutRayGenShaders.Add(RayGenShader.GetRayTracingShader());
+			{
+				FPathTracingDebugRG::FPermutationDomain PermutationVector = GetPathTracingDebugRGPermutation(View, Scene, true, false, DebugMode);
+				auto RayGenShader = GetGlobalShaderMap(ViewFamily.GetShaderPlatform())->GetShader<FPathTracingDebugRG>(PermutationVector);
+				OutRayGenShaders.Add(RayGenShader.GetRayTracingShader());
+			}
+			if (DebugMode == PATH_TRACER_DEBUG_VIZ_PRIMARY_RAYS)
+			{
+				if (CompactionDepth != 1)
+				{
+					FPathTracingDebugRG::FPermutationDomain PermutationVector = GetPathTracingDebugRGPermutation(View, Scene, false, true, DebugMode);
+					auto RayGenShader = GetGlobalShaderMap(ViewFamily.GetShaderPlatform())->GetShader<FPathTracingDebugRG>(PermutationVector);
+					OutRayGenShaders.Add(RayGenShader.GetRayTracingShader());
+				}
+				if (CompactionDepth > 0)
+				{
+					FPathTracingDebugRG::FPermutationDomain PermutationVector = GetPathTracingDebugRGPermutation(View, Scene, false, false, DebugMode);
+					auto RayGenShader = GetGlobalShaderMap(ViewFamily.GetShaderPlatform())->GetShader<FPathTracingDebugRG>(PermutationVector);
+					OutRayGenShaders.Add(RayGenShader.GetRayTracingShader());
+				}
+			}
 		}
 		else
 		{
-			const int CompactionDepth = CVarPathTracingCompactionDepth.GetValueOnRenderThread();
+			
 			// Declare all RayGen shaders that require material closest hit shaders to be bound
-			if (CompactionDepth >= 0)
 			{
-				FPathTracingRG::FPermutationDomain PermutationVector = GetPathTracingRGPermutation(View, Scene, false);
+				// camera ray kernel
+				FPathTracingRG::FPermutationDomain PermutationVector = GetPathTracingRGPermutation(View, Scene, true, false);
+				auto RayGenShader = GetGlobalShaderMap(ViewFamily.GetShaderPlatform())->GetShader<FPathTracingRG>(PermutationVector);
+				OutRayGenShaders.Add(RayGenShader.GetRayTracingShader());
+			}
+			if (CompactionDepth != 1)
+			{
+				// at least one compacted innner bounce
+				FPathTracingRG::FPermutationDomain PermutationVector = GetPathTracingRGPermutation(View, Scene, false, true);
 				auto RayGenShader = GetGlobalShaderMap(ViewFamily.GetShaderPlatform())->GetShader<FPathTracingRG>(PermutationVector);
 				OutRayGenShaders.Add(RayGenShader.GetRayTracingShader());
 			}
-			if (CompactionDepth != 0)
+			if (CompactionDepth > 0)
 			{
-				FPathTracingRG::FPermutationDomain PermutationVector = GetPathTracingRGPermutation(View, Scene, true);
+				// non-compacted tail handling
+				FPathTracingRG::FPermutationDomain PermutationVector = GetPathTracingRGPermutation(View, Scene, false, false);
 				auto RayGenShader = GetGlobalShaderMap(ViewFamily.GetShaderPlatform())->GetShader<FPathTracingRG>(PermutationVector);
 				OutRayGenShaders.Add(RayGenShader.GetRayTracingShader());
 			}
@@ -3530,105 +3560,122 @@ void FDeferredShadingSceneRenderer::RenderPathTracing(
 		GPathTracingRealtimeDenoiserResources.LinearDepth = LinearDepthTexture;		
 
 		FPathTracingState* PathTracingState = GetPathTracingStateFromView(View);
-		FRDGBuffer* MarkovChainStateGrid = nullptr;
-		uint32 MarkovChainStateGridHashSize = FMath::RoundUpToPowerOfTwo(CVarPathTracingVisualizeIndirectGuidingHashSize.GetValueOnRenderThread());
-		if (PathTracingState->MarkovChainStateGrid && PathTracingState->MarkovChainStateGrid->Desc.NumElements == MarkovChainStateGridHashSize)
+
+		const bool bUsePrimaryRays = DebugMode == PATH_TRACER_DEBUG_VIZ_PRIMARY_RAYS;
+		const int SamplesPerPixel = bUsePrimaryRays ? FMath::Max(CVarPathTracingVisualizeSamplesPerPixel.GetValueOnRenderThread(), 1) : 1;
+
+		int32 MaxBounces = bUsePrimaryRays ? CVarPathTracingMaxBounces.GetValueOnRenderThread() : 0;
+		if (MaxBounces < 0)
 		{
-			// ok to re-use previous grid
-			MarkovChainStateGrid = GraphBuilder.RegisterExternalBuffer(PathTracingState->MarkovChainStateGrid, TEXT("PathTracer.MarkovChainStateGrid"));
+			MaxBounces = View.FinalPostProcessSettings.PathTracingMaxBounces;
 		}
-		else
+		int32 CompactionDepth = bUsePrimaryRays ? CVarPathTracingCompactionDepth.GetValueOnRenderThread() : -1;
+		if (CompactionDepth <= 0 || CompactionDepth > MaxBounces)
 		{
-			// need to make a new grid and initialize it to all zeros
-			MarkovChainStateGrid = GraphBuilder.CreateBuffer(FRDGBufferDesc::CreateStructuredDesc(sizeof(FMarkovChainState), MarkovChainStateGridHashSize), TEXT("PathTracer.MarkovChainStateGrid"));
-			AddClearUAVPass(GraphBuilder, GraphBuilder.CreateUAV(MarkovChainStateGrid), 0);
+			CompactionDepth = MaxBounces;
 		}
 
-		const bool bUsePrimaryRays = DebugMode == PATH_TRACER_DEBUG_VIZ_PRIMARY_RAYS;
-		const int SamplesPerPixel = bUsePrimaryRays ? FMath::Max(CVarPathTracingVisualizeSamplesPerPixel.GetValueOnRenderThread(), 1) : 1;
+		FRDGBuffer* ActivePaths[2] = {};
+		FRDGBuffer* NumPathStates = nullptr;
+		FRDGBuffer* PathStateData = nullptr;
+		{
+			const int32 NumPaths = DispatchResX * DispatchResY;
+			ActivePaths[0] = GraphBuilder.CreateBuffer(FRDGBufferDesc::CreateBufferDesc(sizeof(int32), NumPaths), TEXT("PathTracer.ActivePaths0"));
+			ActivePaths[1] = GraphBuilder.CreateBuffer(FRDGBufferDesc::CreateBufferDesc(sizeof(int32), NumPaths), TEXT("PathTracer.ActivePaths1"));
+			NumPathStates = GraphBuilder.CreateBuffer(FRDGBufferDesc::CreateIndirectDesc<FRHIDispatchIndirectParameters>(CompactionDepth + 1), TEXT("PathTracer.NumPathStates"));
+			PathStateData = GraphBuilder.CreateBuffer(FRDGBufferDesc::CreateStructuredDesc(sizeof(FPathTracingDebugPackedPathState), NumPaths), TEXT("PathTracer.PathStateData"));
+		}
 
 		FPathTracingDebugRG::FParameters* FirstPassParameters = nullptr;
 		for (int SampleIndex = 0; SampleIndex < SamplesPerPixel; SampleIndex++)
 		{
-			// TODO: can we do better than making a new set of parameters per sample?
-			FPathTracingDebugRG::FParameters* PassParameters = GraphBuilder.AllocParameters<FPathTracingDebugRG::FParameters>();
-			if (FirstPassParameters == nullptr)
+			AddClearUAVPass(GraphBuilder, GraphBuilder.CreateUAV(NumPathStates, PF_R32_UINT), 0);
+			for (int Bounce = 0; Bounce <= CompactionDepth; Bounce++)
 			{
-				PassParameters->RWSceneColor = GraphBuilder.CreateUAV(SceneColorOutputTexture);
-				PassParameters->RWSceneDepth = GraphBuilder.CreateUAV(DepthTexture);
-				PassParameters->TLAS = Scene->RayTracingScene.GetLayerView(ERayTracingSceneLayer::Base, View.GetRayTracingSceneViewHandle());
-				PassParameters->UsePathGuiding = CVarPathTracingVisualizeIndirectGuiding.GetValueOnRenderThread();
-				PassParameters->MarkovChainStateGrid = GraphBuilder.CreateUAV(MarkovChainStateGrid);
-				PassParameters->MarkovChainStateGridSize = MarkovChainStateGridHashSize;
-				PassParameters->ViewUniformBuffer = View.ViewUniformBuffer;
-				PassParameters->RandomSequenceSpaceFillingCurve = GSystemTextures.GetSpaceFillingCurveTexture(GraphBuilder);
-				PassParameters->PreIntegratedGF = GSystemTextures.PreintegratedGF->GetRHI();
-				PassParameters->PreIntegratedGFSampler = TStaticSamplerState<SF_Bilinear, AM_Clamp, AM_Clamp, AM_Clamp>::GetRHI();
-				PassParameters->DebugMode = DebugMode;
-
-				PassParameters->SamplesPerPixel = SamplesPerPixel;
-				PassParameters->MaxBounces = CVarPathTracingVisualizeMaxBounces.GetValueOnRenderThread();
-				PassParameters->EnableEmissive = CVarPathTracingVisualizeEnableEmissive.GetValueOnRenderThread();
+				// TODO: can we do better than making a new set of parameters per sample?
+				FPathTracingDebugRG::FParameters* PassParameters = GraphBuilder.AllocParameters<FPathTracingDebugRG::FParameters>();
+				if (FirstPassParameters == nullptr)
 				{
-					SetLightParameters(
-						GraphBuilder,
-						Scene,
-						View,
-						true,
-						false,
-						&PassParameters->SkylightParameters,
-						&PassParameters->LightGridParameters,
-						&PassParameters->SceneVisibleLightCount,
-						&PassParameters->SceneLightCount,
-						&PassParameters->SceneLights
-					);
-				}
+					PassParameters->RWSceneColor = GraphBuilder.CreateUAV(SceneColorOutputTexture);
+					PassParameters->RWSceneDepth = GraphBuilder.CreateUAV(DepthTexture);
+					PassParameters->TLAS = Scene->RayTracingScene.GetLayerView(ERayTracingSceneLayer::Base, View.GetRayTracingSceneViewHandle());
+					PassParameters->ViewUniformBuffer = View.ViewUniformBuffer;
+					PassParameters->RandomSequenceSpaceFillingCurve = GSystemTextures.GetSpaceFillingCurveTexture(GraphBuilder);
+					PassParameters->PreIntegratedGF = GSystemTextures.PreintegratedGF->GetRHI();
+					PassParameters->PreIntegratedGFSampler = TStaticSamplerState<SF_Bilinear, AM_Clamp, AM_Clamp, AM_Clamp>::GetRHI();
+					PassParameters->DebugMode = DebugMode;
+
+					PassParameters->SamplesPerPixel = SamplesPerPixel;
+					PassParameters->MaxBounces = MaxBounces;
+					PassParameters->EnableEmissive = CVarPathTracingEnableEmissive.GetValueOnRenderThread();
+					if (PassParameters->EnableEmissive < 0)
+					{
+						PassParameters->EnableEmissive = View.FinalPostProcessSettings.PathTracingEnableEmissiveMaterials;
+					}
+					{
+						SetLightParameters(
+							GraphBuilder,
+							Scene,
+							View,
+							true,
+							false,
+							&PassParameters->SkylightParameters,
+							&PassParameters->LightGridParameters,
+							&PassParameters->SceneVisibleLightCount,
+							&PassParameters->SceneLightCount,
+							&PassParameters->SceneLights
+						);
+					}
 
-				PassParameters->RWSceneDiffuseColor = GraphBuilder.CreateUAV(DiffuseColorTexture);
-				PassParameters->RWSceneSpecularColor = GraphBuilder.CreateUAV(SpecularColorTexture);
-				PassParameters->RWSceneNormal = GraphBuilder.CreateUAV(NormalTexture);
-				PassParameters->RWSceneRoughness = GraphBuilder.CreateUAV(RoughnessTexture);
-				PassParameters->RWSceneLinearDepth = GraphBuilder.CreateUAV(LinearDepthTexture);
+					PassParameters->RWSceneDiffuseColor = GraphBuilder.CreateUAV(DiffuseColorTexture);
+					PassParameters->RWSceneSpecularColor = GraphBuilder.CreateUAV(SpecularColorTexture);
+					PassParameters->RWSceneNormal = GraphBuilder.CreateUAV(NormalTexture);
+					PassParameters->RWSceneRoughness = GraphBuilder.CreateUAV(RoughnessTexture);
+					PassParameters->RWSceneLinearDepth = GraphBuilder.CreateUAV(LinearDepthTexture);
 
-				PassParameters->DecalParameters = View.RayTracingDecalUniformBuffer;
-				PassParameters->Scene = View.GetSceneUniforms().GetBuffer(GraphBuilder);
-				PassParameters->NaniteRayTracing = Nanite::GRayTracingManager.GetUniformBuffer();
+					PassParameters->DecalParameters = View.RayTracingDecalUniformBuffer;
+					PassParameters->Scene = View.GetSceneUniforms().GetBuffer(GraphBuilder);
+					PassParameters->NaniteRayTracing = Nanite::GRayTracingManager.GetUniformBuffer();
 
-				FirstPassParameters = PassParameters;
-			}
-			else
-			{
-				// Already made one set of parameters, re-use it
-				*PassParameters = *FirstPassParameters;
-			}
-			PassParameters->SampleIndex = SampleIndex;
-			PassParameters->SampleBlendFactor = 1.0f / (SampleIndex + 1);
-			const bool bUseSER = GRHIGlobals.SupportsShaderExecutionReordering && CVarPathTracingShaderExecutionReordering.GetValueOnRenderThread();
-			FPathTracingDebugRG::FPermutationDomain PermutationVector;
-			PermutationVector.Set<FPathTracingDebugRG::FUseSER>(bUseSER && bUsePrimaryRays);
-			PermutationVector.Set<FPathTracingDebugRG::FUsePrimaryRays>(bUsePrimaryRays);
-			PermutationVector.Set<FPathTracingDebugRG::FIncludeTranslucent>(CVarPathTracingVisualizeIncludeTranslucent.GetValueOnRenderThread());
-			TShaderMapRef<FPathTracingDebugRG> RayGenShader(View.ShaderMap, PermutationVector);
-			ClearUnusedGraphResources(RayGenShader, PassParameters);
-			GraphBuilder.AddPass(
-				RDG_EVENT_NAME("Path Tracing Visualize [%d/%d]", SampleIndex + 1, SamplesPerPixel),
-				PassParameters,
-				ERDGPassFlags::Compute,
-				[PassParameters, RayGenShader, &View, DispatchResX, DispatchResY](FRHICommandList& RHICmdList)
+					FirstPassParameters = PassParameters;
+				}
+				else
 				{
-					FRHIBatchedShaderParameters& GlobalResources = RHICmdList.GetScratchShaderParameters();
-					SetShaderParameters(GlobalResources, RayGenShader, *PassParameters);
-				
-					TOptional<FScopedUniformBufferStaticBindings> StaticUniformBufferScope = RayTracing::BindStaticUniformBufferBindings(View, PassParameters->Scene->GetRHI(), PassParameters->NaniteRayTracing->GetRHI(), RHICmdList);
-
-					RHICmdList.RayTraceDispatch(
-						View.MaterialRayTracingData.PipelineState,
-						RayGenShader.GetRayTracingShader(),
-						View.MaterialRayTracingData.ShaderBindingTable, GlobalResources,
-						DispatchResX, DispatchResY
-					);
+					// Already made one set of parameters, re-use it
+					*PassParameters = *FirstPassParameters;
 				}
-			);
+				PassParameters->SampleIndex = SampleIndex;
+				PassParameters->SampleBlendFactor = 1.0f / (SampleIndex + 1);
+				PassParameters->FirstBounce = Bounce;
+				PassParameters->NumPathStates = GraphBuilder.CreateUAV(NumPathStates, PF_R32_UINT);
+				PassParameters->ActivePaths = GraphBuilder.CreateUAV(ActivePaths[Bounce & 1], PF_R32_UINT);
+				PassParameters->NextActivePaths = GraphBuilder.CreateUAV(ActivePaths[(Bounce & 1) ^ 1], PF_R32_UINT);
+				PassParameters->PathStateData = GraphBuilder.CreateUAV(PathStateData);
+
+				const bool bUseCompaction = Bounce < CompactionDepth || Bounce == MaxBounces;
+				FPathTracingDebugRG::FPermutationDomain PermutationVector = GetPathTracingDebugRGPermutation(View, *Scene, Bounce == 0, bUseCompaction, DebugMode);
+				TShaderMapRef<FPathTracingDebugRG> RayGenShader(View.ShaderMap, PermutationVector);
+				ClearUnusedGraphResources(RayGenShader, PassParameters);
+				GraphBuilder.AddPass(
+					RDG_EVENT_NAME("Path Tracing Visualize [%d/%d] - Bounce %d%s", SampleIndex + 1, SamplesPerPixel, Bounce, bUseCompaction ? "" : "+"),
+					PassParameters,
+					ERDGPassFlags::Compute,
+					[PassParameters, RayGenShader, &View, DispatchResX, DispatchResY](FRHICommandList& RHICmdList)
+					{
+						FRHIBatchedShaderParameters& GlobalResources = RHICmdList.GetScratchShaderParameters();
+						SetShaderParameters(GlobalResources, RayGenShader, *PassParameters);
+				
+						TOptional<FScopedUniformBufferStaticBindings> StaticUniformBufferScope = RayTracing::BindStaticUniformBufferBindings(View, PassParameters->Scene->GetRHI(), PassParameters->NaniteRayTracing->GetRHI(), RHICmdList);
+
+						RHICmdList.RayTraceDispatch(
+							View.MaterialRayTracingData.PipelineState,
+							RayGenShader.GetRayTracingShader(),
+							View.MaterialRayTracingData.ShaderBindingTable, GlobalResources,
+							DispatchResX, DispatchResY
+						);
+					}
+				);
+			}
 		}
 
 		{
@@ -3654,13 +3701,9 @@ void FDeferredShadingSceneRenderer::RenderPathTracing(
 				DepthStencilState,
 				DisplayParameters);
 		}
-
-		GraphBuilder.QueueBufferExtraction(MarkovChainStateGrid, &PathTracingState->MarkovChainStateGrid);
 		return;
 	}
 
-
-
 	FPathTracingConfig Config = {};
 
 	// Get current value of MaxSPP and reset render if it has changed
@@ -4070,8 +4113,8 @@ void FDeferredShadingSceneRenderer::RenderPathTracing(
 			// We are writing to the texture, we'll need to extract it...
 			bNeedsTextureExtract = true;
 
-			int32 CompactionDepth = bUseExperimental ? CVarPathTracingCompactionDepth.GetValueOnRenderThread() : -1;
-			if (CompactionDepth < 0 || CompactionDepth > int32(Config.PathTracingData.MaxBounces))
+			int32 CompactionDepth = CVarPathTracingCompactionDepth.GetValueOnRenderThread();
+			if (CompactionDepth <= 0 || CompactionDepth > int32(Config.PathTracingData.MaxBounces))
 			{
 				CompactionDepth = Config.PathTracingData.MaxBounces;
 			}
@@ -4080,7 +4123,7 @@ void FDeferredShadingSceneRenderer::RenderPathTracing(
 			const int FlushRenderingCommands = CVarPathTracingFlushDispatch.GetValueOnRenderThread();
 
 			FRDGBuffer* ActivePaths[2] = {};
-			FRDGBuffer* NumActivePaths = nullptr;
+			FRDGBuffer* NumPathStates = nullptr;
 			FRDGBuffer* PathStateData = nullptr;
 			{
 				const int32 NumPaths = FMath::Min(
@@ -4089,7 +4132,7 @@ void FDeferredShadingSceneRenderer::RenderPathTracing(
 				);
 				ActivePaths[0] = GraphBuilder.CreateBuffer(FRDGBufferDesc::CreateBufferDesc(sizeof(int32), NumPaths), TEXT("PathTracer.ActivePaths0"));
 				ActivePaths[1] = GraphBuilder.CreateBuffer(FRDGBufferDesc::CreateBufferDesc(sizeof(int32), NumPaths), TEXT("PathTracer.ActivePaths1"));
-				NumActivePaths = GraphBuilder.CreateBuffer(FRDGBufferDesc::CreateIndirectDesc<uint32>(3 * (CompactionDepth + 1)), TEXT("PathTracer.NumActivePaths"));
+				NumPathStates = GraphBuilder.CreateBuffer(FRDGBufferDesc::CreateIndirectDesc<FRHIDispatchIndirectParameters>(CompactionDepth + 1), TEXT("PathTracer.NumPathStates"));
 				PathStateData = GraphBuilder.CreateBuffer(FRDGBufferDesc::CreateStructuredDesc(sizeof(FPathTracingPackedPathState), NumPaths), TEXT("PathTracer.PathStateData"));
 			}
 
@@ -4116,7 +4159,7 @@ void FDeferredShadingSceneRenderer::RenderPathTracing(
 
 						RDG_EVENT_SCOPE_CONDITIONAL(GraphBuilder, DispatchResX > DispatchSize || DispatchResY > DispatchSize, "Tile=(%d,%d - %dx%d)", TileX, TileY, DispatchSizeX, DispatchSizeYLocal);
 
-						AddClearUAVPass(GraphBuilder, GraphBuilder.CreateUAV(NumActivePaths, PF_R32_UINT), 0);
+						AddClearUAVPass(GraphBuilder, GraphBuilder.CreateUAV(NumPathStates, PF_R32_UINT), 0);
 						if (Config.UseAdaptiveSampling && Config.PathTracingData.Iteration > 0)
 						{
 							// If we are using adaptive sampling, build a smaller list of active paths after the first iteration
@@ -4130,7 +4173,7 @@ void FDeferredShadingSceneRenderer::RenderPathTracing(
 							PassParameters->AdaptiveSamplingErrorThreshold = Config.AdaptiveSamplingThreshold;
 
 							PassParameters->NextActivePaths = GraphBuilder.CreateUAV(ActivePaths[0], PF_R32_UINT);
-							PassParameters->NumPathStates = GraphBuilder.CreateUAV(NumActivePaths, PF_R32_UINT);
+							PassParameters->NumPathStates = GraphBuilder.CreateUAV(NumPathStates, PF_R32_UINT);
 
 							PassParameters->TileTextureOffset.X = TileX;
 							PassParameters->TileTextureOffset.Y = TileY + CurrentGPU * DispatchSizeYSplit;
@@ -4262,10 +4305,10 @@ void FDeferredShadingSceneRenderer::RenderPathTracing(
 							PassParameters->ActivePaths = GraphBuilder.CreateUAV(ActivePaths[Bounce & 1], PF_R32_UINT);
 							PassParameters->NextActivePaths = GraphBuilder.CreateUAV(ActivePaths[(Bounce & 1) ^ 1], PF_R32_UINT);
 							PassParameters->PathStateData = GraphBuilder.CreateUAV(PathStateData);
-							PassParameters->NumPathStates = GraphBuilder.CreateUAV(NumActivePaths, PF_R32_UINT);
+							PassParameters->NumPathStates = GraphBuilder.CreateUAV(NumPathStates, PF_R32_UINT);
 							if (bUseIndirectDispatch)
 							{
-								PassParameters->PathTracingIndirectArgs = NumActivePaths;
+								PassParameters->PathTracingIndirectArgs = NumPathStates;
 							}
 							const bool bEnableDebug = CVarPathTracingDebug.GetValueOnRenderThread() > 0;
 							if (bEnableDebug)
@@ -4278,8 +4321,8 @@ void FDeferredShadingSceneRenderer::RenderPathTracing(
 
 
 							const bool bUseCompaction = Bounce < CompactionDepth || Bounce == Config.PathTracingData.MaxBounces;
-							TShaderMapRef<FPathTracingRG> RayGenShader(View.ShaderMap, GetPathTracingRGPermutation(View, *Scene, bUseCompaction));
-							ClearUnusedGraphResources(RayGenShader, PassParameters);
+							TShaderMapRef<FPathTracingRG> RayGenShader(View.ShaderMap, GetPathTracingRGPermutation(View, *Scene, Bounce == 0, bUseCompaction));
+							ClearUnusedGraphResources(RayGenShader, PassParameters, { PassParameters->PathTracingIndirectArgs });
 							const bool bFlushRenderingCommands = FlushRenderingCommands == 1 || (FlushRenderingCommands == 2 && Bounce == CompactionDepth);
 							const bool bUse1DDispatch = (Config.UseAdaptiveSampling && Config.PathTracingData.Iteration > 0) || Bounce > 0;
 							GraphBuilder.AddPass(
@@ -4301,7 +4344,7 @@ void FDeferredShadingSceneRenderer::RenderPathTracing(
 											View.MaterialRayTracingData.PipelineState,
 											RayGenShader.GetRayTracingShader(),
 											View.MaterialRayTracingData.ShaderBindingTable, GlobalResources,
-											PassParameters->PathTracingIndirectArgs->GetIndirectRHICallBuffer(), 3 * PassParameters->FirstBounce * sizeof(uint32)
+											PassParameters->PathTracingIndirectArgs->GetIndirectRHICallBuffer(), PassParameters->FirstBounce * sizeof(FRHIDispatchIndirectParameters)
 										);
 									}
 									else if (bUse1DDispatch)
