Fix hair cards having incorrect tangent orientation when using TwoSided.

The two-sided sign was applied on tangent orientation, which shouldn't be affected.

#Synced-CL 50344416

diff --git a/Engine/Shaders/Private/MaterialTemplate.ush b/Engine/Shaders/Private/MaterialTemplate.ush
--- a/Engine/Shaders/Private/MaterialTemplate.ush
+++ b/Engine/Shaders/Private/MaterialTemplate.ush
@@ -4342,7 +4342,8 @@ SHADER_POP_WARNINGS_STATE
 
 #endif //MATERIAL_TANGENTSPACENORMAL
 
-#if MATERIAL_TANGENTSPACENORMAL || TWO_SIDED_WORLD_SPACE_SINGLELAYERWATER_NORMAL
+// Hair strands/cards ignore TwoSided as WorldNormal is interpreted as a tangent, and thus is not affected by TwoSided
+#if (MATERIAL_TANGENTSPACENORMAL || TWO_SIDED_WORLD_SPACE_SINGLELAYERWATER_NORMAL)  && !(HAIR_STRAND_MESH_FACTORY || HAIR_CARD_MESH_FACTORY)
 	// flip the normal for backfaces being rendered with a two-sided material
 	Parameters.WorldNormal *= Parameters.TwoSidedSign;
 #endif
@@ -4414,7 +4415,8 @@ SHADER_POP_WARNINGS_STATE
 
 #endif //MATERIAL_TANGENTSPACENORMAL
 
-#if MATERIAL_TANGENTSPACENORMAL || TWO_SIDED_WORLD_SPACE_SINGLELAYERWATER_NORMAL
+// Hair strands/cards ignore TwoSided as WorldNormal is interpreted as a tangent, and thus is not affected by TwoSided
+#if (MATERIAL_TANGENTSPACENORMAL || TWO_SIDED_WORLD_SPACE_SINGLELAYERWATER_NORMAL) && !(HAIR_STRAND_MESH_FACTORY || HAIR_CARD_MESH_FACTORY)
 	// flip the normal for backfaces being rendered with a two-sided material
 	Parameters.WorldNormal *= Parameters.TwoSidedSign;
 #endif
