Fix r.SSR.Compute=1 in splitscreen

#Synced-CL 49782167

diff --git a/Engine/Shaders/Private/SSRT/SSRTReflections.usf b/Engine/Shaders/Private/SSRT/SSRTReflections.usf
--- a/Engine/Shaders/Private/SSRT/SSRTReflections.usf
+++ b/Engine/Shaders/Private/SSRT/SSRTReflections.usf
@@ -455,9 +455,9 @@ void ScreenSpaceReflectionsCS(uint3 DispatchThreadId : SV_DispatchThreadID)
 #endif
 	);
 
-	SSRColorOutput[DispatchThreadId.xy] = OutColor;
+	SSRColorOutput[View.ViewRectMinAndSize.xy + DispatchThreadId.xy] = OutColor;
 #if SSR_OUTPUT_FOR_DENOISER
-	SSRHitDistanceOutput[DispatchThreadId.xy] = OutClosestHitDistance.x;
+	SSRHitDistanceOutput[View.ViewRectMinAndSize.xy + DispatchThreadId.xy] = OutClosestHitDistance.x;
 #endif
 }
 
