[TSR] Move more paths to tsr tensor operations and constant calculation to CPU. Save 0.01ms when antiflickering is enabled on consoles with native 1080p.

#tests All EngineTests passes except two that can be ignored.

#Synced-CL 50303332

diff --git a/Engine/Shaders/Private/TemporalSuperResolution/TSRCommon.ush b/Engine/Shaders/Private/TemporalSuperResolution/TSRCommon.ush
--- a/Engine/Shaders/Private/TemporalSuperResolution/TSRCommon.ush
+++ b/Engine/Shaders/Private/TemporalSuperResolution/TSRCommon.ush
@@ -364,6 +364,18 @@ float2 SafeRcp(float2 x)
 	return float2(SafeRcp(x.x), SafeRcp(x.y));
 }
 
+CALL_SITE_DEBUGLOC
+float3 SafeRcp(float3 x)
+{
+	return float3(SafeRcp(x.x), SafeRcp(x.y), SafeRcp(x.z));
+}
+
+CALL_SITE_DEBUGLOC
+float4 SafeRcp(float4 x)
+{
+	return float4(SafeRcp(x.x), SafeRcp(x.y), SafeRcp(x.z), SafeRcp(x.w));
+}
+
 #if CONFIG_COMPILE_FP16
 
 CALL_SITE_DEBUGLOC
@@ -392,6 +404,32 @@ half2 SafeRcp(half2 x)
 }
 #endif
 
+CALL_SITE_DEBUGLOC
+half3 SafeRcp(half3 x)
+#if PLATFORM_GPU_ARCH >= PLATFORM_GPU_ARCH_AMD_RDNA_2 && PLATFORM_GPU_ARCH <= PLATFORM_GPU_ARCH_AMD_LATTEST
+{
+	// If x=0.0, then MaxHalfFloat * 0.0 = 0.0
+	return min(rcp(x), half(MaxHalfFloat)) * saturate(x * asfloat16(uint16_t(0x7C00)));
+}
+#else
+{
+	return select(x > half(0.0), min(rcp(x), half(MaxHalfFloat)), half(0.0));
+}
+#endif
+
+CALL_SITE_DEBUGLOC
+half4 SafeRcp(half4 x)
+#if PLATFORM_GPU_ARCH >= PLATFORM_GPU_ARCH_AMD_RDNA_2 && PLATFORM_GPU_ARCH <= PLATFORM_GPU_ARCH_AMD_LATTEST
+{
+	// If x=0.0, then MaxHalfFloat * 0.0 = 0.0
+	return min(rcp(x), half(MaxHalfFloat)) * saturate(x * asfloat16(uint16_t(0x7C00)));
+}
+#else
+{
+	return select(x > half(0.0), min(rcp(x), half(MaxHalfFloat)), half(0.0));
+}
+#endif
+
 #endif
 
 CALL_SITE_DEBUGLOC
diff --git a/Engine/Shaders/Private/TemporalSuperResolution/TSRConvolutionNetwork.ush b/Engine/Shaders/Private/TemporalSuperResolution/TSRConvolutionNetwork.ush
--- a/Engine/Shaders/Private/TemporalSuperResolution/TSRConvolutionNetwork.ush
+++ b/Engine/Shaders/Private/TemporalSuperResolution/TSRConvolutionNetwork.ush
@@ -5,6 +5,10 @@
 #include "TSRKernels.ush"
 #include "/Engine/Public/LaneVectorization.ush"
 
+/** Additional functions for LaneVectorization used in TSR */
+TVECTOR_FUNCTION_1PARAMS(ScalarType, SafeRcp)
+
+//-------------------------------------------------------
 
 #ifndef WAVE_COUNT_X
 	#error WAVE_COUNT_X is undefined.
diff --git a/Engine/Shaders/Private/TemporalSuperResolution/TSRShadingAnalysis.ush b/Engine/Shaders/Private/TemporalSuperResolution/TSRShadingAnalysis.ush
--- a/Engine/Shaders/Private/TemporalSuperResolution/TSRShadingAnalysis.ush
+++ b/Engine/Shaders/Private/TemporalSuperResolution/TSRShadingAnalysis.ush
@@ -31,6 +31,9 @@ float FlickeringFramePeriod;
 uint bPassthroughAlpha;
 uint RejectShadingMode;
 
+// Cache the threshold at which the variation count should quick in
+float2 TotalVariationCountInvThreshold;
+
 //------------------------------------------------------- COMPOSE SEPARATE TRANSLUCENCY
 
 /** Compose translucency in linear color space. */
@@ -206,7 +209,6 @@ tsr_tensor_bool ShouldSpatialAntiAlias(tsr_tensor_bool bIsDisoccluded, tsr_tenso
 	return bSpatialAntiAlias;
 }
 
-
 //------------------------------------------------------- FLICKERING TEMPORAL ANALYSIS
 
 void ComputeMoireError(
@@ -223,9 +225,6 @@ void ComputeMoireError(
 	const tsr_half MoireEncodingError = tsr_half(rcp(127.0));
 	const tsr_half MaxPrevTotalVariationCount = tsr_half(20.0);
 
-	// Compute the threshold at which the variation count should quick in
-	const tsr_half TotalVariationCountThreshold = tsr_half(1.0 / (1.0 - pow(1.0 - MinBlendFinal, tsr_half(FlickeringFramePeriod))));
-
 	// Input luminance to monitor flickering from.
 	tsr_tensor_half InputLuma = GCSToSMCS(MoireInput);
 
@@ -371,25 +370,17 @@ void ComputeMoireError(
 		CurrentGradient = sign(CurrentGradient) * max(abs(CurrentGradient) - tsr_tensor_half::Const(0.5) * InputBoxSize, tsr_tensor_half::Const(0.0));
 	#endif
 
-	tsr_tensor_half IsFlicker;
-	UNROLL_N(SIMD_SIZE)
-	for (uint ElementIndex = 0; ElementIndex < SIMD_SIZE; ElementIndex++)
-	{
-		tsr_half LocalCurrentGradient = CurrentGradient.GetElement(ElementIndex);
-		tsr_half LocalPrevGradient = PrevGradient.GetElement(ElementIndex);
-
-		bool bGradientsAreSameSign = LocalCurrentGradient * LocalPrevGradient > tsr_half(0.0);
-		bool bGradientsAreWithinError = abs(LocalCurrentGradient) < MoireEncodingError || abs(LocalPrevGradient) < MoireEncodingError;
-		//bool bCurrentGradientSubstentialEnough = abs(LocalCurrentGradient) > View.GeneralPurposeTweak * abs(LocalPrevGradient); // - 2.0 * MoireEncodingError;
-
-		tsr_half LocalIsFlicker = select((
-			bGradientsAreSameSign ||
-			bGradientsAreWithinError ||
+	const tsr_tensor_half MoireEncodingErrors = tsr_tensor_half::Const(tsr_half(MoireEncodingError));
+	tsr_tensor_bool bGradientsAreSameSign = CurrentGradient * PrevGradient > tsr_tensor_half::Const(0.0);
+	tsr_tensor_bool bGradientsAreWithinError = or(abs(CurrentGradient) < MoireEncodingErrors, abs(PrevGradient) < MoireEncodingErrors);
+	//bool bCurrentGradientSubstentialEnough = abs(LocalCurrentGradient) > View.GeneralPurposeTweak * abs(LocalPrevGradient); // - 2.0 * MoireEncodingError;
+	tsr_tensor_bool bCameraCuts = tsr_tensor_bool::Const(bool(bCameraCut));
+	tsr_tensor_half IsFlicker = select(or(
+			or(bGradientsAreSameSign,
+			bGradientsAreWithinError),
 			//!bCurrentGradientSubstentialEnough ||
-			bCameraCut), tsr_half(0.0), tsr_half(1.0));
+			bCameraCuts), tsr_tensor_half::Const(0.0), tsr_tensor_half::Const(tsr_half(1.0)));
 
-		IsFlicker.SetElement(ElementIndex, LocalIsFlicker);
-	}
 
 	// Compute the total variation of the gradient over time.
 	tsr_tensor_half GradientVariation = min(abs(PrevGradient), abs(CurrentGradient)) * IsFlicker;
@@ -400,63 +391,43 @@ void ComputeMoireError(
 	tsr_tensor_half TotalVariationCountContribution;
 	Deconcatenate(Max3x3(Max3x3(Concatenate(GradientVariation, IsFlicker))), /* out */ TotalVariationContribution, /* out */ TotalVariationCountContribution);
 
-	UNROLL_N(SIMD_SIZE)
-	for (uint ElementIndex = 0; ElementIndex < SIMD_SIZE; ElementIndex++)
-	{
-		tsr_half LocalCurrentGradient = CurrentGradient.GetElement(ElementIndex);
-
-		tsr_half LocalPrevGradient = PrevGradient.GetElement(ElementIndex);
-		tsr_half LocalPrevTotalVariation = PrevTotalVariation.GetElement(ElementIndex);
-		tsr_half LocalPrevTotalVariationCount = PrevTotalVariationCount.GetElement(ElementIndex);
-		bool LocalIsDisoccluded = bIsDisoccluded.GetElement(ElementIndex);
-
-		tsr_half LocalIsFlicker = IsFlicker.GetElement(ElementIndex);
-
-		tsr_half LocalTotalVariationContribution = TotalVariationContribution.GetElement(ElementIndex);
-		tsr_half LocalTotalVariationCountContribution = TotalVariationCountContribution.GetElement(ElementIndex);
-			
-		// Blend out previous gradient
-		tsr_half LocalPrevBlendedGradient = LocalPrevGradient * (tsr_half(1.0) - MinBlendFinal) * (tsr_half(1.0) - LocalIsFlicker);
+	// Blend out previous gradient
+	tsr_tensor_half PrevBlendedGradient = PrevGradient * (MinBlendFinal - tsr_half(1.0)) * (IsFlicker - tsr_half(1.0));
 
-		// Compute new gradient.
-		#if 1
-			tsr_half LocalNewGradient = LocalPrevBlendedGradient + LocalCurrentGradient;
-		#else
-			tsr_half LocalNewGradient = abs(LocalPrevBlendedGradient) > abs(LocalCurrentGradient) ? LocalPrevBlendedGradient : LocalCurrentGradient;
-		#endif
-			
-		// Accumulate current frame variation.
-		tsr_half LocalNewTotalVariation = LocalPrevTotalVariation + LocalTotalVariationContribution;
-		tsr_half LocalNewTotalVariationCount = LocalPrevTotalVariationCount + LocalTotalVariationCountContribution;
-
-		// Discard all moire history on parallax disocclusion.
-		LocalNewGradient            = select(LocalIsDisoccluded, tsr_half(0.0), LocalNewGradient);
-		LocalNewTotalVariation      = select(LocalIsDisoccluded, tsr_half(0.0), LocalNewTotalVariation);
-		LocalNewTotalVariationCount = select(LocalIsDisoccluded, tsr_half(0.0), LocalNewTotalVariationCount);
-
-		// Quantise total variation and count to history bit depth to variation drift in the history over time.
-		tsr_half LocalQuantizedNewTotalVariationCount = floor(LocalNewTotalVariationCount * (tsr_half(255.0) / MaxPrevTotalVariationCount)) * (MaxPrevTotalVariationCount / tsr_half(255.0));
-		LocalNewTotalVariation *= LocalQuantizedNewTotalVariationCount * SafeRcp(LocalNewTotalVariationCount);
-
-		// Compute the final luminance 
-		#if 1
-			tsr_half LocalCountFadeIn = saturate(LocalNewTotalVariationCount / TotalVariationCountThreshold - tsr_half(0.5));
-		#else
-			tsr_half LocalCountFadeIn = saturate(LocalNewTotalVariationCount - TotalVariationCountThreshold);
-		#endif
-		tsr_half LocalMoireError = (abs(LocalNewTotalVariation * SafeRcp(LocalNewTotalVariationCount)) + LocalNewTotalVariationCount * MoireEncodingError) * LocalCountFadeIn;
-
-		tsr_half4 PackedMoireHistory = tsr_half4(
-			FinalFlickeringHistory.GetElement(ElementIndex),
-			LocalNewGradient * tsr_half(127.0 / 255.0) + tsr_half(127.0 / 255.0),
-			LocalNewTotalVariation,
-			LocalQuantizedNewTotalVariationCount * rcp(MaxPrevTotalVariationCount));
-
-		OutMoireHistory.SetElement(ElementIndex, PackedMoireHistory);
+	// Compute new gradient.
+	#if 1
+		tsr_tensor_half NewGradient = PrevBlendedGradient + CurrentGradient;
+	#else
+		tsr_tensor_half NewGradient = abs(PrevBlendedGradient) > abs(CurrentGradient) ? PrevBlendedGradient : CurrentGradient;
+	#endif
+	
+	// Accumulate current frame variation.
+	tsr_tensor_half NewTotalVariation =  PrevTotalVariation + TotalVariationContribution;
+	tsr_tensor_half NewTotalVariationCount = PrevTotalVariationCount + TotalVariationCountContribution;
+
+	// Discard all moire history on parallax disocclusion.
+	NewGradient = select(bIsDisoccluded, tsr_tensor_half::Const(0.0), NewGradient);
+	NewTotalVariation      = select(bIsDisoccluded, tsr_tensor_half::Const(0.0), NewTotalVariation);
+	NewTotalVariationCount = select(bIsDisoccluded, tsr_tensor_half::Const(0.0), NewTotalVariationCount);
+
+	// Quantise total variation and count to history bit depth to variation drift in the history over time.
+	tsr_tensor_half QuantizedNewTotalVariationCount = floor(NewTotalVariationCount * (tsr_half(255.0) / MaxPrevTotalVariationCount)) * (MaxPrevTotalVariationCount / tsr_half(255.0));
+	NewTotalVariation = NewTotalVariation * QuantizedNewTotalVariationCount * SafeRcp(NewTotalVariationCount);
+
+	// Compute the final luminance 
+	#if 1
+		tsr_tensor_half CountFadeIn = saturate(NewTotalVariationCount * tsr_half(TotalVariationCountInvThreshold[0]) - tsr_half(0.5));
+	#else
+		tsr_tensor_half CountFadeIn = saturate(NewTotalVariationCount - rcp(tsr_half(TotalVariationCountInvThreshold[0])));
+	#endif
 
-		OutMoireError.SetElement(ElementIndex, LocalMoireError);
-	}
-		
+	OutMoireError = (abs(NewTotalVariation * SafeRcp(NewTotalVariationCount)) + NewTotalVariationCount * MoireEncodingError) *
+CountFadeIn;
+	
+	OutMoireHistory = Concatenate( 
+						Concatenate(FinalFlickeringHistory, NewGradient * tsr_half(127.0 / 255.0) + tsr_half(127.0 / 255.0)),
+						Concatenate(NewTotalVariation, QuantizedNewTotalVariationCount * rcp(MaxPrevTotalVariationCount)));
+	
 	// Moving object in the scene have very low chance of creating moire pattern with TAA jitter, so discard the moire error to still
 	// have character animation ghost free evem as environement shadow project different things on character.
 	OutMoireError = saturate(OutMoireError * bIsStatic - InputLumaModification);
diff --git a/Engine/Source/Runtime/Renderer/Private/PostProcess/TemporalSuperResolution.cpp b/Engine/Source/Runtime/Renderer/Private/PostProcess/TemporalSuperResolution.cpp
--- a/Engine/Source/Runtime/Renderer/Private/PostProcess/TemporalSuperResolution.cpp
+++ b/Engine/Source/Runtime/Renderer/Private/PostProcess/TemporalSuperResolution.cpp
@@ -967,6 +967,7 @@ class FTSRRejectShadingCS : public FTSRConvolutionNetworkShader
 		SHADER_PARAMETER(int32, bPassthroughAlpha)
 		SHADER_PARAMETER(float, ExposureOffsetFactor)
 		SHADER_PARAMETER(int32, RejectShadingMode)
+		SHADER_PARAMETER(FVector2f, TotalVariationCountInvThreshold)
 
 		SHADER_PARAMETER_RDG_TEXTURE(Texture2D, InputTexture)
 		SHADER_PARAMETER_RDG_TEXTURE_SRV(Texture2D, InputMoireLumaTexture)
@@ -2953,6 +2954,7 @@ FDefaultTemporalUpscaler::FOutputs AddTemporalSuperResolutionPasses(
 		PassParameters->bEnableFlickeringHeuristic = FlickeringFramePeriod > 0.0f;
 		PassParameters->bPassthroughAlpha = IsPrimitiveAlphaHoldoutEnabled(View);
 		PassParameters->RejectShadingMode = CVarTSRShadingRejectionMode.GetValueOnRenderThread();
+		PassParameters->TotalVariationCountInvThreshold = FVector2f((1.0 - FMath::Pow(1.0 - 0.05 /*MinBlendFinal*/, FlickeringFramePeriod)));
 
 		PassParameters->InputTexture = PassInputs.SceneColor.Texture;
 		if (PassInputs.FlickeringInputTexture.IsValid() && FlickeringFramePeriod > 0.0f)
