Substrate - toon - code clean up.

#tests PC editor dx12
#rnx
#Synced-CL 50447436

diff --git a/Engine/Shaders/Private/ShadowProjectionPixelShader.usf b/Engine/Shaders/Private/ShadowProjectionPixelShader.usf
--- a/Engine/Shaders/Private/ShadowProjectionPixelShader.usf
+++ b/Engine/Shaders/Private/ShadowProjectionPixelShader.usf
@@ -399,7 +399,10 @@ void Main(
 				// SUBSTRATE_TODO implmenent that path
 			#else
 				#if SUBSTRATE_ENABLED
-				SSSTransmittanceOrOpticalDepth *= EvaluateToonProfileShadowHatching(ToonProfileId, 0*ScreenUV + 1*PatternUVs, ToonProfileShadowSettings, SSSTransmittanceOrOpticalDepth);
+				if (GBufferData.ShadingModelID == SHADINGMODELID_SUBSTRATE_TOON)
+				{
+					SSSTransmittanceOrOpticalDepth *= EvaluateToonProfileShadowHatching(ToonProfileId, PatternUVs, ToonProfileShadowSettings, SSSTransmittanceOrOpticalDepth);
+				}
 				#endif
 			#endif
 			}
