[Shaders] Fix operator precedence for substitute of bitwise-assignment operators in SingleLayerWaterComposite and HZBTracing (Potential fallout from CL 50318098).

#rnx
#Synced-CL 50342528

diff --git a/Engine/Shaders/Private/HZBTracing.ush b/Engine/Shaders/Private/HZBTracing.ush
--- a/Engine/Shaders/Private/HZBTracing.ush
+++ b/Engine/Shaders/Private/HZBTracing.ush
@@ -173,7 +173,7 @@ void TraceHZB(
 #endif
 		{
 			IntersectionTime = min(IntersectionTime, PlaneIntersectionTimes.z);
-			bSkippedTile = bSkippedTile && IntersectionTime != PlaneIntersectionTimes.z;
+			bSkippedTile = bSkippedTile && (IntersectionTime != PlaneIntersectionTimes.z);
 		}
 
 		if (bSkippedTile)
diff --git a/Engine/Shaders/Private/SingleLayerWaterComposite.usf b/Engine/Shaders/Private/SingleLayerWaterComposite.usf
--- a/Engine/Shaders/Private/SingleLayerWaterComposite.usf
+++ b/Engine/Shaders/Private/SingleLayerWaterComposite.usf
@@ -524,7 +524,7 @@ void WaterTileClassificationBuildListsCS(
 			#if DOWNSAMPLE_FACTOR_X == 2
 				Mask |= select(FullTileCoord.x + 1 < FullTiledViewRes.x, Mask << 1u, Mask);
 			#endif
-				bTileUsed |= (MaskDword & Mask) != 0;
+				bTileUsed = bTileUsed || ((MaskDword & Mask) != 0);
 
 				if (bTileUsed)
 				{
@@ -619,7 +619,7 @@ void WaterTileClassificationBuildListsCS(
 					#if DOWNSAMPLE_FACTOR_X == 2
 						Mask |= select(FullTileCoord.x + 1 < FullTiledViewRes.x, Mask << 1u, Mask);
 					#endif
-						bTileUsed = bTileUsed || (MaskDword & Mask) != 0;
+						bTileUsed = bTileUsed || ((MaskDword & Mask) != 0);
 
 						if (bTileUsed)
 						{
