Substrate - Toon - added support for shadow hatching to gbufferformat 1. Might change later...

#tests PC editor d3d with substrate=1 and gbufferfor;at=1
#rnx
#Synced-CL 50447436
#Synced-CL 50447436

diff --git a/Engine/Shaders/Private/ShadowProjectionPixelShader.usf b/Engine/Shaders/Private/ShadowProjectionPixelShader.usf
--- a/Engine/Shaders/Private/ShadowProjectionPixelShader.usf
+++ b/Engine/Shaders/Private/ShadowProjectionPixelShader.usf
@@ -276,35 +276,37 @@ void Main(
 		bool bHasSubsurface				= false;
 		bool bHasSubsurfaceTransmission	= false;
 		bool bIsEyeOrHair 				= false;
+		bool bIsToon					= false;
 	#if SUBSTRATE_GBUFFER_FORMAT==1
 		FSubstrateAddressing SubstrateAddressing = GetSubstratePixelDataByteOffset(PixelCoord.xy, uint2(ResolvedView.BufferSizeAndInvSize.xy), Substrate.MaxBytesPerPixel);
 		FSubstratePixelHeader SubstrateHeader = UnpackSubstrateHeaderIn(Substrate.MaterialTextureArray, SubstrateAddressing, Substrate.TopLayerTexture);
 		if (SubstrateHeader.IsSubstrateMaterial())
 		{
+			bIsToon						= SubstrateHeader.ClosureCount == 1 && SubstrateHeader.SubstrateGetBSDFType() == SUBSTRATE_BSDF_TYPE_TOON;
 			bIsEyeOrHair				= SubstrateHeader.SubstrateGetBSDFType() == SUBSTRATE_BSDF_TYPE_HAIR || SubstrateHeader.SubstrateGetBSDFType() == SUBSTRATE_BSDF_TYPE_EYE;
 			bHasSubsurface				= SubstrateHeader.HasSubsurface() || bIsEyeOrHair;
 			bHasSubsurfaceTransmission	= SubstrateHeader.HasSubsurface();
 		}
-	#else // SUBSTRATE_GBUFFER_FORMAT==1
+	#else // SUBSTRATE_GBUFFER_FORMAT==0
 		FGBufferData GBufferData = GetGBufferData(ScreenUV);
 		{
 			bHasSubsurface = IsSubsurfaceModel(GBufferData.ShadingModelID);
 			bHasSubsurfaceTransmission = GBufferData.ShadingModelID == SHADINGMODELID_SUBSURFACE_PROFILE;
 			bIsEyeOrHair = GBufferData.ShadingModelID == SHADINGMODELID_HAIR || GBufferData.ShadingModelID == SHADINGMODELID_EYE;
+			bIsToon = GBufferData.ShadingModelID == SHADINGMODELID_SUBSTRATE_TOON;
 		}
-	#endif // SUBSTRATE_GBUFFER_FORMAT==1
+	#endif // SUBSTRATE_GBUFFER_FORMAT
 
 		BRANCH
 		if (bIsSubsurfaceCompatible && bHasSubsurface)
 		{
+			#if SUBSTRATE_ENABLED
 			#if SUBSTRATE_GBUFFER_FORMAT==1
 			const FSubstrateSubsurfaceHeader SSSHeader = SubstrateLoadSubsurfaceHeader(Substrate.MaterialTextureArray, Substrate.FirstSliceStoringSubstrateSSSData, SVPos.xy);
 			const uint SSSType = SubstrateSubSurfaceHeaderGetSSSType(SSSHeader);
-			#else
-			#if SUBSTRATE_ENABLED
-				FToonProfileShadowSettings ToonProfileShadowSettings = (FToonProfileShadowSettings)0;
-			#endif
 			#endif
+			FToonProfileShadowSettings ToonProfileShadowSettings = (FToonProfileShadowSettings)0;
+			#endif // SUBSTRATE_ENABLED
 
 			bool bApplyPCFOverBlurCorrection = true; // Set to false when the transmittance from extinction is correctly respected. Toon & TransmitttanceShadow disables that to have a physically based non squared transmittance curve.
 			float Density = 1;
@@ -318,6 +320,17 @@ void Main(
 					bApplyPCFOverBlurCorrection = false;
 					const float3 MFPcm = SubstrateSubSurfaceHeaderGetTransmittanceShadowMFPInMeters(SSSHeader) * METER_TO_CENTIMETER;
 					Density = 1.0f / max(0.000001f, MFPcm.x); // Recover extinction coefficient as 1/cm from MFP in cm.
+
+					// In this case we also need to check if this is for a toon profile and load profile ID, PatternUVs and Shadow settings.
+					// SSS_TYPE_TRANSMITTANCE_SHADOW was added to try and be general. But it is only used by toon shading today. We could remove it and only load data from the toon profile instead, as done for SUBSTRATE_GBUFFER_FORMAT==0.
+					if (bIsToon)
+					{
+						FSubstrateBSDF BSDF	= UnpackSubstrateBSDFIn(Substrate.MaterialTextureArray, SubstrateAddressing, SubstrateHeader);
+						PatternUVs			= TOON_PATTERNUVS(BSDF);
+						ToonProfileId		= TOON_PROFILEID(BSDF);
+
+						ToonProfileShadowSettings = GetToonProfileShadowSettings(ToonProfileId);
+					}
 				}
 				else if (SSSType != SSS_TYPE_NONE)
 				{
@@ -325,9 +338,9 @@ void Main(
 					bHasSubsurfaceTransmission = bHasSubsurfaceTransmission && (SSSType >= SSS_TYPE_DIFFUSION);
 					Density = SubsurfaceDensityFromOpacity(SubstrateSubSurfaceHeaderGetOpacity(SSSHeader));
 				}
-			#else // SUBSTRATE_GBUFFER_FORMAT==1
+			#else // SUBSTRATE_GBUFFER_FORMAT==0
 				#if SUBSTRATE_ENABLED
-				if (GBufferData.ShadingModelID == SHADINGMODELID_SUBSTRATE_TOON)
+				if (bIsToon)
 				{
 					bApplyPCFOverBlurCorrection = false;
 
@@ -342,7 +355,7 @@ void Main(
 				{
 					Density = SubsurfaceDensityFromOpacity(GBufferData.CustomData.a);
 				}
-			#endif // SUBSTRATE_GBUFFER_FORMAT==1
+			#endif // SUBSTRATE_GBUFFER_FORMAT
 			}
 			
 			FPCFSamplerSettings Settings;
@@ -395,16 +408,13 @@ void Main(
 				// => SSSTransmittanceOrOpticalDepth stores the result of ManualPCF() evaluating Transmittance.
 				SSSTransmittanceOrOpticalDepth = ManualPCF(ShadowPosition.xy, Settings);
 
-			#if SUBSTRATE_GBUFFER_FORMAT==1
-				// SUBSTRATE_TODO implmenent that path
-			#else
 				#if SUBSTRATE_ENABLED
-				if (GBufferData.ShadingModelID == SHADINGMODELID_SUBSTRATE_TOON)
+				BRANCH
+				if (bIsToon)
 				{
 					SSSTransmittanceOrOpticalDepth *= EvaluateToonProfileShadowHatching(ToonProfileId, PatternUVs, ToonProfileShadowSettings, SSSTransmittanceOrOpticalDepth);
 				}
 				#endif
-			#endif
 			}
 
 			// Legacy code maintained: here we have a square applied over the transmittance curve or encoded optical depth... 
