Change RT occlusion for SSS surfaces with Substrate Format=1 to use the unified pass, and fix inconsistency compare to Format=0.

#Synced-CL 50480066

diff --git a/Engine/Shaders/Private/RayTracing/RayTracingOcclusionRGS.usf b/Engine/Shaders/Private/RayTracing/RayTracingOcclusionRGS.usf
--- a/Engine/Shaders/Private/RayTracing/RayTracingOcclusionRGS.usf
+++ b/Engine/Shaders/Private/RayTracing/RayTracingOcclusionRGS.usf
@@ -710,27 +710,16 @@ RAY_TRACING_ENTRY_RAYGEN_OR_INLINE(Occlusion)
 
 			if (ENABLE_TRANSMISSION)
 			{
-				if (SSSType == SSS_TYPE_DIFFUSION_PROFILE)
-				{
-					TransmissionParameters = GetTransmissionProfileParams(SubstrateSubSurfaceHeaderGetProfileId(SSSHeader));
-				}
-				else if (SSSType == SSS_TYPE_DIFFUSION)
+				TransmissionParameters = GetSubstrateSubsurfaceTransmissionProfileParams(SSSType, SSSHeader);
+
+				// If MFP is available, use the actual mean free path, otherwise, fallback to the old behavior using extinction scale in the subsurface profile.
+				SubsurfaceSamplingMeanFreePath = TransmissionParameters.ExtinctionScale;
+
+				BRANCH
+				if (TransmissionMeanFreePathType == TransmissionMeanFreePathType_MaxMeanFreePath)
 				{
-					FSubsurfaceOpacityMFP SubsurfaceOpacityMFP = SubstrateGetSubsurfaceOpacityMFP(SSSHeader);
-
-					// We know that or SSS_TYPE_DIFFUSION we have access to the mean free path. 
-					const float MFP = SubsurfaceOpacityMFP.Data;
-					const float Extinction = SubstrateShadowMFPToExtinction(MFP);
-
-					TransmissionParameters = InitTransmissionProfileParams();
-					TransmissionParameters.ExtinctionScale = Extinction;
-					
-					SubsurfaceSamplingMeanFreePath = 1.0;
-					BRANCH
-					if (TransmissionMeanFreePathType == TransmissionMeanFreePathType_MaxMeanFreePath)
-					{
-						SubsurfaceSamplingMeanFreePath = MFP;
-					}
+					const float RadiusScale = SubstrateSubSurfaceHeaderGetProfileRadiusScale(SSSHeader);
+					SubsurfaceSamplingMeanFreePath = SubstrateExtinctionToShadowMFP(TransmissionParameters.ExtinctionScale) * RadiusScale;
 				}
 			}
 
diff --git a/Engine/Shaders/Private/Substrate/SubstrateSubsurface.ush b/Engine/Shaders/Private/Substrate/SubstrateSubsurface.ush
--- a/Engine/Shaders/Private/Substrate/SubstrateSubsurface.ush
+++ b/Engine/Shaders/Private/Substrate/SubstrateSubsurface.ush
@@ -257,6 +257,12 @@ float SubstrateShadowMFPToExtinction(float MFP)
 	return Extinction;
 }
 
+float SubstrateExtinctionToShadowMFP(float Extinction)
+{
+	const float MFP = 1.0f / max(0.00001f, Extinction);
+	return MFP;
+}
+
 float SubstrateShadowColoredMFPToGreyScaleMFP(float3 MFP)
 {
 	return dot(MFP, (1.0f / 3.0f).xxx);
