Substrate - Toon
- Fixed reflection for flexible gbuffer path (env evaluate and importance sampling)
- Fixed eye roughness accessor.

#tests PC editor dx12
#rnx

#Synced-CL 50480041

#Synced-CL 50481140

diff --git a/Engine/Shaders/Private/Substrate/SubstrateEvaluation.ush b/Engine/Shaders/Private/Substrate/SubstrateEvaluation.ush
--- a/Engine/Shaders/Private/Substrate/SubstrateEvaluation.ush
+++ b/Engine/Shaders/Private/Substrate/SubstrateEvaluation.ush
@@ -1604,13 +1604,66 @@ FSubstrateImportanceSampleResult SubstrateImportanceSampleBSDF(FSubstrateBSDFCon
 
 		case SUBSTRATE_BSDF_TYPE_TOON:
 		{
-			const float3 DiffuseColor		= TOON_BASECOLOR(BSDFContext.BSDF);
+			const float3 DiffuseColor						= ComputeDiffuseAlbedo(TOON_BASECOLOR(BSDFContext.BSDF), TOON_METALLIC(BSDFContext.BSDF));
+			const float3 F0									= ComputeF0(TOON_SPECULAR(BSDFContext.BSDF), TOON_BASECOLOR(BSDFContext.BSDF), TOON_METALLIC(BSDFContext.BSDF));
+			const float  F90								= 1.0f;
+			const float SafeRoughness						= MakeRoughnessSafe(TOON_ROUGHNESS(BSDFContext.BSDF));
+
+			// Pick diffuse/specular lobe
+			float TermPDF = Out.Term == SHADING_TERM_DIFFUSE ? 1 : 0;
+			Out.Term = InTermMask;
+			if (Out.Term == (SHADING_TERM_DIFFUSE | SHADING_TERM_SPECULAR))
+			{
+				// Use (main) specular lobe as a heuristic to pick select between specular/diffuse lobe
+				const FBxDFEnergyTerms SpecEnergyTerms = ComputeGGXSpecEnergyTerms(SafeRoughness, BSDFContext.Context.NoV, F0, F90);
+
+				// Probability of picking diffuse lobe vs. specular lobe
+				TermPDF = SubstrateLobeSelectionProb(DiffuseColor * (1 - SpecEnergyTerms.E), SpecEnergyTerms.E);
+				Out.Term = SubstrateSelectLobe(E.x, TermPDF, SHADING_TERM_DIFFUSE, SHADING_TERM_SPECULAR);
+			}
 
-			float4 ImportanceSample	= CosineSampleHemisphere(E);
-			Out.L					= mul(ImportanceSample.xyz, BSDFContext.TangentBasis);
-			Out.PDF					= ImportanceSample.w;
-			Out.Term				= SHADING_TERM_DIFFUSE;
-			Out.Weight				= 1.0f;
+			// Pick specular lobe
+			float SpecularLobePDF = 1.f;
+			float SpecularRoughness = SafeRoughness;
+
+			// Sample Diffuse lobe
+			if (Out.Term == SHADING_TERM_DIFFUSE)
+			{
+				float4 ImportanceSample = CosineSampleHemisphere(E);
+				Out.L = mul(ImportanceSample.xyz, BSDFContext.TangentBasis);
+				Out.PDF = ImportanceSample.w;
+			}
+
+			// Sample Specular lobe
+			if (Out.Term == SHADING_TERM_SPECULAR)
+			{
+				float2 Alpha = Pow2(SpecularRoughness);
+
+				float4 GGXSample = ImportanceSampleVisibleGGX(E, Alpha, BSDFContext.TangentV);
+				const float HPDF = GGXSample.w;
+				const float3 H = mul(GGXSample.xyz, BSDFContext.TangentBasis);
+				const float VoH = saturate(dot(BSDFContext.V, H));
+
+				Out.L = 2 * dot(BSDFContext.V, H) * H - BSDFContext.V;
+				Out.PDF = RayPDFToReflectionRayPDF(VoH, HPDF);
+			}
+
+			// Evaluate lobes
+			BSDFContext.SubstrateUpdateBSDFContext(Out.L);
+			const FSubstrateEvaluateResult Evaluate = SubstrateEvaluateBSDF(BSDFContext, Settings);
+
+			// Add lobes mixtures
+			Out.PDF = 0;
+			Out.Weight = 0;
+			if (InTermMask & SHADING_TERM_DIFFUSE)
+			{
+				AddLobeWithMIS(Out.Weight, Out.PDF, Evaluate.DiffusePathValue / Evaluate.DiffusePDF, Evaluate.DiffusePDF, TermPDF);
+			}
+
+			if (InTermMask & SHADING_TERM_SPECULAR)
+			{
+				AddLobeWithMIS(Out.Weight, Out.PDF, Evaluate.SpecularPathValue / Evaluate.SpecularPDF, Evaluate.SpecularPDF, (1 - TermPDF) * SpecularLobePDF);
+			}
 			break;
 		}
 		break;
@@ -2011,7 +2064,7 @@ FSubstrateEnvLightResult SubstrateEvaluateForEnvLight(FSubstrateBSDFContext BSDF
 			SubstrateEnvLightResult.DiffuseNormal			= BSDFContext.N;
 			SubstrateEnvLightResult.bPostProcessSubsurface	= (BSDF_GETSSSTYPE(BSDFContext.BSDF) == SSS_TYPE_DIFFUSION || BSDF_GETSSSTYPE(BSDFContext.BSDF) == SSS_TYPE_DIFFUSION_PROFILE) && SUBSTRATE_SSS_MATERIAL_OVERRIDE;
 
-			float SafeRoughness = MakeRoughnessSafe(SLAB_ROUGHNESS(BSDFContext.BSDF));
+			float SafeRoughness = MakeRoughnessSafe(EYE_ROUGHNESS(BSDFContext.BSDF));
 
 			float3 EvalEnvBRDF = 0;
 			float3 DirectionalAlbedo_SpecularTransmission = 0;
@@ -2040,14 +2093,24 @@ FSubstrateEnvLightResult SubstrateEvaluateForEnvLight(FSubstrateBSDFContext BSDF
 		case SUBSTRATE_BSDF_TYPE_TOON:
 		{
 			const float3 DiffuseColor						= ComputeDiffuseAlbedo(TOON_BASECOLOR(BSDFContext.BSDF), TOON_METALLIC(BSDFContext.BSDF));
+			const float3 F0									= ComputeF0(TOON_SPECULAR(BSDFContext.BSDF), TOON_BASECOLOR(BSDFContext.BSDF), TOON_METALLIC(BSDFContext.BSDF));
+			const float  F90								= 1.0f;
+			const float SafeRoughness						= MakeRoughnessSafe(TOON_ROUGHNESS(BSDFContext.BSDF));
+
+			const FBxDFEnergyTermsRGB SpecularEnergyTerms	= ComputeGGXSpecEnergyTermsRGB_Env(SafeRoughness, BSDFContext.Context.NoV, F0, F90);
+			const FBxDFEnergyTermsRGB DiffuseEnergyTerms	= ComputeDiffuseEnergyTermsRGB(SafeRoughness, BSDFContext.Context.NoV);
+			const float3 EvalEnvBRDF						= SpecularEnergyTerms.E;
 
 			SubstrateEnvLightResult.DiffuseNormal			= BSDFContext.N;
 			SubstrateEnvLightResult.bPostProcessSubsurface	= false;
 
-			SubstrateEnvLightResult.DiffuseWeight			= DiffuseColor;
+			SubstrateEnvLightResult.DiffuseWeight			= DiffuseColor * DiffuseEnergyTerms.E;
 			SubstrateEnvLightResult.DiffuseColor			= DiffuseColor;
 
-			// TOON_TODO Specular support.
+			SubstrateEnvLightResult.SpecularDirection		= SubstrateGetOffSpecularPeakReflectionDir(BSDFContext.N, BSDFContext.R, SafeRoughness);
+			SubstrateEnvLightResult.SpecularWeight			= EvalEnvBRDF;
+			SubstrateEnvLightResult.SpecularSafeRoughness	= SafeRoughness;
+			SubstrateEnvLightResult.SpecularColor			= F0;
 		}
 		break;
 
