[TSR] Add antiflickering path for thin geometry detection. So the residual flickering can be captured by the moire error. Default on. use `r.TSR.ThinGeometryDetection.AntiFlickering 0` to disable it.

#tests TSR level shared across the team. The flickering improves a lot
#Synced-CL 50430604

diff --git a/Engine/Shaders/Private/TemporalSuperResolution/TSRRejectShading.usf b/Engine/Shaders/Private/TemporalSuperResolution/TSRRejectShading.usf
--- a/Engine/Shaders/Private/TemporalSuperResolution/TSRRejectShading.usf
+++ b/Engine/Shaders/Private/TemporalSuperResolution/TSRRejectShading.usf
@@ -461,7 +461,9 @@ void RejectShadingCS(
 	{
 		Debug[DebugId].SetAllElements(0.0);
 	}
-		
+	#if CONFIG_THIN_GEOMETRY_DETECTION
+	bool bApplyThinGeometryRelaxation = false;
+	#endif	
 	// Analyses the moire pattern overtime.
 	tsr_tensor_half  MoireError;
 	#if CONFIG_MOIRE_DETECTION
@@ -491,7 +493,18 @@ void RejectShadingCS(
 				/* out */ PrevMoireHistory,
 				/* out */ bIsStatic);
 		}
-	
+		
+		#if CONFIG_THIN_GEOMETRY_DETECTION
+		tsr_tensor_half ThinGeometryRelaxationFactor = tsr_tensor_half::Const(0);
+		if (ThinGeometryAffectMoireError)
+		{
+			tsr_tensor_short2 PixelPos = ComputePixelPos();
+			tsr_tensor_bool IsThinGeometryRegion;
+			FetchThinGeometryEdge(PixelPos, ThinGeometryRelaxationFactor, IsThinGeometryRegion);
+			bApplyThinGeometryRelaxation = AnyThread(AnyComponent(ThinGeometryRelaxationFactor > tsr_tensor_half::Const(0)));
+		}
+		#endif
+
 		tsr_tensor_half4 MoireHistory;
 		ComputeMoireError(
 			bIsDisoccluded,
@@ -500,6 +513,10 @@ void RejectShadingCS(
 			OriginalOpaqueInput,
 			OriginalTranslucencyInput,
 			PrevMoireHistory,
+		#if CONFIG_THIN_GEOMETRY_DETECTION
+			ThinGeometryRelaxationFactor,
+			bApplyThinGeometryRelaxation,
+		#endif
 			/* out */ MoireHistory,
 			/* out */ MoireError);
 
@@ -625,13 +642,20 @@ void RejectShadingCS(
 	tsr_tensor_half RejectionClampBlend = tsr_tensor_half::Const(1.0);
 	tsr_tensor_bool bAliasingIsVisible = tsr_tensor_bool::Const(false);
 #if CONFIG_THIN_GEOMETRY_DETECTION
-	bool bApplyThinGeometryRelaxation = false;
 	tsr_tensor_half ThinGeometryRelaxationFactor = tsr_tensor_half::Const(0);
 	{
 		tsr_tensor_short2 PixelPos = ComputePixelPos();
 		tsr_tensor_bool IsThinGeometryRegion;
-		FetchThinGeometryEdge(PixelPos, ThinGeometryRelaxationFactor, IsThinGeometryRegion);
-		bApplyThinGeometryRelaxation = AnyThread(AnyComponent(ThinGeometryRelaxationFactor > tsr_tensor_half::Const(0)));
+		#if CONFIG_MOIRE_DETECTION
+			BRANCH
+			if (bApplyThinGeometryRelaxation)
+			{
+				FetchThinGeometryEdge(PixelPos, ThinGeometryRelaxationFactor, IsThinGeometryRegion);
+			}
+		#else
+			FetchThinGeometryEdge(PixelPos, ThinGeometryRelaxationFactor, IsThinGeometryRegion);
+			bApplyThinGeometryRelaxation = AnyThread(AnyComponent(ThinGeometryRelaxationFactor > tsr_tensor_half::Const(0)));
+		#endif
 	}
 #endif
 	#if CONFIG_HISTORY_RESURRECTION
@@ -684,7 +708,7 @@ void RejectShadingCS(
 			/* HistoryC2 = */  ClampedHistory,
 			/* MoireError = */ MoireError,
 #if CONFIG_THIN_GEOMETRY_DETECTION
-			ThinGeometryRelaxationFactor ,
+			ThinGeometryRelaxationFactor,
 			bApplyThinGeometryRelaxation,
 #endif
 			/* bEnableMoireError = */ CONFIG_MOIRE_DETECTION != 0 && bEnableFlickeringHeuristic != 0,
diff --git a/Engine/Shaders/Private/TemporalSuperResolution/TSRShadingAnalysis.ush b/Engine/Shaders/Private/TemporalSuperResolution/TSRShadingAnalysis.ush
--- a/Engine/Shaders/Private/TemporalSuperResolution/TSRShadingAnalysis.ush
+++ b/Engine/Shaders/Private/TemporalSuperResolution/TSRShadingAnalysis.ush
@@ -30,6 +30,7 @@ float ExposureOffsetFactor;
 float FlickeringFramePeriod;
 uint bPassthroughAlpha;
 uint RejectShadingMode;
+int ThinGeometryAffectMoireError;
 
 // Cache the threshold at which the variation count should quick in
 float2 TotalVariationCountInvThreshold;
@@ -218,12 +219,17 @@ void ComputeMoireError(
 	tsr_tensor_halfC OriginalOpaqueInput,
 	tsr_tensor_half4 OriginalTranslucencyInput,
 	tsr_tensor_half4 PrevMoireHistory,
+#if CONFIG_THIN_GEOMETRY_DETECTION
+	tsr_tensor_half HistoryRelaxationFactor,
+	const bool bEnableThinGeometry,
+#endif
 	out tsr_tensor_half4 OutMoireHistory,
 	out tsr_tensor_half  OutMoireError)
 {
 	const tsr_half MinBlendFinal = tsr_half(0.05);
 	const tsr_half MoireEncodingError = tsr_half(rcp(127.0));
 	const tsr_half MaxPrevTotalVariationCount = tsr_half(20.0);
+     const tsr_tensor_half ValidGeometryGradientThreshold = tsr_tensor_half::Const(tsr_half(0.00007));
 
 	// Input luminance to monitor flickering from.
 	tsr_tensor_half InputLuma = GCSToSMCS(MoireInput);
@@ -295,6 +301,38 @@ void ComputeMoireError(
 			FilteredBoxMin = Min[0];
 			FilteredBoxMax = Max[0];
 		
+		#if CONFIG_THIN_GEOMETRY_DETECTION
+			// The scene can still flicker with thin geometry detection. Add thin geometry logic so that the moire error can
+			// detect the residual flickering.
+			if (bEnableThinGeometry)
+			{
+				// VGPR optimization for LerpClamp3x3()
+				HistoryRelaxationFactor = select(HistoryRelaxationFactor > tsr_tensor_half::Const(tsr_half(1.0f / 127.0f)), HistoryRelaxationFactor, tsr_tensor_half::Const(0));
+				tsr_tensor_half LerpFactorVector = HistoryRelaxationFactor;
+		
+				PLATFORM_SPECIFIC_ISOLATE
+				{
+					tsr_tensor_half ToClampBoundaryMin;
+					tsr_tensor_half ToClampBoundaryMax;
+					MinMax3x3(FilteredHistory, /*out*/ToClampBoundaryMin, /*out*/ToClampBoundaryMax);
+
+					FilteredBoxMin = lerp(FilteredBoxMin, ToClampBoundaryMin, LerpFactorVector);
+					FilteredBoxMax = lerp(FilteredBoxMax, ToClampBoundaryMax, LerpFactorVector);
+				}
+		
+				tsr_tensor_half LerpClampedHistoryInput = fastClamp(FilteredHistory, FilteredBoxMin, FilteredBoxMax);
+
+				{
+					tsr_tensor_half ToClampBoundaryMin;
+					tsr_tensor_half ToClampBoundaryMax;
+					MinMax3x3(LerpClampedHistoryInput, /*out*/ToClampBoundaryMin, /*out*/ToClampBoundaryMax);
+
+					FilteredBoxMin = select(LerpFactorVector > tsr_tensor_half::Const(0), ToClampBoundaryMin, FilteredBoxMin);
+					FilteredBoxMax = select(LerpFactorVector > tsr_tensor_half::Const(0), ToClampBoundaryMax, FilteredBoxMax);
+				}
+			}
+		#endif
+
 			InputC2BoxSize = Max[1] - Min[1];
 		}
 
@@ -374,6 +412,18 @@ void ComputeMoireError(
 	tsr_tensor_bool bGradientsAreSameSign = CurrentGradient * PrevGradient > tsr_tensor_half::Const(0.0);
 	tsr_tensor_bool bGradientsAreWithinError = or(abs(CurrentGradient) < MoireEncodingErrors, abs(PrevGradient) < MoireEncodingErrors);
 	//bool bCurrentGradientSubstentialEnough = abs(LocalCurrentGradient) > View.GeneralPurposeTweak * abs(LocalPrevGradient); // - 2.0 * MoireEncodingError;
+	
+	#if CONFIG_THIN_GEOMETRY_DETECTION
+	tsr_tensor_bool bGradientsAreThinGeometryRegion = tsr_tensor_bool::Const(false);
+	if (bEnableThinGeometry)
+	{
+		tsr_tensor_bool bIsPotentialGeometryFlickering = HistoryRelaxationFactor > tsr_tensor_half::Const(0.0);
+		tsr_tensor_bool bAnyGradient = or(abs(CurrentGradient) > ValidGeometryGradientThreshold,abs(PrevGradient) > ValidGeometryGradientThreshold);
+		bGradientsAreThinGeometryRegion = and(and(bIsPotentialGeometryFlickering, bAnyGradient), bGradientsAreWithinError);
+		bGradientsAreWithinError = select(bGradientsAreThinGeometryRegion, tsr_tensor_bool::Const(false), bGradientsAreWithinError);
+	}
+	#endif
+	
 	tsr_tensor_bool bCameraCuts = tsr_tensor_bool::Const(bool(bCameraCut));
 	tsr_tensor_half IsFlicker = select(or(
 			or(bGradientsAreSameSign,
@@ -381,7 +431,6 @@ void ComputeMoireError(
 			//!bCurrentGradientSubstentialEnough ||
 			bCameraCuts), tsr_tensor_half::Const(0.0), tsr_tensor_half::Const(tsr_half(1.0)));
 
-
 	// Compute the total variation of the gradient over time.
 	tsr_tensor_half GradientVariation = min(abs(PrevGradient), abs(CurrentGradient)) * IsFlicker;
 
@@ -414,15 +463,23 @@ void ComputeMoireError(
 	tsr_tensor_half QuantizedNewTotalVariationCount = floor(NewTotalVariationCount * (tsr_half(255.0) / MaxPrevTotalVariationCount)) * (MaxPrevTotalVariationCount / tsr_half(255.0));
 	NewTotalVariation = NewTotalVariation * QuantizedNewTotalVariationCount * SafeRcp(NewTotalVariationCount);
 
-	// Compute the final luminance 
+	// Compute the final luminance
+	#if CONFIG_THIN_GEOMETRY_DETECTION
+		const tsr_tensor_half TotalVariationCountInvThresholdDefault = tsr_tensor_half::Const(tsr_half(TotalVariationCountInvThreshold[0]));
+		const tsr_tensor_half TotalVariationCountInvThresholdThinGeo = tsr_tensor_half::Const(tsr_half(TotalVariationCountInvThreshold[1]));
+		
+		tsr_tensor_half TotalVariationCountInvThresholds = select(bGradientsAreThinGeometryRegion, TotalVariationCountInvThresholdThinGeo, TotalVariationCountInvThresholdDefault);
+	#else
+		tsr_tensor_half TotalVariationCountInvThresholds = tsr_tensor_half::Const(tsr_half(TotalVariationCountInvThreshold[0]));
+	#endif
+
 	#if 1
-		tsr_tensor_half CountFadeIn = saturate(NewTotalVariationCount * tsr_half(TotalVariationCountInvThreshold[0]) - tsr_half(0.5));
+		tsr_tensor_half CountFadeIn = saturate(NewTotalVariationCount * TotalVariationCountInvThresholds - tsr_half(0.5));
 	#else
-		tsr_tensor_half CountFadeIn = saturate(NewTotalVariationCount - rcp(tsr_half(TotalVariationCountInvThreshold[0])));
+		tsr_tensor_half CountFadeIn = saturate(NewTotalVariationCount - rcp(TotalVariationCountInvThresholds));
 	#endif
 
-	OutMoireError = (abs(NewTotalVariation * SafeRcp(NewTotalVariationCount)) + NewTotalVariationCount * MoireEncodingError) *
-CountFadeIn;
+	OutMoireError = (abs(NewTotalVariation * SafeRcp(NewTotalVariationCount)) + NewTotalVariationCount * MoireEncodingError) * CountFadeIn;
 	
 	OutMoireHistory = Concatenate( 
 						Concatenate(FinalFlickeringHistory, NewGradient * tsr_half(127.0 / 255.0) + tsr_half(127.0 / 255.0)),
diff --git a/Engine/Source/Runtime/Renderer/Private/PostProcess/TemporalSuperResolution.cpp b/Engine/Source/Runtime/Renderer/Private/PostProcess/TemporalSuperResolution.cpp
--- a/Engine/Source/Runtime/Renderer/Private/PostProcess/TemporalSuperResolution.cpp
+++ b/Engine/Source/Runtime/Renderer/Private/PostProcess/TemporalSuperResolution.cpp
@@ -373,6 +373,11 @@ TAutoConsoleVariable<float> CVarTSRThinGeometryDetectionErrorMultiplier(
 	TEXT("(Default = 200)."),
 	ECVF_RenderThreadSafe);
 
+TAutoConsoleVariable<int32> CVarTSRThinGeometryDetectionAntiFlickering(
+	TEXT("r.TSR.ThinGeometryDetection.AntiFlickering"), 1,
+	TEXT("Thin geometry detection can sometimes fail. Whether the anti-flickering logic considers the thin geometry detection pass."),
+	ECVF_RenderThreadSafe);
+
 #if !UE_BUILD_OPTIMIZED_SHOWFLAGS
 
 TAutoConsoleVariable<int32> CVarTSRVisualize(
@@ -968,6 +973,7 @@ class FTSRRejectShadingCS : public FTSRConvolutionNetworkShader
 		SHADER_PARAMETER(float, ExposureOffsetFactor)
 		SHADER_PARAMETER(int32, RejectShadingMode)
 		SHADER_PARAMETER(FVector2f, TotalVariationCountInvThreshold)
+		SHADER_PARAMETER(int32, ThinGeometryAffectMoireError)
 
 		SHADER_PARAMETER_RDG_TEXTURE(Texture2D, InputTexture)
 		SHADER_PARAMETER_RDG_TEXTURE_SRV(Texture2D, InputMoireLumaTexture)
@@ -2960,7 +2966,10 @@ FDefaultTemporalUpscaler::FOutputs AddTemporalSuperResolutionPasses(
 		PassParameters->bEnableFlickeringHeuristic = FlickeringFramePeriod > 0.0f;
 		PassParameters->bPassthroughAlpha = IsPrimitiveAlphaHoldoutEnabled(View);
 		PassParameters->RejectShadingMode = CVarTSRShadingRejectionMode.GetValueOnRenderThread();
-		PassParameters->TotalVariationCountInvThreshold = FVector2f((1.0 - FMath::Pow(1.0 - 0.05 /*MinBlendFinal*/, FlickeringFramePeriod)));
+		PassParameters->TotalVariationCountInvThreshold = FVector2f(
+			(1.0 - FMath::Pow(1.0 - 0.05 /*MinBlendFinal*/, FlickeringFramePeriod)),
+			(1.0 - FMath::Pow(1.0 - 0.05 /*MinBlendFinal*/, 3.0f * FlickeringFramePeriod))); // Total variation count for thin geometry region.
+		PassParameters->ThinGeometryAffectMoireError = CVarTSRThinGeometryDetectionAntiFlickering.GetValueOnRenderThread();
 
 		PassParameters->InputTexture = PassInputs.SceneColor.Texture;
 		if (PassInputs.FlickeringInputTexture.IsValid() && FlickeringFramePeriod > 0.0f)
