Fix mising virtual texture feedback when feeding Thin Translucent.
FMaterialPixelParameters wasn't inout so that the any feedback returned by GetThinTranslucentMaterialOutput0() was being lost.

#Synced-CL 50484185

diff --git a/Engine/Shaders/Private/ThinTranslucentCommon.ush b/Engine/Shaders/Private/ThinTranslucentCommon.ush
--- a/Engine/Shaders/Private/ThinTranslucentCommon.ush
+++ b/Engine/Shaders/Private/ThinTranslucentCommon.ush
@@ -10,7 +10,7 @@
 void AccumulateThinTranslucentModel(inout float3 DualBlendSurfaceLuminancePostCoverage,
 								inout float3 DualBlendSurfaceTransmittancePreCoverage,
 								inout float DualBlendSurfaceCoverage,
-								FMaterialPixelParameters MaterialParams,
+								inout FMaterialPixelParameters MaterialParams,
 								FGBufferData GBuffer,
 								float3 DiffuseColor,
 								float3 SpecularColor,
