Substrate - Toon - GI scale settings

#tests pc editor dx12 with toon material.
#rnx

#Synced-CL 50481895

#Synced-CL 50491562

diff --git a/Engine/Shaders/Private/DiffuseIndirectComposite.usf b/Engine/Shaders/Private/DiffuseIndirectComposite.usf
--- a/Engine/Shaders/Private/DiffuseIndirectComposite.usf
+++ b/Engine/Shaders/Private/DiffuseIndirectComposite.usf
@@ -608,6 +608,21 @@ void MainPS(
 					DiffuseColor += ClothFuzz * GBuffer.CustomData.a;
 				}
 
+			#if SUBSTRATE_ENABLED
+				// This path is not calling SubstrateEvaluateForEnvLight so we need to add that specialised code here.
+				BRANCH
+				if (GBuffer.ShadingModelID == SHADINGMODELID_SUBSTRATE_TOON)
+				{
+					float ToonProfileId = 0.0f;
+					float2 PatternUVs = 0.0;
+					UnpackToonCustomData(PackRGBA8(GBuffer.CustomData), PatternUVs, ToonProfileId);
+
+					FToonProfileGISettings ToonProfileGISettings = GetToonProfileGISettings(ToonProfileId);
+					DiffuseColor *= ToonProfileGISettings.DiffuseIndirectScale;
+					SpecularColor *= ToonProfileGISettings.SpecularIndirectScale;
+				}
+			#endif
+
 			#if USE_ENERGY_CONSERVATION
 				FBxDFEnergyTerms SpecularEnergyTerms = ComputeGGXSpecEnergyTerms(GBuffer.Roughness, NoV, SpecularColor);
 				float3 EnvBRDFValue = SpecularEnergyTerms.E; // EnvBRDF accounting for multiple scattering when enabled
diff --git a/Engine/Shaders/Private/Substrate/SubstrateEvaluation.ush b/Engine/Shaders/Private/Substrate/SubstrateEvaluation.ush
--- a/Engine/Shaders/Private/Substrate/SubstrateEvaluation.ush
+++ b/Engine/Shaders/Private/Substrate/SubstrateEvaluation.ush
@@ -2101,14 +2101,16 @@ FSubstrateEnvLightResult SubstrateEvaluateForEnvLight(FSubstrateBSDFContext BSDF
 			const FBxDFEnergyTermsRGB DiffuseEnergyTerms	= ComputeDiffuseEnergyTermsRGB(SafeRoughness, BSDFContext.Context.NoV);
 			const float3 EvalEnvBRDF						= SpecularEnergyTerms.E;
 
+			FToonProfileGISettings ToonProfileGISettings	= GetToonProfileGISettings(TOON_PROFILEID(BSDFContext.BSDF));
+
 			SubstrateEnvLightResult.DiffuseNormal			= BSDFContext.N;
 			SubstrateEnvLightResult.bPostProcessSubsurface	= false;
 
-			SubstrateEnvLightResult.DiffuseWeight			= DiffuseColor * DiffuseEnergyTerms.E;
+			SubstrateEnvLightResult.DiffuseWeight			= DiffuseColor * DiffuseEnergyTerms.E * ToonProfileGISettings.DiffuseIndirectScale;
 			SubstrateEnvLightResult.DiffuseColor			= DiffuseColor;
 
 			SubstrateEnvLightResult.SpecularDirection		= SubstrateGetOffSpecularPeakReflectionDir(BSDFContext.N, BSDFContext.R, SafeRoughness);
-			SubstrateEnvLightResult.SpecularWeight			= EvalEnvBRDF;
+			SubstrateEnvLightResult.SpecularWeight			= EvalEnvBRDF * ToonProfileGISettings.SpecularIndirectScale;
 			SubstrateEnvLightResult.SpecularSafeRoughness	= SafeRoughness;
 			SubstrateEnvLightResult.SpecularColor			= F0;
 		}
diff --git a/Engine/Shaders/Private/ToonProfile.usf b/Engine/Shaders/Private/ToonProfile.usf
--- a/Engine/Shaders/Private/ToonProfile.usf
+++ b/Engine/Shaders/Private/ToonProfile.usf
@@ -20,6 +20,9 @@ float SpecularDitherPatternSize;
 float ShadowHatchingPatternSize;
 float ShadowHatchingPatternStrength;
 
+float DiffuseIndirectScale;
+float SpecularIndirectScale;
+
 RWTexture2DArray<float4>	ToonProfileTexture;
 
 Texture2D<float4>			DiffuseDitherPattern;
@@ -60,14 +63,18 @@ void ToonProfileUpdateAtlasTileCS(uint3 DispatchThreadId : SV_DispatchThreadID)
 		}
 		else if (PixelPos.y == TOON_PROFILE_PARAMETERS_Y)
 		{
-			if (PixelPos.x == 0)
+			if (PixelPos.x == TOON_PROFILE_EVAL_PARAMS_X)
 			{
 				ToonProfileTexture[uint3(PixelPos, ToonProfileIndex)] = float4(DiffuseDitherPatternStrength, DiffuseDitherPatternSize, SpecularDitherPatternStrength, SpecularDitherPatternSize);
 			}
-			else if (PixelPos.x == 1)
+			else if (PixelPos.x == TOON_PROFILE_SHADOW_PARAMS_X)
 			{
 				ToonProfileTexture[uint3(PixelPos, ToonProfileIndex)] = float4(ToonProfileShadowExtinctionCoefficientInInvCentimeter, ShadowHatchingPatternSize, ShadowHatchingPatternStrength, 0.0f);
 			}
+			else if (PixelPos.x == TOON_PROFILE_GI_PARAMS_X)
+			{
+				ToonProfileTexture[uint3(PixelPos, ToonProfileIndex)] = float4(DiffuseIndirectScale, SpecularIndirectScale, 0.0f, 0.0f);
+			}
 			else
 			{
 				ToonProfileTexture[uint3(PixelPos, ToonProfileIndex)] = DefaultColor;
diff --git a/Engine/Shaders/Private/ToonProfileCommon.ush b/Engine/Shaders/Private/ToonProfileCommon.ush
--- a/Engine/Shaders/Private/ToonProfileCommon.ush
+++ b/Engine/Shaders/Private/ToonProfileCommon.ush
@@ -4,6 +4,8 @@
 
 #include "/Engine/Shared/ToonProfileDefinitions.h"
 
+////////////////////////////////////////////////////////////////////////////////////////////////////
+
 float3 EvaluateToonProfileRamp(in float ToonProfileId, in float NoL, in float TOON_PROFILE_RAMP)
 {
 	const float3 ToonProfileCoord = float3(
@@ -29,6 +31,8 @@ float EvaluateToonProfileShadowRamp(in float ToonProfileId, in float Transmittan
 	return EvaluateToonProfileRamp(ToonProfileId, Transmittance, TOON_PROFILE_SHADOW_RAMP_Y).r;
 }
 
+////////////////////////////////////////////////////////////////////////////////////////////////////
+
 struct FToonProfileEvaluateSettings
 {
 	float DiffuseDitherPatternStrength;
@@ -39,8 +43,9 @@ struct FToonProfileEvaluateSettings
 
 FToonProfileEvaluateSettings GetToonProfileEvaluateSettings(in float ToonProfileId)
 {
-	const float3 ToonProfileCoord = float3(0.5f * View.ToonProfileTextureSizeAndInvSize.z,
-		(TOON_PROFILE_PARAMETERS_Y + 0.5f) * View.ToonProfileTextureSizeAndInvSize.w,		// Center of the row of texel
+	const float3 ToonProfileCoord = float3(
+		(TOON_PROFILE_EVAL_PARAMS_X + 0.5f) * View.ToonProfileTextureSizeAndInvSize.z,
+		(TOON_PROFILE_PARAMETERS_Y  + 0.5f) * View.ToonProfileTextureSizeAndInvSize.w,		// Center of the row of texel
 		ToonProfileId																		// EvaluateToonProfileDiffuseDither / Slice index
 	);
 	float4 Data = View.ToonProfileTexture.SampleLevel(View.ToonProfileSampler, ToonProfileCoord, 0);
@@ -99,6 +104,8 @@ float EvaluateToonProfileSpecularDither(in float ToonProfileId, in float2 Patter
 	return 0.0f;
 }
 
+////////////////////////////////////////////////////////////////////////////////////////////////////
+
 struct FToonProfileShadowSettings
 {
 	float ShadowExtinctionCoefficientInInvCentimeter;	// Needed because transmittance evaluation code from ShadowProjectionPixelShader.usf uses centimeters.
@@ -109,8 +116,9 @@ struct FToonProfileShadowSettings
 
 FToonProfileShadowSettings GetToonProfileShadowSettings(in float ToonProfileId)
 {
-	const float3 ToonProfileCoord = float3(1.5f * View.ToonProfileTextureSizeAndInvSize.z,
-		(TOON_PROFILE_PARAMETERS_Y + 0.5f) * View.ToonProfileTextureSizeAndInvSize.w,		// Center of the row of texel
+	const float3 ToonProfileCoord = float3(
+		(TOON_PROFILE_SHADOW_PARAMS_X + 0.5f) * View.ToonProfileTextureSizeAndInvSize.z,
+		(TOON_PROFILE_PARAMETERS_Y    + 0.5f) * View.ToonProfileTextureSizeAndInvSize.w,	// Center of the row of texel
 		ToonProfileId																		// Slice index
 	);
 	float4 Data = View.ToonProfileTexture.SampleLevel(View.ToonProfileSampler, ToonProfileCoord, 0).xyzw;
@@ -165,3 +173,30 @@ float EvaluateToonProfileShadowHatching(in float ToonProfileId, in float2 Patter
 	return 1.0f;
 }
 
+////////////////////////////////////////////////////////////////////////////////////////////////////
+
+struct FToonProfileGISettings
+{
+	float DiffuseIndirectScale;
+	float SpecularIndirectScale;
+};
+
+FToonProfileGISettings GetToonProfileGISettings(in float ToonProfileId)
+{
+	const float3 ToonProfileCoord = float3(
+		(TOON_PROFILE_GI_PARAMS_X  + 0.5f) * View.ToonProfileTextureSizeAndInvSize.z,
+		(TOON_PROFILE_PARAMETERS_Y + 0.5f) * View.ToonProfileTextureSizeAndInvSize.w,		// Center of the row of texel
+		ToonProfileId																		// Slice index
+	);
+	float4 Data = View.ToonProfileTexture.SampleLevel(View.ToonProfileSampler, ToonProfileCoord, 0).xyzw;
+
+	FToonProfileGISettings Ret;
+	Ret.DiffuseIndirectScale = Data.x;
+	Ret.SpecularIndirectScale = Data.y;
+	return Ret;
+}
+
+
+
+
+
diff --git a/Engine/Shaders/Shared/ToonProfileDefinitions.h b/Engine/Shaders/Shared/ToonProfileDefinitions.h
--- a/Engine/Shaders/Shared/ToonProfileDefinitions.h
+++ b/Engine/Shaders/Shared/ToonProfileDefinitions.h
@@ -8,6 +8,10 @@
 #define TOON_PROFILE_PARAMETERS_Y		3
 #define TOON_PROFILE_TEXTURES_Y			4
 
+#define TOON_PROFILE_EVAL_PARAMS_X		0
+#define TOON_PROFILE_SHADOW_PARAMS_X	1
+#define TOON_PROFILE_GI_PARAMS_X		2
+
 #define TOON_PROFILE_TEXTURE_DITHERPATTERN_IDX		0
 #define TOON_PROFILE_TEXTURE_SHADOWHATCHING_IDX		1
 #define TOON_PROFILE_TEXTURE_COUNT					2
diff --git a/Engine/Source/Runtime/Engine/Classes/Engine/ToonProfile.h b/Engine/Source/Runtime/Engine/Classes/Engine/ToonProfile.h
--- a/Engine/Source/Runtime/Engine/Classes/Engine/ToonProfile.h
+++ b/Engine/Source/Runtime/Engine/Classes/Engine/ToonProfile.h
@@ -104,6 +104,18 @@ struct FToonProfileStruct
 	UPROPERTY(Category = "Shadow", EditAnywhere, meta = (UIMin = 0.0, UIMax = 1.0, ClampMin = 0.0))
 	float ShadowHatchingPatternStrength;
 
+	/**
+	 * The scale the diffuse global illumination contribution.
+	 */
+	UPROPERTY(Category = "Global Illumination", EditAnywhere, meta = (DisplayName = "Diffuse Indirect Scale", UIMin = 0.0, UIMax = 1.0, ClampMin = 0.0))
+	float DiffuseIndirectScale;
+
+	/**
+	 * The scale the specular global illumination contribution.
+	 */
+	UPROPERTY(Category = "Global Illumination", EditAnywhere, meta = (DisplayName = "Specular Indirect Scale", UIMin = 0.0, UIMax = 1.0, ClampMin = 0.0))
+	float SpecularIndirectScale;
+
 	FToonProfileStruct();
 
 	void Invalidate()
diff --git a/Engine/Source/Runtime/Engine/Private/Rendering/ToonProfile.cpp b/Engine/Source/Runtime/Engine/Private/Rendering/ToonProfile.cpp
--- a/Engine/Source/Runtime/Engine/Private/Rendering/ToonProfile.cpp
+++ b/Engine/Source/Runtime/Engine/Private/Rendering/ToonProfile.cpp
@@ -80,6 +80,8 @@ public:
 		SHADER_PARAMETER(float, SpecularDitherPatternSize)
 		SHADER_PARAMETER(float, ShadowHatchingPatternSize)
 		SHADER_PARAMETER(float, ShadowHatchingPatternStrength)
+		SHADER_PARAMETER(float, DiffuseIndirectScale)
+		SHADER_PARAMETER(float, SpecularIndirectScale)
 
 		SHADER_PARAMETER(float, ToonProfileShadowExtinctionCoefficientInInvCentimeter)
 
@@ -409,6 +411,9 @@ IPooledRenderTarget* FToonProfileTextureManager::GetAtlasTexture(FRDGBuilder& Gr
 			PassParameters->ShadowHatchingPatternSize = Data.ShadowHatchingPatternSize;
 			PassParameters->ShadowHatchingPatternStrength = bHasShadowHatchingPatternTexture ? Data.ShadowHatchingPatternStrength : 0.0f;
 
+			PassParameters->DiffuseIndirectScale = Data.DiffuseIndirectScale;
+			PassParameters->SpecularIndirectScale = Data.SpecularIndirectScale;
+
 			auto CreateBufferFRomCurve = [&](FRuntimeCurveLinearColor ColorCurve)
 			{
 					const uint32 ComponentsPerElement = 4;
@@ -534,6 +539,9 @@ FToonProfileStruct::FToonProfileStruct()
 	ShadowHatchingPatternSize = 1.0f;
 	ShadowHatchingPatternStrength = 1.0f;
 
+	DiffuseIndirectScale = 1.0f;
+	SpecularIndirectScale = 1.0f;
+
 	AddPointToCurve(ShadowRamp, 0.00, 0.0);
 	AddPointToCurve(ShadowRamp, 1.00, 1.0);
 	
