Substrate - toon - moved toon shadow hatching out from shadow pass into evaluate.

#rnx
#tests PC editor
#Synced-CL 50491562

diff --git a/Engine/Shaders/Private/ShadowProjectionPixelShader.usf b/Engine/Shaders/Private/ShadowProjectionPixelShader.usf
--- a/Engine/Shaders/Private/ShadowProjectionPixelShader.usf
+++ b/Engine/Shaders/Private/ShadowProjectionPixelShader.usf
@@ -305,13 +305,10 @@ void Main(
 			const FSubstrateSubsurfaceHeader SSSHeader = SubstrateLoadSubsurfaceHeader(Substrate.MaterialTextureArray, Substrate.FirstSliceStoringSubstrateSSSData, SVPos.xy);
 			const uint SSSType = SubstrateSubSurfaceHeaderGetSSSType(SSSHeader);
 			#endif
-			FToonProfileShadowSettings ToonProfileShadowSettings = (FToonProfileShadowSettings)0;
 			#endif // SUBSTRATE_ENABLED
 
 			bool bApplyPCFOverBlurCorrection = true; // Set to false when the transmittance from extinction is correctly respected. Toon & TransmitttanceShadow disables that to have a physically based non squared transmittance curve.
 			float Density = 1;
-			float ToonProfileId = 0.0f;
-			float2 PatternUVs = ScreenUV;
 			if (!bIsEyeOrHair)
 			{
 			#if SUBSTRATE_GBUFFER_FORMAT==1
@@ -320,17 +317,6 @@ void Main(
 					bApplyPCFOverBlurCorrection = false;
 					const float3 MFPcm = SubstrateSubSurfaceHeaderGetTransmittanceShadowMFPInMeters(SSSHeader) * METER_TO_CENTIMETER;
 					Density = 1.0f / max(0.000001f, MFPcm.x); // Recover extinction coefficient as 1/cm from MFP in cm.
-
-					// In this case we also need to check if this is for a toon profile and load profile ID, PatternUVs and Shadow settings.
-					// SSS_TYPE_TRANSMITTANCE_SHADOW was added to try and be general. But it is only used by toon shading today. We could remove it and only load data from the toon profile instead, as done for SUBSTRATE_GBUFFER_FORMAT==0.
-					if (bIsToon)
-					{
-						FSubstrateBSDF BSDF	= UnpackSubstrateBSDFIn(Substrate.MaterialTextureArray, SubstrateAddressing, SubstrateHeader);
-						PatternUVs			= TOON_PATTERNUVS(BSDF);
-						ToonProfileId		= TOON_PROFILEID(BSDF);
-
-						ToonProfileShadowSettings = GetToonProfileShadowSettings(ToonProfileId);
-					}
 				}
 				else if (SSSType != SSS_TYPE_NONE)
 				{
@@ -344,10 +330,12 @@ void Main(
 				{
 					bApplyPCFOverBlurCorrection = false;
 
+					float ToonProfileId = 0.0f;
+					float2 PatternUVs = ScreenUV;
 					UnpackToonCustomData(PackRGBA8(GBufferData.CustomData), PatternUVs, ToonProfileId);
 
 					// Density is extinction as 1/cm (not 1/m) to match with the transmittance shadow evaluation from ManualPCF below.
-					ToonProfileShadowSettings = GetToonProfileShadowSettings(ToonProfileId);
+					FToonProfileShadowSettings ToonProfileShadowSettings = GetToonProfileShadowSettings(ToonProfileId);
 					Density = ToonProfileShadowSettings.ShadowExtinctionCoefficientInInvCentimeter;
 				}
 				else
@@ -393,14 +381,6 @@ void Main(
 				// This branch is taken by Subsurface shading model (Wrap for Substrate) or the toon shading model. 
 				// => SSSTransmittanceOrOpticalDepth stores the result of ManualPCF() evaluating Transmittance.
 				SSSTransmittanceOrOpticalDepth = ManualPCF(ShadowPosition.xy, Settings);
-
-				#if SUBSTRATE_ENABLED
-				BRANCH
-				if (bIsToon)
-				{
-					SSSTransmittanceOrOpticalDepth *= EvaluateToonProfileShadowHatching(ToonProfileId, PatternUVs, ToonProfileShadowSettings, SSSTransmittanceOrOpticalDepth);
-				}
-				#endif
 			}
 
 			// Legacy code maintained: here we have a square applied over the transmittance curve or encoded optical depth... 
diff --git a/Engine/Shaders/Private/Substrate/SubstrateEvaluation.ush b/Engine/Shaders/Private/Substrate/SubstrateEvaluation.ush
--- a/Engine/Shaders/Private/Substrate/SubstrateEvaluation.ush
+++ b/Engine/Shaders/Private/Substrate/SubstrateEvaluation.ush
@@ -1330,6 +1330,11 @@ FSubstrateEvaluateResult SubstrateEvaluateBSDFCommon(FSubstrateBSDFContext BSDFC
 
 			// TransmissionThickness in fact already stores transmittance.
 			float TransmittanceShadow = ShadowTerms.TransmittanceOrOpticalThickness;
+			{
+				// Apply shadow/transmittance hatching pattern
+				FToonProfileShadowSettings ToonProfileShadowSettings = GetToonProfileShadowSettings(ToonProfileId);
+				TransmittanceShadow *= EvaluateToonProfileShadowHatching(ToonProfileId, ToonPatternUV, ToonProfileShadowSettings, TransmittanceShadow);
+			}
 
 			// Diffuse
 			Sample.DiffuseColor				= DiffuseColor * EvaluateToonProfileDiffuseRamp(ToonProfileId, BSDFContext.Context.NoL + ToonProfileDiffuseDither); // AreaLight.NoL to handle capsule shape?
