[HWRT] Colors in debug visualizations based on MurmurAdd with GPUSceneInstanceId.
- matches the logic in NaniteVisualize.usf so Instances have the same color in the visualizations.

#Synced-CL 50529101

diff --git a/Engine/Shaders/Private/RayTracing/RayTracingDebug.usf b/Engine/Shaders/Private/RayTracing/RayTracingDebug.usf
--- a/Engine/Shaders/Private/RayTracing/RayTracingDebug.usf
+++ b/Engine/Shaders/Private/RayTracing/RayTracingDebug.usf
@@ -254,7 +254,7 @@ float4 DebugRayTracingInstance(FRayDesc Ray, out float OutDepth)
 
 	const FInstanceSceneData InstanceSceneData = GetInstanceSceneData(GPUSceneInstanceId);
 
-	Payload.InstanceHash = MurmurAdd(InstanceSceneData.PrimitiveId, InstanceSceneData.RelativeId);
+	Payload.GPUSceneInstanceId = GPUSceneInstanceId;
 	Payload.TriangleIndex = TraceRayResult.PrimitiveIndex;
 
 	float3 WorldNormal;
@@ -338,9 +338,9 @@ float4 DebugRayTracingInstance(FRayDesc Ray, out float OutDepth)
 
 	if (Payload.IsHit())
 	{
-		const float3 InstanceColor = IntToColor(Payload.InstanceHash);
-		const float3 GeometryColor = IntToColor(Payload.InstanceHash + Payload.GeometryIndex);
-		const float3 TriangleColor = IntToColor(Payload.InstanceHash + Payload.TriangleIndex);
+		const float3 InstanceColor = IntToColor(Payload.GPUSceneInstanceId);
+		const float3 GeometryColor = IntToColor(MurmurAdd(Payload.GPUSceneInstanceId, Payload.GeometryIndex));
+		const float3 TriangleColor = IntToColor(MurmurAdd(Payload.GPUSceneInstanceId, Payload.TriangleIndex));
 
 		const FRayTracingInstanceExtraData InstanceExtraData = InstancesExtraData[Payload.InstanceIndex];
 		const FRayTracingInstanceDebugData InstanceDebugData = InstancesDebugData[InstanceExtraData.SceneInstanceIndex];
diff --git a/Engine/Shaders/Private/RayTracing/RayTracingDebugCHS.usf b/Engine/Shaders/Private/RayTracing/RayTracingDebugCHS.usf
--- a/Engine/Shaders/Private/RayTracing/RayTracingDebugCHS.usf
+++ b/Engine/Shaders/Private/RayTracing/RayTracingDebugCHS.usf
@@ -17,7 +17,7 @@ RAY_TRACING_ENTRY_CLOSEST_HIT(RayTracingDebugMainCHS,
 
 	const FInstanceSceneData InstanceSceneData = GetInstanceSceneData(GPUSceneInstanceId);
 
-	Payload.InstanceHash = MurmurAdd(InstanceSceneData.PrimitiveId, InstanceSceneData.RelativeId);
+	Payload.GPUSceneInstanceId = GPUSceneInstanceId;
 
 	Payload.TriangleIndex = PrimitiveIndex();
 
diff --git a/Engine/Shaders/Private/RayTracing/RayTracingDebugUtils.ush b/Engine/Shaders/Private/RayTracing/RayTracingDebugUtils.ush
--- a/Engine/Shaders/Private/RayTracing/RayTracingDebugUtils.ush
+++ b/Engine/Shaders/Private/RayTracing/RayTracingDebugUtils.ush
@@ -7,7 +7,7 @@
 
 struct FRayTracingDebugPayload : FMinimalPayload
 {
-	uint InstanceHash;
+	uint GPUSceneInstanceId;
 	uint TriangleIndex;
 	uint InstanceIndex;
 	uint GeometryIndex;
