Fix  FarthestDeviceDepth function was incorrectly using max3 (finding the maximum of three values) instead of min3 (minimum). With an inverted Z buffer, "farthest" depth actually corresponds to the smallest numerical value, so the function should use min3.

#tests Console with odsc

diff --git a/Engine/Shaders/Private/Common.ush b/Engine/Shaders/Private/Common.ush
--- a/Engine/Shaders/Private/Common.ush
+++ b/Engine/Shaders/Private/Common.ush
@@ -127,7 +127,7 @@ float FarthestDeviceDepth(float DepthA,
                           float DepthC = NearDepthValue
 )
 {
-	return max3(DepthA, DepthB, DepthC);
+	return min3(DepthA, DepthB, DepthC);
 }
 
 const static MaterialFloat PI = 3.1415926535897932f;
