[Subsurface lighting] Skip trasmission lighting when the checkerboard pixel should only output specular lighting to SceneColor

#tests tested in editor

#Synced-CL 50540393

diff --git a/Engine/Shaders/Private/DeferredLightingCommon.ush b/Engine/Shaders/Private/DeferredLightingCommon.ush
--- a/Engine/Shaders/Private/DeferredLightingCommon.ush
+++ b/Engine/Shaders/Private/DeferredLightingCommon.ush
@@ -413,7 +413,20 @@ FLightAccumulator AccumulateDynamicLighting(
 		#endif
 
 			LightAccumulator_AddSplit( LightAccumulator, Lighting.Diffuse, Lighting.Specular, Lighting.Diffuse, MaskedLightColor * Shadow.SurfaceShadow, bNeedsSeparateSubsurfaceLightAccumulation );
-			LightAccumulator_AddSplit( LightAccumulator, Lighting.Transmission, 0.0f, Lighting.Transmission, MaskedLightColor * Shadow.TransmissionShadow, bNeedsSeparateSubsurfaceLightAccumulation );
+			
+			// SubsurfaceProfileBxDF can return non-zero transmission lighting even if this is a specular only pixel so we need to skip tranmission explicitly
+			bool bOutputSpecularOnly = false;
+		#if USE_TRANSMISSION && ALLOW_SSS_MATERIAL_OVERRIDE && ALLOW_SSS_MATERIAL_CHECKERBOARD_OVERRIDE
+			if (View.bSubsurfacePostprocessEnabled > 0 && View.bCheckerboardSubsurfaceProfileRendering > 0 && bNeedsSeparateSubsurfaceLightAccumulation)
+			{
+				bOutputSpecularOnly = !CheckerFromPixelPos(SVPos);
+			}
+		#endif
+
+			if (!bOutputSpecularOnly)
+			{
+				LightAccumulator_AddSplit( LightAccumulator, Lighting.Transmission, 0.0f, Lighting.Transmission, MaskedLightColor * Shadow.TransmissionShadow, bNeedsSeparateSubsurfaceLightAccumulation );
+			}
 
 			LightAccumulator.EstimatedCost += 0.4f;		// add the cost of the lighting computations (should sum up to 1 form one light)
 		#endif
