[Nanite Ray Tracing] Cluster visualization mode.

#Synced-CL 50529101

diff --git a/Engine/Shaders/Private/RayTracing/RayTracingDebug.usf b/Engine/Shaders/Private/RayTracing/RayTracingDebug.usf
--- a/Engine/Shaders/Private/RayTracing/RayTracingDebug.usf
+++ b/Engine/Shaders/Private/RayTracing/RayTracingDebug.usf
@@ -316,7 +316,9 @@ float4 DebugRayTracingInstance(FRayDesc Ray, out float OutDepth)
 	Payload.WorldNormal = WorldNormal;
 	Payload.InstanceIndex = TraceRayResult.InstanceIndex;
 	Payload.GeometryIndex = TraceRayResult.GeometryIndex;
-	Payload.ScenePrimitiveIndex = InstanceSceneData.PrimitiveId;
+#if PLATFORM_SUPPORTS_RAYTRACING_CLUSTER_OPS
+	Payload.ClusterId = TraceRayResult.ClusterID;
+#endif
 
 	if (TraceRayResult.bIsFrontFace)
 	{
@@ -342,6 +344,12 @@ float4 DebugRayTracingInstance(FRayDesc Ray, out float OutDepth)
 		const float3 GeometryColor = IntToColor(MurmurAdd(Payload.GPUSceneInstanceId, Payload.GeometryIndex));
 		const float3 TriangleColor = IntToColor(MurmurAdd(Payload.GPUSceneInstanceId, Payload.TriangleIndex));
 
+#if PLATFORM_SUPPORTS_RAYTRACING_CLUSTER_OPS
+		const uint PageIndex = Payload.ClusterId & NANITE_MAX_GPU_PAGES_MASK;
+		const uint ClusterIndex = (Payload.ClusterId >> NANITE_MAX_GPU_PAGES_BITS) & NANITE_MAX_CLUSTERS_PER_PAGE_MASK;
+		const float3 ClusterColor = IntToColor(MurmurAdd(PageIndex, ClusterIndex)); // match NaniteVisualize.usf
+#endif
+
 		const FRayTracingInstanceExtraData InstanceExtraData = InstancesExtraData[Payload.InstanceIndex];
 		const FRayTracingInstanceDebugData InstanceDebugData = InstancesDebugData[InstanceExtraData.SceneInstanceIndex];
 
@@ -368,6 +376,11 @@ float4 DebugRayTracingInstance(FRayDesc Ray, out float OutDepth)
 		case RAY_TRACING_DEBUG_VIZ_TRIANGLES:
 			Result.rgb = TriangleColor;
 			break;
+#if PLATFORM_SUPPORTS_RAYTRACING_CLUSTER_OPS
+		case RAY_TRACING_DEBUG_VIZ_CLUSTERS:
+			Result.rgb = ClusterColor;
+			break;
+#endif
 		case RAY_TRACING_DEBUG_VIZ_DYNAMIC_INSTANCES:
 			Result.rgb = InstanceDebugData.Flags ? float3(0.25f, 0.50f, 0.25f) : float3(0.50f, 0.25f, 0.25f);
 			break;
diff --git a/Engine/Shaders/Private/RayTracing/RayTracingDebugCHS.usf b/Engine/Shaders/Private/RayTracing/RayTracingDebugCHS.usf
--- a/Engine/Shaders/Private/RayTracing/RayTracingDebugCHS.usf
+++ b/Engine/Shaders/Private/RayTracing/RayTracingDebugCHS.usf
@@ -24,7 +24,9 @@ RAY_TRACING_ENTRY_CLOSEST_HIT(RayTracingDebugMainCHS,
 	Payload.WorldNormal = GetGeometryNormalFromTriangleBaseAttributes(PrimitiveIndex(), GetClusterID());
 	Payload.InstanceIndex = InstanceIndex();
 	Payload.GeometryIndex = GeometryIndex();
-	Payload.ScenePrimitiveIndex = InstanceSceneData.PrimitiveId;
+#if PLATFORM_SUPPORTS_RAYTRACING_CLUSTER_OPS
+	Payload.ClusterId = GetClusterID();
+#endif
 	
 	if (HitKind() == HIT_KIND_TRIANGLE_FRONT_FACE)
 	{
diff --git a/Engine/Shaders/Private/RayTracing/RayTracingDebugUtils.ush b/Engine/Shaders/Private/RayTracing/RayTracingDebugUtils.ush
--- a/Engine/Shaders/Private/RayTracing/RayTracingDebugUtils.ush
+++ b/Engine/Shaders/Private/RayTracing/RayTracingDebugUtils.ush
@@ -11,11 +11,13 @@ struct FRayTracingDebugPayload : FMinimalPayload
 	uint TriangleIndex;
 	uint InstanceIndex;
 	uint GeometryIndex;
+#if PLATFORM_SUPPORTS_RAYTRACING_CLUSTER_OPS
+	uint ClusterId;
+#endif
 
 	float3 WorldNormal;
 	
 	uint TriangleHitCountPerRay;
-	uint ScenePrimitiveIndex;
 	uint Flags;
 	
 	void SetFrontFace()
diff --git a/Engine/Shaders/Shared/RaytracingDebugDefinitions.h b/Engine/Shaders/Shared/RaytracingDebugDefinitions.h
--- a/Engine/Shaders/Shared/RaytracingDebugDefinitions.h
+++ b/Engine/Shaders/Shared/RaytracingDebugDefinitions.h
@@ -34,7 +34,7 @@
 #define RAY_TRACING_DEBUG_VIZ_TIMING_ANY_HIT					24 
 #define RAY_TRACING_DEBUG_VIZ_TIMING_MATERIAL					25 
 #define RAY_TRACING_DEBUG_VIZ_TRIANGLES							26
-#define RAY_TRACING_DEBUG_VIZ_UNUSED_0							27
+#define RAY_TRACING_DEBUG_VIZ_CLUSTERS							27
 #define RAY_TRACING_DEBUG_VIZ_TRAVERSAL_PRIMARY_NODE			28
 #define RAY_TRACING_DEBUG_VIZ_TRAVERSAL_PRIMARY_CLUSTER			29
 #define RAY_TRACING_DEBUG_VIZ_TRAVERSAL_PRIMARY_TRIANGLE		30
@@ -74,6 +74,7 @@ INLINE_ATTR bool RequiresRayTracingDebugCHS(UINT_TYPE DebugVisualizationMode)
 {
 	return DebugVisualizationMode == RAY_TRACING_DEBUG_VIZ_INSTANCES
 		|| DebugVisualizationMode == RAY_TRACING_DEBUG_VIZ_TRIANGLES
+		|| DebugVisualizationMode == RAY_TRACING_DEBUG_VIZ_CLUSTERS
 		|| DebugVisualizationMode == RAY_TRACING_DEBUG_VIZ_DYNAMIC_INSTANCES
 		|| DebugVisualizationMode == RAY_TRACING_DEBUG_VIZ_PROXY_TYPE
 		|| DebugVisualizationMode == RAY_TRACING_DEBUG_VIZ_PICKER
diff --git a/Engine/Source/Runtime/Engine/Private/RayTracingVisualizationData.cpp b/Engine/Source/Runtime/Engine/Private/RayTracingVisualizationData.cpp
--- a/Engine/Source/Runtime/Engine/Private/RayTracingVisualizationData.cpp
+++ b/Engine/Source/Runtime/Engine/Private/RayTracingVisualizationData.cpp
@@ -53,6 +53,10 @@ FRayTracingVisualizationData::FRayTracingVisualizationData()
 		// debugging the geometry itself
 		AddVisualizationMode(TEXT("Instances"), LOCTEXT("Instances", "Instances"), FModeType::Standard, RAY_TRACING_DEBUG_VIZ_INSTANCES, false);
 		AddVisualizationMode(TEXT("Triangles"), LOCTEXT("Triangles", "Triangles"), FModeType::Standard, RAY_TRACING_DEBUG_VIZ_TRIANGLES, false);
+		if (GRHIGlobals.RayTracing.SupportsClusterOps)
+		{
+			AddVisualizationMode(TEXT("Clusters"), LOCTEXT("Clusters", "Clusters"), FModeType::Standard, RAY_TRACING_DEBUG_VIZ_CLUSTERS, false);
+		}
 		AddVisualizationMode(TEXT("DynamicInstances"), LOCTEXT("DynamicInstances", "Dynamic Instances"), FModeType::Standard, RAY_TRACING_DEBUG_VIZ_DYNAMIC_INSTANCES, false);
 		AddVisualizationMode(TEXT("ProxyType"), LOCTEXT("ProxyType", "Proxy Type"), FModeType::Standard, RAY_TRACING_DEBUG_VIZ_PROXY_TYPE, false);
 		AddVisualizationMode(TEXT("Picker"), LOCTEXT("Picker", "Picker"), FModeType::Standard, RAY_TRACING_DEBUG_VIZ_PICKER, false);
