Workaround for error compiling FMobileReflectionEnvironmentSkyLightingPS.

#Synced-CL 50568677

diff --git a/Engine/Shaders/Private/MobileLightingCommon.ush b/Engine/Shaders/Private/MobileLightingCommon.ush
--- a/Engine/Shaders/Private/MobileLightingCommon.ush
+++ b/Engine/Shaders/Private/MobileLightingCommon.ush
@@ -456,13 +456,21 @@ float3 GetImageBasedReflectionSpecular(
 	half3 Out = 0;
 
 	// IBL reflection
+	half4 DiffuseIBL = half4(0.0, 0.0, 0.0, 1.0);
 	Out = GetImageBasedReflectionLighting_Mobile(
 		SpecularDirection,
 		TranslatedWorldPosition,
 		Roughness,
 		IndirectIrradiance,
 		SpecularOcclusion,
-		GridIndex);
+		GridIndex
+#if DIFFUSE_FROM_CAPTURE
+		, WorldNormal
+		, DiffuseIBL
+#endif
+	);
+
+	// TODO: Use DiffuseIBL?
 
 	// Planar reflection
 	Out = GetPlanarReflectionbasedReflectionLighting_Mobile(TranslatedWorldPosition, WorldNormal, Roughness, Out);
