Fix GPU hang on some platforms, working around some compiler issues. More details in the jira

#rnx

#Synced-CL 50430605

#Synced-CL 50430605

diff --git a/Engine/Shaders/Private/LandscapeGrassWeight.usf b/Engine/Shaders/Private/LandscapeGrassWeight.usf
--- a/Engine/Shaders/Private/LandscapeGrassWeight.usf
+++ b/Engine/Shaders/Private/LandscapeGrassWeight.usf
@@ -61,165 +61,174 @@ int4 PerChannelGrassWeightIndices; // One index per channel in the OutColor. 0 i
 #define HeightGrassWeightIndex0 -2
 #define HeightGrassWeightIndex1 -1
 
-// Pixel Shader
-void PSMain( 
-	FLandscapeGrassWeightInterpolantsVSToPS Interpolants
-	OPTIONAL_IsFrontFace,
-	out HALF4_TYPE OutColor	: SV_Target0
-	)
+float GetGrassWeightForChannel(uint ChannelIndex, float2 PackedHeight, FMaterialPixelParameters MaterialParameters)
 {
-	ResolvedView = ResolveView();
-
-	// Build the SvPosition without RenderOffset so that we correctly reproject WorldPosition etc.
-	float4 SvPositionPreOffset = Interpolants.PositionPreOffset / Interpolants.PositionPreOffset.w;
-	SvPositionPreOffset.xy = (SvPositionPreOffset.xy * float2(0.5f, -0.5f) + 0.5f) * ResolvedView.ViewSizeAndInvSize.xy;
-
-	FMaterialPixelParameters MaterialParameters = GetMaterialPixelParameters(Interpolants.FactoryInterpolants, SvPositionPreOffset);
-	FPixelMaterialInputs PixelMaterialInputs;
-	CalcMaterialParameters(MaterialParameters, PixelMaterialInputs, SvPositionPreOffset, bIsFrontFace);
-
-	// Calculate and pack landscape height.
-#if VF_USE_PRIMITIVE_SCENE_DATA
-	float Height = clamp(round(Interpolants.LocalPosZ / TERRAIN_ZSCALE) + 32768.0, 0, 65535);
-#else // !VF_USE_PRIMITIVE_SCENE_DATA
-	float3 LocalPos = WSMultiplyDemote(GetWorldPosition(MaterialParameters), GetWorldToInstance(MaterialParameters));
-	float Height = clamp(round(LocalPos.z / TERRAIN_ZSCALE) + 32768.0, 0, 65535);
-#endif // VF_USE_PRIMITIVE_SCENE_DATA
-	float2 PackedHeight = float2(floor(Height / 256), fmod(Height, 256)) / 255;
-
-	OutColor = 0.0;
+	float OutputColor = 0;
 
 #define GRASSWEIGHTCASE(N) case N: \
 	{ \
-		OutColor[ChannelIndex] = GetGrassWeight##N(MaterialParameters).x; \
-		break; \
+	OutputColor = GetGrassWeight##N(MaterialParameters).x; \
+	break; \
 	}
 
-	// Do not remove : this unroll directive is actually necessary to avoid a FXC warning (error X3511: ... unrolled loop is too large)
-	UNROLL_N(4)
-	for (uint ChannelIndex = 0; ChannelIndex < 4; ++ChannelIndex)
+
+	// Do not remove : this branch directive is actually necessary to avoid FXC screwing up the generated shader
+	BRANCH
+	switch(PerChannelGrassWeightIndices[ChannelIndex])
 	{
-		// Do not remove : this branch directive is actually necessary to avoid FXC screwing up the generated shader
-		BRANCH
-		switch(PerChannelGrassWeightIndices[ChannelIndex])
+		case GrassWeightInvalidIndex: 
+			break;
+		case HeightGrassWeightIndex0: 
+		{
+			OutputColor = PackedHeight.x; 
+			break;
+		}
+		case HeightGrassWeightIndex1: 
 		{
-			case GrassWeightInvalidIndex: 
-				break;
-			case HeightGrassWeightIndex0: 
-			{
-				OutColor[ChannelIndex] = PackedHeight.x; 
-				break;
-			}
-			case HeightGrassWeightIndex1: 
-			{
-				OutColor[ChannelIndex] = PackedHeight.y; 
-				break;
-			}
+			OutputColor = PackedHeight.y; 
+			break;
+		}
 #if HAVE_GetGrassWeight0
-			GRASSWEIGHTCASE(0);
+		GRASSWEIGHTCASE(0);
 #endif // HAVE_GetGrassWeight0
 #if HAVE_GetGrassWeight1
-			GRASSWEIGHTCASE(1);
+		GRASSWEIGHTCASE(1);
 #endif // HAVE_GetGrassWeight1
 #if HAVE_GetGrassWeight2
-			GRASSWEIGHTCASE(2);
+		GRASSWEIGHTCASE(2);
 #endif // HAVE_GetGrassWeight2
 #if HAVE_GetGrassWeight3
-			GRASSWEIGHTCASE(3);
+		GRASSWEIGHTCASE(3);
 #endif // HAVE_GetGrassWeight3
 #if HAVE_GetGrassWeight4
-			GRASSWEIGHTCASE(4);
+		GRASSWEIGHTCASE(4);
 #endif // HAVE_GetGrassWeight4
 #if HAVE_GetGrassWeight5
-			GRASSWEIGHTCASE(5);
+		GRASSWEIGHTCASE(5);
 #endif // HAVE_GetGrassWeight5
 #if HAVE_GetGrassWeight6
-			GRASSWEIGHTCASE(6);
+		GRASSWEIGHTCASE(6);
 #endif // HAVE_GetGrassWeight6
 #if HAVE_GetGrassWeight7
-			GRASSWEIGHTCASE(7);
+		GRASSWEIGHTCASE(7);
 #endif // HAVE_GetGrassWeight7
 #if HAVE_GetGrassWeight8
-			GRASSWEIGHTCASE(8);
+		GRASSWEIGHTCASE(8);
 #endif // HAVE_GetGrassWeight8
 #if HAVE_GetGrassWeight9
-			GRASSWEIGHTCASE(9);
+		GRASSWEIGHTCASE(9);
 #endif // HAVE_GetGrassWeight9
 
 #if HAVE_GetGrassWeight10
-			GRASSWEIGHTCASE(10);
+		GRASSWEIGHTCASE(10);
 #endif // HAVE_GetGrassWeight10
 #if HAVE_GetGrassWeight11
-			GRASSWEIGHTCASE(11);
+		GRASSWEIGHTCASE(11);
 #endif // HAVE_GetGrassWeight11
 #if HAVE_GetGrassWeight12
-			GRASSWEIGHTCASE(12);
+		GRASSWEIGHTCASE(12);
 #endif // HAVE_GetGrassWeight12
 #if HAVE_GetGrassWeight13
-			GRASSWEIGHTCASE(13);
+		GRASSWEIGHTCASE(13);
 #endif // HAVE_GetGrassWeight13
 #if HAVE_GetGrassWeight14
-			GRASSWEIGHTCASE(14);
+		GRASSWEIGHTCASE(14);
 #endif // HAVE_GetGrassWeight14
 #if HAVE_GetGrassWeight15
-			GRASSWEIGHTCASE(15);
+		GRASSWEIGHTCASE(15);
 #endif // HAVE_GetGrassWeight15
 #if HAVE_GetGrassWeight16
-			GRASSWEIGHTCASE(16);
+		GRASSWEIGHTCASE(16);
 #endif // HAVE_GetGrassWeight16
 #if HAVE_GetGrassWeight17
-			GRASSWEIGHTCASE(17);
+		GRASSWEIGHTCASE(17);
 #endif // HAVE_GetGrassWeight17
 #if HAVE_GetGrassWeight18
-			GRASSWEIGHTCASE(18);
+		GRASSWEIGHTCASE(18);
 #endif // HAVE_GetGrassWeight18
 #if HAVE_GetGrassWeight19
-			GRASSWEIGHTCASE(19);
+		GRASSWEIGHTCASE(19);
 #endif // HAVE_GetGrassWeight19
 
 #if HAVE_GetGrassWeight20
-			GRASSWEIGHTCASE(20);
+		GRASSWEIGHTCASE(20);
 #endif // HAVE_GetGrassWeight20
 #if HAVE_GetGrassWeight21
-			GRASSWEIGHTCASE(21);
+		GRASSWEIGHTCASE(21);
 #endif // HAVE_GetGrassWeight21
 #if HAVE_GetGrassWeight22
-			GRASSWEIGHTCASE(22);
+		GRASSWEIGHTCASE(22);
 #endif // HAVE_GetGrassWeight22
 #if HAVE_GetGrassWeight23
-			GRASSWEIGHTCASE(23);
+		GRASSWEIGHTCASE(23);
 #endif // HAVE_GetGrassWeight23
 #if HAVE_GetGrassWeight24
-			GRASSWEIGHTCASE(24);
+		GRASSWEIGHTCASE(24);
 #endif // HAVE_GetGrassWeight24
 #if HAVE_GetGrassWeight25
-			GRASSWEIGHTCASE(25);
+		GRASSWEIGHTCASE(25);
 #endif // HAVE_GetGrassWeight25
 #if HAVE_GetGrassWeight26
-			GRASSWEIGHTCASE(26);
+		GRASSWEIGHTCASE(26);
 #endif // HAVE_GetGrassWeight26
 #if HAVE_GetGrassWeight27
-			GRASSWEIGHTCASE(27);
+		GRASSWEIGHTCASE(27);
 #endif // HAVE_GetGrassWeight27
 #if HAVE_GetGrassWeight28
-			GRASSWEIGHTCASE(28);
+		GRASSWEIGHTCASE(28);
 #endif // HAVE_GetGrassWeight28
 #if HAVE_GetGrassWeight29
-			GRASSWEIGHTCASE(29);
+		GRASSWEIGHTCASE(29);
 #endif // HAVE_GetGrassWeight29
 
 #if HAVE_GetGrassWeight30
-			GRASSWEIGHTCASE(30);
+		GRASSWEIGHTCASE(30);
 #endif // HAVE_GetGrassWeight30
 #if HAVE_GetGrassWeight31
-			GRASSWEIGHTCASE(31);
+		GRASSWEIGHTCASE(31);
 #endif // HAVE_GetGrassWeight31
 
-			default:
-				break;
-		}
+		default:
+			break;
 	}
+	
+	return OutputColor;
+}
+
+// Pixel Shader
+void PSMain( 
+	FLandscapeGrassWeightInterpolantsVSToPS Interpolants
+	OPTIONAL_IsFrontFace,
+	out HALF4_TYPE OutColor	: SV_Target0
+	)
+{
+	OutColor = 0;
+
+	ResolvedView = ResolveView();
+
+	// Build the SvPosition without RenderOffset so that we correctly reproject WorldPosition etc.
+	float4 SvPositionPreOffset = Interpolants.PositionPreOffset / Interpolants.PositionPreOffset.w;
+	SvPositionPreOffset.xy = (SvPositionPreOffset.xy * float2(0.5f, -0.5f) + 0.5f) * ResolvedView.ViewSizeAndInvSize.xy;
+
+	FMaterialPixelParameters MaterialParameters = GetMaterialPixelParameters(Interpolants.FactoryInterpolants, SvPositionPreOffset);
+	FPixelMaterialInputs PixelMaterialInputs;
+	CalcMaterialParameters(MaterialParameters, PixelMaterialInputs, SvPositionPreOffset, bIsFrontFace);
+
+	// Calculate and pack landscape height.
+#if VF_USE_PRIMITIVE_SCENE_DATA
+	float Height = clamp(round(Interpolants.LocalPosZ / TERRAIN_ZSCALE) + 32768.0, 0, 65535);
+#else // !VF_USE_PRIMITIVE_SCENE_DATA
+	float3 LocalPos = WSMultiplyDemote(GetWorldPosition(MaterialParameters), GetWorldToInstance(MaterialParameters));
+	float Height = clamp(round(LocalPos.z / TERRAIN_ZSCALE) + 32768.0, 0, 65535);
+#endif // VF_USE_PRIMITIVE_SCENE_DATA
+	float2 PackedHeight = float2(floor(Height / 256), fmod(Height, 256)) / 255;
+
+	OutColor = 0.0;
+
+	OutColor.r = GetGrassWeightForChannel(0, PackedHeight, MaterialParameters);
+	OutColor.g = GetGrassWeightForChannel(1, PackedHeight, MaterialParameters);
+	OutColor.b = GetGrassWeightForChannel(2, PackedHeight, MaterialParameters);
+	OutColor.a = GetGrassWeightForChannel(3, PackedHeight, MaterialParameters);
 
-	// UMaterialExpressionLandscapeGrassOutput::MaxGrassTypes == 32. Keep on going if this gets increased
 }
 #endif
\ No newline at end of file
