Fix missing material's Emissive contribution on Mobile deferred path.

#Synced-CL 50586717

diff --git a/Engine/Shaders/Private/MobileBasePassPixelShader.usf b/Engine/Shaders/Private/MobileBasePassPixelShader.usf
--- a/Engine/Shaders/Private/MobileBasePassPixelShader.usf
+++ b/Engine/Shaders/Private/MobileBasePassPixelShader.usf
@@ -879,7 +879,7 @@ void Main(
 		MobileEncodeGBuffer(GBuffer, OutGBufferA, OutGBufferB, OutGBufferC, OutGBufferD);
 	#endif
 
-		Color = DirectLighting.TotalLight;
+		Color = DirectLighting.TotalLight + Export.EmissiveLuminance;
 	}
 #endif // MATERIAL_SUBSTRATE_OPAQUE_PRECOMPUTED_LIGHTING
 
