Substrate - Missing pre-exposure on sky lighting when using Lumen GI is disabled.

#tests PC editor EngineTests, MaterialPack

#Synced-CL 50764865

diff --git a/Engine/Shaders/Private/ReflectionEnvironmentPixelShader.usf b/Engine/Shaders/Private/ReflectionEnvironmentPixelShader.usf
--- a/Engine/Shaders/Private/ReflectionEnvironmentPixelShader.usf
+++ b/Engine/Shaders/Private/ReflectionEnvironmentPixelShader.usf
@@ -545,7 +545,10 @@ void ReflectionEnvironmentSkyLighting(
 				CombinedScreenAndMaterialAO,
 				DiffuseLighting,
 				SpecularLighting);
-	
+
+			// Diffuse lighting needs pre exposition as GetEnvDiffuseLighting() is not preexposed.
+			DiffuseLighting *= View.PreExposure; 
+
 		#if REFLECTION_SOURCE_TYPE == REFLECTION_SOURCE_LUMEN_STANDALONE
 			{
 				// Pre exposure the sky light specular evaluated from within SubstrateEnvLightingCommon (CompositeReflectionCapturesAndSkylightTWS) when rendering lumen HWRT reflections while diffuse global illumination is "none".
