[SingleLayerWater] Fixed "fuzzy" artifacts on silhouette of objects far above a SLW water surface.

#Synced-CL 50766943

diff --git a/Engine/Shaders/Private/SingleLayerWaterCommon.ush b/Engine/Shaders/Private/SingleLayerWaterCommon.ush
--- a/Engine/Shaders/Private/SingleLayerWaterCommon.ush
+++ b/Engine/Shaders/Private/SingleLayerWaterCommon.ush
@@ -19,6 +19,12 @@
 // Used in MobileBasePassPixelShader.usf.
 #define SINGLE_LAYER_WATER_SHADING_QUALITY_MOBILE_WITH_DEPTH_BUFFER		2
 
+#ifndef SINGLE_LAYER_WATER_USE_REFRACTION_MASK
+#define SINGLE_LAYER_WATER_USE_REFRACTION_MASK 1
+#endif
+
+#define SINGLE_LAYER_WATER_REFRACTION_MASK_AVAILABLE ((SINGLE_LAYER_WATER_SHADING_QUALITY == SINGLE_LAYER_WATER_SHADING_QUALITY_FULL) && SINGLE_LAYER_WATER_USE_REFRACTION_MASK) 
+
 #if MATERIAL_SHADINGMODEL_SINGLELAYERWATER
 	#if SUPPORTS_INDEPENDENT_SAMPLERS
 		#define SingleLayerWaterSceneColorSampler View.SharedBilinearClampedSampler
@@ -89,3 +95,49 @@ void UnpackWaterRefractionMask(float4 PackedRefractionMask, out bool4 bOutIsWate
 
 #undef SLW_MASK_IS_WATER_BITMASK
 #undef SLW_MASK_REFRACTION_BITMASK
+
+void SampleSceneDepthWithoutWaterBilateral(Texture2D Tex, SamplerState TexSampler, float2 UV, float2 TexSize, float WaterSurfaceSceneDepth,
+#if SINGLE_LAYER_WATER_REFRACTION_MASK_AVAILABLE
+	Texture2D RefractionMaskTex,
+#endif
+	out float OutDeviceZ, out float OutSceneDepth, out float4 OutWeights, out float4 OutRefractionMask)
+{
+	// Gather deviceZ depth values and convert to scene depth
+	const float4 DeviceZ4 = Tex.Gather(TexSampler, UV);
+	const float4 SceneDepth4 = float4(ConvertFromDeviceZ(DeviceZ4.x), ConvertFromDeviceZ(DeviceZ4.y), ConvertFromDeviceZ(DeviceZ4.z), ConvertFromDeviceZ(DeviceZ4.w));
+
+	// Standard bilinear interpolation factors
+	const float2 BilinearInterp = frac(UV * TexSize - 0.5f);
+	const float4 BilinearWeights = float4(
+		(1.0f - BilinearInterp.x) * BilinearInterp.y,
+		BilinearInterp.x * BilinearInterp.y,
+		BilinearInterp.x * (1.0f - BilinearInterp.y),
+		(1.0f - BilinearInterp.x) * (1.0f - BilinearInterp.y));
+
+	// Heuristic to reject non-water pixels belonging to foreground objects based on depth
+	float SlopeTolerance = max(abs(ddx_fine(WaterSurfaceSceneDepth)), abs(ddy_fine(WaterSurfaceSceneDepth)));
+	float4 BilateralWeights = select((SceneDepth4 + SlopeTolerance) >= WaterSurfaceSceneDepth, 1.0f, 0.0f);
+
+	// If the refraction mask is available, use it as well
+#if SINGLE_LAYER_WATER_REFRACTION_MASK_AVAILABLE
+	const float4 PackedRefractionMask4 = RefractionMaskTex.Gather(TexSampler, UV);
+	bool4 bIsWater4 = false;
+	float4 RefractionMask4 = 0.0f;
+	UnpackWaterRefractionMask(PackedRefractionMask4, bIsWater4, RefractionMask4);
+	BilateralWeights *= select(bIsWater4, 1.0f, 0.0f);
+#endif
+
+	// Calculate final weights and make sure to "normalize" them. If the weights are too small, force them to zero, signalling an invalid sample.
+	float4 FinalWeights = BilinearWeights * BilateralWeights;
+	const float WeightsSum = FinalWeights.x + FinalWeights.y + FinalWeights.z + FinalWeights.w;
+	FinalWeights *= WeightsSum > 1e-5f ? rcp(WeightsSum) : 0.0f;
+
+	OutSceneDepth = dot(SceneDepth4, FinalWeights);
+	OutDeviceZ = ConvertToDeviceZ(OutSceneDepth);
+	OutWeights = FinalWeights;
+#if SINGLE_LAYER_WATER_REFRACTION_MASK_AVAILABLE
+	OutRefractionMask = RefractionMask4;
+#else
+	OutRefractionMask = 1.0f;
+#endif
+}
diff --git a/Engine/Shaders/Private/SingleLayerWaterComposite.usf b/Engine/Shaders/Private/SingleLayerWaterComposite.usf
--- a/Engine/Shaders/Private/SingleLayerWaterComposite.usf
+++ b/Engine/Shaders/Private/SingleLayerWaterComposite.usf
@@ -12,6 +12,7 @@
 #include "SceneTextureParameters.ush"
 #include "LightGridCommon.ush"
 #include "/Engine/Shared/SingleLayerWaterDefinitions.h"
+#define SINGLE_LAYER_WATER_USE_REFRACTION_MASK 0
 #include "SingleLayerWaterCommon.ush"
 
 #define ENABLE_SKY_LIGHT 1
@@ -85,18 +86,14 @@ SingleLayerWaterCompositeOutput SingleLayerWaterComposite(float2 BufferUV, float
 	{
 		float DeviceZ = SampleDeviceZFromSceneTextures(BufferUV);
 		float SceneDepth = ConvertFromDeviceZ(DeviceZ);
-
+		float WaterDepth = SceneDepth;
 		const float2 SceneNoWaterUV = clamp(BufferUV, EffectiveSceneNoWaterMinMaxUV.xy, EffectiveSceneNoWaterMinMaxUV.zw);
 
-		float OpaqueDeviceZ = WaterSampleSceneDepthWithoutWater(
-			SceneNoWaterDepthTexture,
-			SceneNoWaterDepthSampler,
-			SceneNoWaterUV,
-			SceneNoWaterTextureSize,
-			SceneNoWaterInvTextureSize);
-		float OpaqueDepth = ConvertFromDeviceZ(OpaqueDeviceZ);
-
-		float WaterDepth = SceneDepth;
+		float OpaqueDeviceZ;
+		float OpaqueDepth;
+		float4 WeightsDummy;
+		float4 RefractionMaskDummy;
+		SampleSceneDepthWithoutWaterBilateral(SceneNoWaterDepthTexture, SceneNoWaterDepthSampler, SceneNoWaterUV, SceneNoWaterTextureSize, WaterDepth, OpaqueDeviceZ, OpaqueDepth, WeightsDummy, RefractionMaskDummy);
 
 		float DeltaDepth = CameraIsUnderWater ? WaterDepth : OpaqueDepth - WaterDepth;	// Inverted depth
 		if (DeltaDepth > 0.0)
@@ -176,17 +173,14 @@ SingleLayerWaterCompositeOutput SingleLayerWaterComposite(float2 BufferUV, float
 
 	if (ShadingModelID == SHADINGMODELID_SINGLELAYERWATER)
 	{
+		float WaterDepth = GBuffer.Depth;
 		const float2 SceneNoWaterUV = clamp(BufferUV, EffectiveSceneNoWaterMinMaxUV.xy, EffectiveSceneNoWaterMinMaxUV.zw);
 
-		float OpaqueDeviceZ = WaterSampleSceneDepthWithoutWater(
-			SceneNoWaterDepthTexture,
-			SceneNoWaterDepthSampler,
-			SceneNoWaterUV,
-			SceneNoWaterTextureSize,
-			SceneNoWaterInvTextureSize);
-		float OpaqueDepth = ConvertFromDeviceZ(OpaqueDeviceZ);
-
-		float WaterDepth = GBuffer.Depth;
+		float OpaqueDeviceZ;
+		float OpaqueDepth;
+		float4 WeightsDummy;
+		float4 RefractionMaskDummy;
+		SampleSceneDepthWithoutWaterBilateral(SceneNoWaterDepthTexture, SceneNoWaterDepthSampler, SceneNoWaterUV, SceneNoWaterTextureSize, WaterDepth, OpaqueDeviceZ, OpaqueDepth, WeightsDummy, RefractionMaskDummy);
 
 		float DeltaDepth = CameraIsUnderWater ? WaterDepth : OpaqueDepth - WaterDepth;	// Inverted depth
 		if (DeltaDepth > 0.0)
diff --git a/Engine/Shaders/Private/SingleLayerWaterShading.ush b/Engine/Shaders/Private/SingleLayerWaterShading.ush
--- a/Engine/Shaders/Private/SingleLayerWaterShading.ush
+++ b/Engine/Shaders/Private/SingleLayerWaterShading.ush
@@ -8,10 +8,6 @@
 #include "DistortionCommon.ush"
 #include "SingleLayerWaterCommon.ush"
 
-#ifndef SINGLE_LAYER_WATER_USE_REFRACTION_MASK
-#define SINGLE_LAYER_WATER_USE_REFRACTION_MASK 1
-#endif
-
 // The deferred mobile renderer with support for SM5 creates a scene color copy with encoding to fit the data into RGB8
 #ifndef SINGLE_LAYER_WATER_ENCODED_SCENE_COLOR
 #define SINGLE_LAYER_WATER_ENCODED_SCENE_COLOR 0
@@ -30,6 +26,22 @@
 	#define EARLY_Z_PASS_ONLY_MATERIAL_MASKING 0
 #endif
 
+float3 SampleSceneColorWithoutWaterBilateral(Texture2D Tex, SamplerState TexSampler, float2 UV, float4 Weights)
+{
+	const float4 Red = Tex.GatherRed(TexSampler, UV);
+	const float4 Green = Tex.GatherGreen(TexSampler, UV);
+	const float4 Blue = Tex.GatherBlue(TexSampler, UV);
+
+	float3 BehindWaterSceneLuminance = float3(dot(Red, Weights), dot(Green, Weights), dot(Blue, Weights));
+
+#if SINGLE_LAYER_WATER_ENCODED_SCENE_COLOR
+	// Undo the function in EncodeSceneColorForMaterialNode
+	BehindWaterSceneLuminance = pow(BehindWaterSceneLuminance.rgb, 4) * 10;
+#endif
+
+	return BehindWaterSceneLuminance;
+}
+
 struct WaterVolumeLightingOutput
 {
 	float3 Luminance;
@@ -56,7 +68,7 @@ WaterVolumeLightingOutput EvaluateWaterVolumeLighting(
 #endif
 	float Specular, const float4 DistortionParams,
 	float3 SunIlluminance, float3 AmbiantIlluminance, float3 EnvBrdf,
-	bool CameraIsUnderWater, float WaterVisibility, uint EyeIndex,
+	bool bCameraIsUnderWater, float WaterVisibility, uint EyeIndex,
 	bool bSeparateMainDirLight, inout float3 SeparatedWaterMainDirLightScatteredLuminance,
 	in float3 LightFunctionColor = 1.0f)
 {
@@ -69,25 +81,23 @@ WaterVolumeLightingOutput EvaluateWaterVolumeLighting(
 
 	if (WaterVisibility > 0.0f)
 	{
-		float WaterDepth = GetPixelDepth(MaterialParameters); // Scene depth with Pixel-Depth-Offset taken into account is in ScreenPosition.w already
+		float WaterSurfaceSceneDepth = GetPixelDepth(MaterialParameters); // Scene depth with Pixel-Depth-Offset taken into account is in ScreenPosition.w already
 #if SINGLE_LAYER_WATER_MOBILE_SHADING
 
 #if SINGLE_LAYER_WATER_SHADING_QUALITY == SINGLE_LAYER_WATER_SHADING_QUALITY_MOBILE_WITH_DEPTH_TEXTURE 
-
-		float PixelSceneDeviceZ = WaterSampleSceneDepthWithoutWater(
-			SceneDepthWithoutSingleLayerWaterTexture,
-			SceneDepthWithoutSingleLayerWaterSampler,
-			ViewportUV,
-			SceneDepthWithoutSingleLayerWaterTextureSize,
-			SceneDepthWithoutSingleLayerWaterTextureInvSize);
-
-		float PixelSceneDepth = ConvertFromDeviceZ(PixelSceneDeviceZ);
+		float BehindWaterSceneDeviceZ;
+		float BehindWaterSceneDepth;
+		float4 BehindWaterSceneColorWeights;
+		float4 DummyRefractionMask;
+		SampleSceneDepthWithoutWaterBilateral(SceneDepthWithoutSingleLayerWaterTexture, SceneDepthWithoutSingleLayerWaterSampler, ViewportUV, SceneDepthWithoutSingleLayerWaterTextureSize,
+			WaterSurfaceSceneDepth,
+			BehindWaterSceneDeviceZ, BehindWaterSceneDepth, BehindWaterSceneColorWeights, DummyRefractionMask);
 #else
-		float PixelSceneDeviceZ = LookupDeviceZ(ViewportUV);
-		float PixelSceneDepth = ConvertFromDeviceZ(PixelSceneDeviceZ);
+		float BehindWaterSceneDeviceZ = LookupDeviceZ(ViewportUV);
+		float BehindWaterSceneDepth = ConvertFromDeviceZ(BehindWaterSceneDeviceZ);
 #endif
-		float SceneDeviceZ = PixelSceneDeviceZ;
-		float SceneDepth = PixelSceneDepth;
+
+		float2 BehindWaterUV = ViewportUV;
 
 #else  // !SINGLE_LAYER_WATER_MOBILE_SHADING
 
@@ -95,89 +105,65 @@ WaterVolumeLightingOutput EvaluateWaterVolumeLighting(
 
 		ViewportUV = clamp(ViewportUV, SceneWithoutSingleLayerWaterMinUV, SceneWithoutSingleLayerWaterMaxUV);
 
-		float PixelSceneDeviceZ = WaterSampleSceneDepthWithoutWater(
-			SceneDepthWithoutSingleLayerWaterTexture,
-			SceneDepthWithoutSingleLayerWaterSampler,
-			ViewportUV,
-			SceneDepthWithoutSingleLayerWaterTextureSize,
-			SceneDepthWithoutSingleLayerWaterTextureInvSize);
+		// Sample SceneDepthWithoutWater at current pixel location
+		float PixelBehindWaterSceneDeviceZ;
+		float PixelBehindWaterSceneDepth;
+		float4 PixelSceneColorWeights;
+		float4 PixelRefractionMask;
+		SampleSceneDepthWithoutWaterBilateral(SceneDepthWithoutSingleLayerWaterTexture, SceneDepthWithoutSingleLayerWaterSampler, ViewportUV, SceneDepthWithoutSingleLayerWaterTextureSize, WaterSurfaceSceneDepth,
+		#if SINGLE_LAYER_WATER_REFRACTION_MASK_AVAILABLE
+			RefractionMaskTexture,
+		#endif
+			PixelBehindWaterSceneDeviceZ, PixelBehindWaterSceneDepth, PixelSceneColorWeights, PixelRefractionMask);
 
-		float PixelSceneDepth = ConvertFromDeviceZ(PixelSceneDeviceZ);
-		
+		// Compute a distorted/refracted UV
 		FMaterialRefractionData RefractionData = GetMaterialRefraction(PixelMaterialInputs);
 		float2 BufferUVDistortion = ComputeBufferUVDistortion(
 			MaterialParameters, PixelMaterialInputs, ResolvedView,
 			MaterialParameters.WorldNormal, IorWater,
 			DistortionParams, ViewportUV, RefractionData, false, EyeIndex);
-		PostProcessUVDistortion(MaterialParameters, PixelMaterialInputs, PixelSceneDepth, BufferUVDistortion, RefractionData);
+		PostProcessUVDistortion(MaterialParameters, PixelMaterialInputs, PixelBehindWaterSceneDepth, BufferUVDistortion, RefractionData);
 		// Also reduce refraction according to water thickness
-		BufferUVDistortion *= saturate((PixelSceneDepth - WaterDepth) * 1.0 / 30.0f);// *(1 - saturate(WaterDepth * 1.0 / 2000.0f));
-
-		float2 DistortedUV = clamp(ViewportUV + BufferUVDistortion, SceneWithoutSingleLayerWaterMinUV, SceneWithoutSingleLayerWaterMaxUV);
-
-		float RefractionMask = 1.0f;
-#if (SINGLE_LAYER_WATER_SHADING_QUALITY == SINGLE_LAYER_WATER_SHADING_QUALITY_FULL) && SINGLE_LAYER_WATER_USE_REFRACTION_MASK
-		{
-			// Gather the four refraction mask samples at the distorted UV to see if they are valid for refraction.
-			const float4 PackedRefractionMask4 = RefractionMaskTexture.Gather(SceneColorWithoutSingleLayerWaterSampler, DistortedUV);
-			bool4 bIsWater4 = false;
-			float4 RefractionMask4 = 0.0f;
-			UnpackWaterRefractionMask(PackedRefractionMask4, bIsWater4, RefractionMask4);
-			
-			const float2 PixelCoord = DistortedUV * ResolvedView.BufferSizeAndInvSize.xy;
-			const float2 BilinearInterp = frac(PixelCoord - 0.5f);
-			const float4 BilinearWeights = float4(
-				(1.0f - BilinearInterp.x) * BilinearInterp.y,
-				BilinearInterp.x * BilinearInterp.y,
-				BilinearInterp.x * (1.0f - BilinearInterp.y),
-				(1.0f - BilinearInterp.x) * (1.0f - BilinearInterp.y));
-
-			const float4 IsWater4 = select(bIsWater4, 1.0f, 0.0f);
-			const float IsWater = dot(IsWater4, BilinearWeights);
-			RefractionMask = dot(RefractionMask4, BilinearWeights);
-			DistortedUV = lerp(ViewportUV, DistortedUV, RefractionMask);
-			
-			// Some samples don't belong to water, so we skip the distortion. However, we must still check if there is valid refraction data at the viewport UV pixel.
-			// IsWater is bilenarly interpolated, so even if a non-water pixel is in the 2x2 footprint, if the UV does not or only barely "touch" it, it shouldn't cause
-			// it to take this branch.
-			if (IsWater <= 0.99f)
-			{
-				float PackedRefractionMaskLocal = RefractionMaskTexture.SampleLevel(SceneColorWithoutSingleLayerWaterSampler, ViewportUV, 0.0f).x;
-				bool bIsWaterLocal = false;
-				float RefractionMaskLocal = 0.0f;
-				UnpackWaterRefractionMask(PackedRefractionMaskLocal, bIsWaterLocal, RefractionMaskLocal);
-				RefractionMask = RefractionMaskLocal;
-				DistortedUV = ViewportUV;
-			}
-		}
+		BufferUVDistortion *= saturate((PixelBehindWaterSceneDepth - WaterSurfaceSceneDepth) * 1.0 / 30.0f);// *(1 - saturate(WaterDepth * 1.0 / 2000.0f));
+		const float2 DistortedUV = clamp(ViewportUV + BufferUVDistortion, SceneWithoutSingleLayerWaterMinUV, SceneWithoutSingleLayerWaterMaxUV);
+
+		// Sample SceneDepthWithoutWater at distorted UV
+		float DistortedSceneDeviceZ;
+		float DistortedSceneDepth;
+		float4 DistortedSceneColorWeights;
+		float4 DistortedRefractionMask;
+		SampleSceneDepthWithoutWaterBilateral(SceneDepthWithoutSingleLayerWaterTexture, SceneDepthWithoutSingleLayerWaterSampler, DistortedUV, SceneDepthWithoutSingleLayerWaterTextureSize, WaterSurfaceSceneDepth,
+#if SINGLE_LAYER_WATER_REFRACTION_MASK_AVAILABLE
+			RefractionMaskTexture,
 #endif
-
-		// Gather the four depth samples that would be used for bilinear interpolation. We need to know if any of these are invalid (above the water). If so, we have to discard this distorted sample and fall back. Otherwise we end up sampling foreground color from the bilinear color sample.
-		// This is only guaranteed to be artifact free if the input depth/color resolution is the same as the output. If the input is downsampled, this will still do the right thing but might still sample leaking color
-		float4 SceneDeviceZ4 = SceneDepthWithoutSingleLayerWaterTexture.Gather(SceneDepthWithoutSingleLayerWaterSampler, DistortedUV);
-		float4 SceneDepth4 = float4(ConvertFromDeviceZ(SceneDeviceZ4.x), ConvertFromDeviceZ(SceneDeviceZ4.y), ConvertFromDeviceZ(SceneDeviceZ4.z), ConvertFromDeviceZ(SceneDeviceZ4.w));
-		
-		float SceneDepth = 0.0f;
-		float SceneDeviceZ = 0.0f;
-		if(any(SceneDepth4 < WaterDepth))
+			DistortedSceneDeviceZ, DistortedSceneDepth, DistortedSceneColorWeights, DistortedRefractionMask);
+
+		// Pick a BehindWaterDepth depending on whether the weights of the refracted location are valid
+		float BehindWaterSceneDepth;
+		float BehindWaterSceneDeviceZ;
+		float4 BehindWaterSceneColorWeights;
+		float2 BehindWaterUV;
+		if (all(DistortedSceneColorWeights == 0.0f))
 		{
-			SceneDepth = PixelSceneDepth; // The pixel we try to sample is closer than the water pixel: fallback to no distortion
-			SceneDeviceZ = PixelSceneDeviceZ;
+			BehindWaterSceneDepth = PixelBehindWaterSceneDepth;
+			BehindWaterSceneDeviceZ = PixelBehindWaterSceneDeviceZ;
+			BehindWaterSceneColorWeights = PixelSceneColorWeights * PixelRefractionMask;
+			BehindWaterUV = ViewportUV;
 		}
 		else
 		{
-			// We can use bilinear interpolation here because all depth samples are guaranteed to be farther away than the water surface depth.
-			SceneDeviceZ = GetBilinearInterpolation(GetBilinearSampleLevelInfos(DistortedUV, SceneDepthWithoutSingleLayerWaterTextureSize, SceneDepthWithoutSingleLayerWaterTextureInvSize), SceneDeviceZ4.wzxy);
-			SceneDepth = ConvertFromDeviceZ(SceneDeviceZ);
-			ViewportUV = DistortedUV;
+			BehindWaterSceneDepth = DistortedSceneDepth;
+			BehindWaterSceneDeviceZ = DistortedSceneDeviceZ;
+			BehindWaterSceneColorWeights = DistortedSceneColorWeights * DistortedRefractionMask;
+			BehindWaterUV = DistortedUV;
 		}
 
 #endif // !SINGLE_LAYER_WATER_MOBILE_SHADING
 
 		// SvPositionToWorld will result in NaN if SceneDeviceZ=0
-		SceneDeviceZ = max(0.000000000001f, SceneDeviceZ);
+		BehindWaterSceneDeviceZ = max(0.000000000001f, BehindWaterSceneDeviceZ);
 
-		const float BehindWaterDeltaDepth = CameraIsUnderWater ? WaterDepth : max(0.0f, SceneDepth - WaterDepth);
+		const float WaterVolumeDepth = bCameraIsUnderWater ? WaterSurfaceSceneDepth : max(0.0f, BehindWaterSceneDepth - WaterSurfaceSceneDepth);
 
 #if SUBSTRATE_INLINE_SINGLELAYERWATER
 		MaterialParameters.SubstrateTree.BSDFs[0].SubstrateSanitizeBSDF();
@@ -188,13 +174,13 @@ WaterVolumeLightingOutput EvaluateWaterVolumeLighting(
 		const float3 AbsorptionCoeff = max(0.0, ExtinctionCoeff - ScatteringCoeff);
 		const float PhaseG = SLW_WATERPHASEG(SLWBSDF);
 		//Sample the optional Material Input ColorScaleBehindWater and fade it out at shorelines to avoid hard edge intersections
-		float3 ColorScaleBehindWater = lerp(1.0f, SLW_COLORSCALEBEHINDWATER(SLWBSDF), saturate(BehindWaterDeltaDepth * 0.02f));
+		float3 ColorScaleBehindWater = lerp(1.0f, SLW_COLORSCALEBEHINDWATER(SLWBSDF), saturate(WaterVolumeDepth * 0.02f));
 #else
 		const float3 ScatteringCoeff = max(0.0f, DFDemote(GetSingleLayerWaterMaterialOutput0(MaterialParameters)));
 		const float3 AbsorptionCoeff = max(0.0f, DFDemote(GetSingleLayerWaterMaterialOutput1(MaterialParameters)));
 		const float PhaseG = clamp(DFDemote(GetSingleLayerWaterMaterialOutput2(MaterialParameters).x), -1.0f, 1.0f);
 		//Sample the optional Material Input ColorScaleBehindWater and fade it out at shorelines to avoid hard edge intersections
-		float3 ColorScaleBehindWater = lerp(1.0f, max(0.0f, DFDemote(GetSingleLayerWaterMaterialOutput3(MaterialParameters))), saturate(BehindWaterDeltaDepth * 0.02f));
+		float3 ColorScaleBehindWater = lerp(1.0f, max(0.0f, DFDemote(GetSingleLayerWaterMaterialOutput3(MaterialParameters))), saturate(WaterVolumeDepth * 0.02f));
 		const float3 ExtinctionCoeff = ScatteringCoeff + AbsorptionCoeff;
 		const float3 Albedo = ScatteringCoeff / max(ExtinctionCoeff, 1e-7f);
 #endif
@@ -221,21 +207,15 @@ WaterVolumeLightingOutput EvaluateWaterVolumeLighting(
 		// We also apply transmittance from light to under water surface. However, the scene has been lit by many sources already.
 		// So the transmittance toabove surface is simply approximated using the travel distance from the scene pixel to the water top, assuming a flat water surface.
 		// We cannot combine this transmittance with the transmittance from view because this would change the behavior of the analytical integration of light scattering integration.
-		const float3 BehindWaterSceneWorldPos = SvPositionToTranslatedWorld(float4(MaterialParameters.SvPosition.xy, SceneDeviceZ, 1.0));
+		const float3 BehindWaterSceneWorldPos = SvPositionToTranslatedWorld(float4(MaterialParameters.SvPosition.xy, BehindWaterSceneDeviceZ, 1.0));
 		const float DistanceFromScenePixelToWaterTop = max(0.0, (GetTranslatedWorldPosition(MaterialParameters).z - BehindWaterSceneWorldPos.z));
 		const float3 MeanTransmittanceToLightSources = exp(-DistanceFromScenePixelToWaterTop * ExtinctionCoeff);
 
 #if SINGLE_LAYER_WATER_MOBILE_SHADING
 		const float3 BehindWaterSceneLuminance = 0.0f; // Cannot read back the scene color in this case
 #else
-		// We use the pixel SvPosition instead of the scene one pre refraction/distortion to avoid those extra ALUs.
-		float3 BehindWaterSceneLuminance = SceneColorWithoutSingleLayerWaterTexture.SampleLevel(SceneColorWithoutSingleLayerWaterSampler, ViewportUV, 0).rgb;
-		#if SINGLE_LAYER_WATER_ENCODED_SCENE_COLOR
-		// Undo the function in EncodeSceneColorForMaterialNode
-		BehindWaterSceneLuminance = pow(BehindWaterSceneLuminance.rgb, 4) * 10;
-		#endif
-		BehindWaterSceneLuminance = MeanTransmittanceToLightSources * ResolvedView.OneOverPreExposure * BehindWaterSceneLuminance;
-		BehindWaterSceneLuminance *= RefractionMask; // Multiply by refraction mask to essentially turn invalid refraction data black.
+		float3 BehindWaterSceneLuminance = SampleSceneColorWithoutWaterBilateral(SceneColorWithoutSingleLayerWaterTexture, SceneColorWithoutSingleLayerWaterSampler, BehindWaterUV, BehindWaterSceneColorWeights);
+		BehindWaterSceneLuminance *= MeanTransmittanceToLightSources * ResolvedView.OneOverPreExposure;
 #endif
 
 		float3 SunScattLuminance = DirLightPhaseValue * SunIlluminance;
@@ -244,7 +224,7 @@ WaterVolumeLightingOutput EvaluateWaterVolumeLighting(
 		const float MainDirLightFactor			= bSeparateMainDirLight ? 0.0f : 1.0f;
 		const float SeparatedMainDirLightFactor = bSeparateMainDirLight ? 1.0f : 0.0f;
 
-		const float3 OpticalDepth = ExtinctionCoeff * BehindWaterDeltaDepth;
+		const float3 OpticalDepth = ExtinctionCoeff * WaterVolumeDepth;
 		float3 Transmittance = exp(-OpticalDepth);
 		float3 IncomingLuminance = (AmbScattLuminance + SunScattLuminance * LightFunctionColor * DirectionalLightShadow * MainDirLightFactor);
 		// The SafeScatteringAmount should be in 0~1, add saturate to fix an artifact caused by half precision on mobile.
@@ -254,11 +234,11 @@ WaterVolumeLightingOutput EvaluateWaterVolumeLighting(
 		// Compute separated direcitonal light scattering if requested
 		float3 SeparatedIncomingLuminance = SunScattLuminance * LightFunctionColor * DirectionalLightShadow * SeparatedMainDirLightFactor;
 		// Apply Fresnel and water visibility to separated main dir light luminance.
-		SeparatedWaterMainDirLightScatteredLuminance += SeparatedIncomingLuminance * SafeScatteringAmount * (CameraIsUnderWater ? 1.0 : (1.0 - EnvBrdf)) * WaterVisibility;
+		SeparatedWaterMainDirLightScatteredLuminance += SeparatedIncomingLuminance * SafeScatteringAmount * (bCameraIsUnderWater ? 1.0 : (1.0 - EnvBrdf)) * WaterVisibility;
 
 		// Apply Fresnel effect to out-scattering towards the view
-		ScatteredLuminance *= CameraIsUnderWater ? 1.0 : (1.0 - EnvBrdf);	// Under water is less visible due to Fresnel effect
-		Transmittance *= CameraIsUnderWater ? (1.0 - EnvBrdf) : 1.0;		// Above	"		"		"		"		"
+		ScatteredLuminance *= bCameraIsUnderWater ? 1.0 : (1.0 - EnvBrdf);	// Under water is less visible due to Fresnel effect
+		Transmittance *= bCameraIsUnderWater ? (1.0 - EnvBrdf) : 1.0;		// Above	"		"		"		"		"
 
 		// Add single in-scattering apply colored transmittance to scene color
 		Output.Luminance = WaterVisibility * (ScatteredLuminance + Transmittance * (BehindWaterSceneLuminance * ColorScaleBehindWater));
