Removed duplicate ${custom_outputs} from MaterialTemplate.ush (only meaningful in the new material translator).

#rnx

#Synced-CL 50775131

#Synced-CL 50775131

diff --git a/Engine/Shaders/Private/MaterialTemplate.ush b/Engine/Shaders/Private/MaterialTemplate.ush
--- a/Engine/Shaders/Private/MaterialTemplate.ush
+++ b/Engine/Shaders/Private/MaterialTemplate.ush
@@ -362,9 +362,6 @@ SHADER_DISABLE_WARNINGS
 %{material_attributes_utilities}
 %{material_aggregates_declarations}
 
-/** Custom output declarations */
-%{custom_outputs}
-
 SHADER_POP_WARNINGS_STATE
 
 /** 
