Path Tracer: Allow refractive materials to automatically contribute to depth buffer

Refractive objects act as opaque in the path tracer, so always write to the depth buffer, even if output to depth buffer is disabled, as long as the transparency color is above the right threshold. This preserves the existing behavior for transparency and for glass with partial coverage.

#Synced-CL 50733702

diff --git a/Engine/Shaders/Private/PathTracing/PathTracingMaterialHitShader.usf b/Engine/Shaders/Private/PathTracing/PathTracingMaterialHitShader.usf
--- a/Engine/Shaders/Private/PathTracing/PathTracingMaterialHitShader.usf
+++ b/Engine/Shaders/Private/PathTracing/PathTracingMaterialHitShader.usf
@@ -1331,12 +1331,14 @@ RAY_TRACING_ENTRY_CLOSEST_HIT(PathTracingMaterialCHS,
 #endif
 
 	// Output depth for solid, masked 
-	// or when `output depth and velocity' is set and Opacity >= Opacity mask clip value to match the behavior of the rasterizer path. 
+	// or when `output depth and velocity' is set and Transparency is below the clip value to match the behavior of the rasterizer path. 
+	// If refraction is enabled, we always output depth (but still check against the transparency to support glass with partial coverage
 #if (MATERIALBLENDING_SOLID || MATERIALBLENDING_MASKED)
 	Payload.SetOutputDepth();
-#elif (TRANSLUCENT_WRITING_VELOCITY && \
+#elif ((TRANSLUCENT_WRITING_VELOCITY || PATH_TRACER_REFRACTION_ENABLED) && \
 	(MATERIALBLENDING_TRANSLUCENT || MATERIALBLENDING_ADDITIVE || MATERIALBLENDING_MODULATE || MATERIALBLENDING_ALPHACOMPOSITE || MATERIALBLENDING_ALPHAHOLDOUT))
-	if (GetMaterialOpacity(PixelMaterialInputs) >= GetMaterialOpacityMaskClipValue())
+	float TransparencyLimit = 1.0 - GetMaterialOpacityMaskClipValue();
+	if (any(Payload.TransparencyColor < TransparencyLimit))
 	{
 		Payload.SetOutputDepth();
 	}
