Fix half overflow warning. Halfs have +-65504 which is less than uint8 << 8 's range. This results in sub 1.0 influence weights on rigid skins. Instead, do uint16 normalization in full precision prior to applying half casting.

#Synced-CL 50609640

diff --git a/Engine/Shaders/Private/GpuSkinCommon.ush b/Engine/Shaders/Private/GpuSkinCommon.ush
--- a/Engine/Shaders/Private/GpuSkinCommon.ush
+++ b/Engine/Shaders/Private/GpuSkinCommon.ush
@@ -267,7 +267,7 @@ FBoneMatrix ComputeBoneMatrixWithUnlimitedInfluences(FComputeBoneMatrixInputs In
 			half BoneWeight;
 			if (BoneWeightSize > 1)
 			{
-				BoneWeight = half(InVertexWeights[BoneWeightOffset + 1] << 8 | InVertexWeights[BoneWeightOffset]) / 65535.0;
+				BoneWeight = half( (InVertexWeights[BoneWeightOffset + 1] << 8 | InVertexWeights[BoneWeightOffset]) / 65535.0 );
 			}
 			else
 			{
@@ -281,7 +281,7 @@ FBoneMatrix ComputeBoneMatrixWithUnlimitedInfluences(FComputeBoneMatrixInputs In
 			half BoneWeight;
 			if (BoneWeightSize > 1)
 			{
-				BoneWeight = half(InVertexWeights[BoneWeightOffset + 1] << 8 | InVertexWeights[BoneWeightOffset]) / 65535.0;
+				BoneWeight = half( (InVertexWeights[BoneWeightOffset + 1] << 8 | InVertexWeights[BoneWeightOffset]) / 65535.0 );
 			}
 			else
 			{
