Physics field direct evaluation in materials

#Synced-CL 50818556

diff --git a/Engine/Plugins/Experimental/ChaosNiagara/Shaders/NiagaraDataInterfacePhysicsField.ush b/Engine/Plugins/Experimental/ChaosNiagara/Shaders/NiagaraDataInterfacePhysicsField.ush
--- a/Engine/Plugins/Experimental/ChaosNiagara/Shaders/NiagaraDataInterfacePhysicsField.ush
+++ b/Engine/Plugins/Experimental/ChaosNiagara/Shaders/NiagaraDataInterfacePhysicsField.ush
@@ -25,6 +25,11 @@ int				{ParameterName}_ClipmapCount;
 int				{ParameterName}_TargetCount;
 int4			{ParameterName}_FieldTargets[MAX_PHYSICS_FIELD_TARGETS];
 
+Buffer<float>	{ParameterName}_NodesParams;
+Buffer<int>	{ParameterName}_NodesOffsets;
+Buffer<int>	{ParameterName}_TargetsOffsets;
+float {ParameterName}_TimeSeconds;
+
 float3 PhysicsVectorField_{ParameterName}(in float3 WorldPosition, in int TargetIndex, in bool bUseClipmap)
 {
 	// LWC_TODO: fix this when the physics field functions are lwc ready
@@ -39,8 +44,7 @@ float3 PhysicsVectorField_{ParameterName}(in float3 WorldPosition, in int Target
 	}
 	else
 	{
-		GFieldGroupThreadId = GGroupThreadId.x;
-		return PhysicsField_EvalPhysicsVectorField(LWCHackToFloat(SamplePos), {ParameterName}_TargetMappings[TargetIndex].x);
+		return PhysicsField_EvalPhysicsVectorField(LWCHackToFloat(SamplePos), {ParameterName}_TargetMappings[TargetIndex].x, {ParameterName}_NodesParams, {ParameterName}_NodesOffsets, {ParameterName}_TargetsOffsets, {ParameterName}_TimeSeconds);
 	}
 }
 
@@ -58,8 +62,7 @@ float PhysicsScalarField_{ParameterName}(in float3 WorldPosition, in int TargetI
 	}
 	else
 	{
-		GFieldGroupThreadId = GGroupThreadId.x;
-		return PhysicsField_EvalPhysicsScalarField(LWCHackToFloat(SamplePos), {ParameterName}_TargetMappings[TargetIndex].y);
+		return PhysicsField_EvalPhysicsScalarField(LWCHackToFloat(SamplePos), {ParameterName}_TargetMappings[TargetIndex].y, {ParameterName}_NodesParams, {ParameterName}_NodesOffsets, {ParameterName}_TargetsOffsets, {ParameterName}_TimeSeconds);
 	}
 }
 
@@ -77,8 +80,7 @@ int PhysicsIntegerField_{ParameterName}(in float3 WorldPosition, in int TargetIn
 	}
 	else
 	{
-		GFieldGroupThreadId = GGroupThreadId.x;
-		return PhysicsField_EvalPhysicsScalarField(LWCHackToFloat(SamplePos), {ParameterName}_TargetMappings[TargetIndex].z);
+		return PhysicsField_EvalPhysicsIntegerField(LWCHackToFloat(SamplePos), {ParameterName}_TargetMappings[TargetIndex].z, {ParameterName}_NodesParams, {ParameterName}_NodesOffsets, {ParameterName}_TargetsOffsets, {ParameterName}_TimeSeconds);
 	}
 }
 
diff --git a/Engine/Plugins/Experimental/ChaosNiagara/Source/ChaosNiagara/Private/NiagaraDataInterfacePhysicsField.cpp b/Engine/Plugins/Experimental/ChaosNiagara/Source/ChaosNiagara/Private/NiagaraDataInterfacePhysicsField.cpp
--- a/Engine/Plugins/Experimental/ChaosNiagara/Source/ChaosNiagara/Private/NiagaraDataInterfacePhysicsField.cpp
+++ b/Engine/Plugins/Experimental/ChaosNiagara/Source/ChaosNiagara/Private/NiagaraDataInterfacePhysicsField.cpp
@@ -30,17 +30,9 @@ struct FNiagaraPhysicsFieldDIFunctionVersion
 
 namespace NDIPhysicsFieldLocal
 {
-	BEGIN_SHADER_PARAMETER_STRUCT(FShaderGlobalParameters, )
-		SHADER_PARAMETER_SRV(Buffer<float>,	NodesParams)
-		SHADER_PARAMETER_SRV(Buffer<int>,	NodesOffsets)
-		SHADER_PARAMETER_SRV(Buffer<int>,	TargetsOffsets)
-		SHADER_PARAMETER(float,				TimeSeconds)
-	END_SHADER_PARAMETER_STRUCT()
-
 	BEGIN_SHADER_PARAMETER_STRUCT(FShaderInstanceParameters, )
 		SHADER_PARAMETER(int,				bClipmapAvailable)
 		SHADER_PARAMETER_ARRAY(FIntVector4, TargetMappings, [MAX_PHYSICS_FIELD_TARGETS])
-
 		SHADER_PARAMETER_SRV(Buffer<float>,	ClipmapBuffer)
 		SHADER_PARAMETER(FVector3f,			ClipmapCenter)
 		SHADER_PARAMETER(float, 			ClipmapDistance)
@@ -49,6 +41,10 @@ namespace NDIPhysicsFieldLocal
 		SHADER_PARAMETER(int, 				ClipmapCount)
 		SHADER_PARAMETER(int,				TargetCount)
 		SHADER_PARAMETER_ARRAY(FIntVector4, FieldTargets, [MAX_PHYSICS_FIELD_TARGETS])
+		SHADER_PARAMETER_SRV(Buffer<float>,	NodesParams)
+		SHADER_PARAMETER_SRV(Buffer<int>,	NodesOffsets)
+		SHADER_PARAMETER_SRV(Buffer<int>,	TargetsOffsets)
+		SHADER_PARAMETER(float,				TimeSeconds)
 	END_SHADER_PARAMETER_STRUCT()
 
 	static const TCHAR* TemplateShaderFilePath = TEXT("/Plugin/Experimental/ChaosNiagara/NiagaraDataInterfacePhysicsField.ush");
@@ -651,7 +647,6 @@ bool UNiagaraDataInterfacePhysicsField::AppendCompileHash(FNiagaraCompileHashVis
 	bool bSuccess = Super::AppendCompileHash(InVisitor);
 	bSuccess &= InVisitor->UpdateShaderFile(NDIPhysicsFieldLocal::TemplateShaderFilePath);
 	bSuccess &= InVisitor->UpdateShaderParameters<NDIPhysicsFieldLocal::FShaderInstanceParameters>();
-	bSuccess &= InVisitor->UpdateShaderParameters<NDIPhysicsFieldLocal::FShaderGlobalParameters>();
 	return bSuccess;
 }
 
@@ -668,7 +663,6 @@ void UNiagaraDataInterfacePhysicsField::GetParameterDefinitionHLSL(const FNiagar
 void UNiagaraDataInterfacePhysicsField::BuildShaderParameters(FNiagaraShaderParametersBuilder& ShaderParametersBuilder) const
 {
 	ShaderParametersBuilder.AddNestedStruct<NDIPhysicsFieldLocal::FShaderInstanceParameters>();
-	ShaderParametersBuilder.AddIncludedStruct<NDIPhysicsFieldLocal::FShaderGlobalParameters>();
 }
 
 void UNiagaraDataInterfacePhysicsField::SetShaderParameters(const FNiagaraDataInterfaceSetShaderParametersContext& Context) const
@@ -677,22 +671,20 @@ void UNiagaraDataInterfacePhysicsField::SetShaderParameters(const FNiagaraDataIn
 	FNDIFieldRenderData* ProxyData = InterfaceProxy.SystemInstancesToProxyData.Find(Context.GetSystemInstanceID());
 
 	NDIPhysicsFieldLocal::FShaderInstanceParameters* ShaderInstanceParameters = Context.GetParameterNestedStruct<NDIPhysicsFieldLocal::FShaderInstanceParameters>();
-	NDIPhysicsFieldLocal::FShaderGlobalParameters* ShaderGlobalParameters = Context.GetParameterIncludedStruct<NDIPhysicsFieldLocal::FShaderGlobalParameters>();
 
 	if (ProxyData != nullptr && ProxyData->FieldResource)
 	{
 		FPhysicsFieldResource* FieldResource = ProxyData->FieldResource;
-
-		// Global Parameters
-		ShaderGlobalParameters->NodesParams = FieldResource->NodesParams.SRV;
-		ShaderGlobalParameters->NodesOffsets = FieldResource->NodesOffsets.SRV;
-		ShaderGlobalParameters->TargetsOffsets = FieldResource->TargetsOffsets.SRV;
-		ShaderGlobalParameters->TimeSeconds = ProxyData->TimeSeconds;
-
+		
 		// Instance Parameters
 		const TArray<EFieldPhysicsType> VectorTypes = GetFieldTargetTypes(EFieldOutputType::Field_Output_Vector);
 		const TArray<EFieldPhysicsType> ScalarTypes = GetFieldTargetTypes(EFieldOutputType::Field_Output_Scalar);
 		const TArray<EFieldPhysicsType> IntegerTypes = GetFieldTargetTypes(EFieldOutputType::Field_Output_Integer);
+		
+		ShaderInstanceParameters->NodesParams = FieldResource->NodesParams.SRV;
+		ShaderInstanceParameters->NodesOffsets = FieldResource->NodesOffsets.SRV;
+		ShaderInstanceParameters->TargetsOffsets = FieldResource->TargetsOffsets.SRV;
+		ShaderInstanceParameters->TimeSeconds = ProxyData->TimeSeconds;
 
 		for (int32 i=0; i < MAX_PHYSICS_FIELD_TARGETS; ++i)
 		{
@@ -718,13 +710,11 @@ void UNiagaraDataInterfacePhysicsField::SetShaderParameters(const FNiagaraDataIn
 	}
 	else
 	{
-		// Global Parameters
-		ShaderGlobalParameters->NodesParams		= FNiagaraRenderer::GetDummyFloatBuffer();
-		ShaderGlobalParameters->NodesOffsets	= FNiagaraRenderer::GetDummyIntBuffer();
-		ShaderGlobalParameters->TargetsOffsets	= FNiagaraRenderer::GetDummyIntBuffer();
-		ShaderGlobalParameters->TimeSeconds		= 0.0f;
-
 		// Instance Parameters
+		ShaderInstanceParameters->NodesParams		= FNiagaraRenderer::GetDummyFloatBuffer();
+		ShaderInstanceParameters->NodesOffsets	= FNiagaraRenderer::GetDummyIntBuffer();
+		ShaderInstanceParameters->TargetsOffsets	= FNiagaraRenderer::GetDummyIntBuffer();
+		ShaderInstanceParameters->TimeSeconds		= 0.0f;
 		ShaderInstanceParameters->bClipmapAvailable	= 0;
 		ShaderInstanceParameters->ClipmapBuffer		= FNiagaraRenderer::GetDummyFloatBuffer();
 		ShaderInstanceParameters->ClipmapCenter		= FVector3f::ZeroVector;
diff --git a/Engine/Shaders/Private/PhysicsFieldBuilder.usf b/Engine/Shaders/Private/PhysicsFieldBuilder.usf
--- a/Engine/Shaders/Private/PhysicsFieldBuilder.usf
+++ b/Engine/Shaders/Private/PhysicsFieldBuilder.usf
@@ -33,6 +33,12 @@ int TargetsIndex[MAX_PHYSICS_FIELD_TARGETS];
 int NumCells;
 int TargetCount;
 
+Buffer<float> NodesParams;
+Buffer<int> NodesOffsets;
+Buffer<int> TargetsOffsets;
+float TimeSeconds;
+
+
 uint3 GetCellVoxel( in int ThreadId, in int CellIndex)
 {
 	const int LocalIndex = ThreadId - CellsOffsets[CellIndex];
@@ -77,7 +83,6 @@ void BuildPhysicsFieldClipmapCS(
 	uint3 DispatchThreadId : SV_DispatchThreadID,
     uint3 GroupThreadId : SV_GroupThreadID)
 {
-	GFieldGroupThreadId = GroupThreadId.x;
 	if ((DispatchThreadId.x < NumCells))
 	{
 		int CellIndex = 0;
@@ -114,20 +119,20 @@ void BuildPhysicsFieldClipmapCS(
 					const float3 SamplePosition = GetSamplePosition(ClipmapExtent, CellVoxel);
 					const int BufferOffset = GetBufferOffset(CellTarget, CellClipmap, CellVoxel);
 					
-					int LocalIndex = 0;
-					SampleFieldDatas(SamplePosition, LocalIndex, CellTarget,
-						ClipmapExtent* 2 * 256 * 4 / ClipmapResolution,  ClipmapDistance * 8 / ClipmapResolution);
+					FFieldEvalContext EvalContext;
+					SampleFieldData(SamplePosition, CellTarget,
+						ClipmapExtent * 2 * 256 * 4 / ClipmapResolution,  ClipmapDistance * 8 / ClipmapResolution, TimeSeconds, NodesParams, NodesOffsets, TargetsOffsets, EvalContext);
 					
 					const int AttributeOffset = ClipmapResolution * ClipmapResolution * ClipmapResolution * ClipmapCount;
-					if (LocalIndex == 3)
+					if (EvalContext.DataIndex  == 3)
 					{
-						FieldClipmap[BufferOffset] = SharedDatas[GFieldGroupThreadId][0] * ClipmapFade;
-						FieldClipmap[BufferOffset + 1 * AttributeOffset] = SharedDatas[GFieldGroupThreadId][1] * ClipmapFade;
-						FieldClipmap[BufferOffset + 2 * AttributeOffset] = SharedDatas[GFieldGroupThreadId][2] * ClipmapFade;
+						FieldClipmap[BufferOffset] = GetFieldValue(EvalContext, 0) * ClipmapFade;
+						FieldClipmap[BufferOffset + 1 * AttributeOffset] = GetFieldValue(EvalContext, 1) * ClipmapFade;
+						FieldClipmap[BufferOffset + 2 * AttributeOffset] = GetFieldValue(EvalContext, 2) * ClipmapFade;
 					}
-					else if (LocalIndex == 1)
+					else if (EvalContext.DataIndex == 1)
 					{
-						FieldClipmap[BufferOffset] = SharedDatas[GFieldGroupThreadId][0] * ClipmapFade;
+						FieldClipmap[BufferOffset] = GetFieldValue(EvalContext, 0) * ClipmapFade;
 					}
 				}
 			}
diff --git a/Engine/Shaders/Private/PhysicsFieldEval.ush b/Engine/Shaders/Private/PhysicsFieldEval.ush
--- a/Engine/Shaders/Private/PhysicsFieldEval.ush
+++ b/Engine/Shaders/Private/PhysicsFieldEval.ush
@@ -1,127 +1,123 @@
 // Copyright Epic Games, Inc. All Rights Reserved.
 
 /*=============================================================================
-	PhysicsFieldEval.ush
+    PhysicsFieldEval.ush
 =============================================================================*/
-
 #pragma once
-
 #include "Common.ush"
 
 /* -----------------------------------------------------------------
- * Field System constants and context
+ * Eval context
  * -----------------------------------------------------------------
  */
 
-Buffer<float> NodesParams;
-Buffer<int> NodesOffsets;
-Buffer<int> TargetsOffsets;
-float TimeSeconds;
+#define MAX_DATA_SIZE 16
+struct FFieldEvalContext
+{
+	/** Data buffer that will represent the result of each node evaluation that could be consumed by other nodes */
+    float DataBuffer[MAX_DATA_SIZE];
+	
+	/** Data current index in that data buffer */
+    int DataIndex;
+	
+	/** Time in seconds since the creation */
+	float TimeSeconds;
+};
 
-/* ----------------------------------------------------------------- 
- * Shared datas to store field evaluation
- * -----------------------------------------------------------------
- */
+// Define here each individual names
+#define PHYSICS_FIELD_NODES_PARAMS EvalContext_NodesParams
+#define PHYSICS_FIELD_NODES_OFFSETS EvalContext_NodesOffsets
+#define PHYSICS_FIELD_TARGETS_OFFSETS EvalContext_TargetsOffsets
 
-#define MAX_DATAS 16
-groupshared float SharedDatas[64][MAX_DATAS];
-static int GFieldGroupThreadId;
+// PS4/PS5 don't support Buffers within struct unfortunately. So need to pass the buffers individually
+#define PHYSICS_FIELD_BUFFERS in Buffer<float> PHYSICS_FIELD_NODES_PARAMS, in Buffer<int> PHYSICS_FIELD_NODES_OFFSETS,  in Buffer<int> PHYSICS_FIELD_TARGETS_OFFSETS
+#define PHYSICS_FIELD_ARGS PHYSICS_FIELD_NODES_PARAMS, PHYSICS_FIELD_NODES_OFFSETS, PHYSICS_FIELD_TARGETS_OFFSETS
 
-/* -----------------------------------------------------------------
- * Reandom Vector + Perlin Noise
- * -----------------------------------------------------------------
- */
- 
-static const int p[512] =
+float GetFieldValue(FFieldEvalContext EvalContext, int DataIndex)
 {
-	151, 160, 137, 91, 90, 15,
-    131, 13, 201, 95, 96, 53, 194, 233, 7, 225, 140, 36, 103, 30, 69, 142, 8, 99, 37, 240, 21, 10, 23,
-    190, 6, 148, 247, 120, 234, 75, 0, 26, 197, 62, 94, 252, 219, 203, 117, 35, 11, 32, 57, 177, 33,
-    88, 237, 149, 56, 87, 174, 20, 125, 136, 171, 168, 68, 175, 74, 165, 71, 134, 139, 48, 27, 166,
-    77, 146, 158, 231, 83, 111, 229, 122, 60, 211, 133, 230, 220, 105, 92, 41, 55, 46, 245, 40, 244,
-    102, 143, 54, 65, 25, 63, 161, 1, 216, 80, 73, 209, 76, 132, 187, 208, 89, 18, 169, 200, 196,
-    135, 130, 116, 188, 159, 86, 164, 100, 109, 198, 173, 186, 3, 64, 52, 217, 226, 250, 124, 123,
-    5, 202, 38, 147, 118, 126, 255, 82, 85, 212, 207, 206, 59, 227, 47, 16, 58, 17, 182, 189, 28, 42,
-    223, 183, 170, 213, 119, 248, 152, 2, 44, 154, 163, 70, 221, 153, 101, 155, 167, 43, 172, 9,
-    129, 22, 39, 253, 19, 98, 108, 110, 79, 113, 224, 232, 178, 185, 112, 104, 218, 246, 97, 228,
-    251, 34, 242, 193, 238, 210, 144, 12, 191, 179, 162, 241, 81, 51, 145, 235, 249, 14, 239, 107,
-    49, 192, 214, 31, 181, 199, 106, 157, 184, 84, 204, 176, 115, 121, 50, 45, 127, 4, 150, 254,
-    138, 236, 205, 93, 222, 114, 67, 29, 24, 72, 243, 141, 128, 195, 78, 66, 215, 61, 156, 180,
- 
-    151, 160, 137, 91, 90, 15,
-    131, 13, 201, 95, 96, 53, 194, 233, 7, 225, 140, 36, 103, 30, 69, 142, 8, 99, 37, 240, 21, 10, 23,
-    190, 6, 148, 247, 120, 234, 75, 0, 26, 197, 62, 94, 252, 219, 203, 117, 35, 11, 32, 57, 177, 33,
-    88, 237, 149, 56, 87, 174, 20, 125, 136, 171, 168, 68, 175, 74, 165, 71, 134, 139, 48, 27, 166,
-    77, 146, 158, 231, 83, 111, 229, 122, 60, 211, 133, 230, 220, 105, 92, 41, 55, 46, 245, 40, 244,
-    102, 143, 54, 65, 25, 63, 161, 1, 216, 80, 73, 209, 76, 132, 187, 208, 89, 18, 169, 200, 196,
-    135, 130, 116, 188, 159, 86, 164, 100, 109, 198, 173, 186, 3, 64, 52, 217, 226, 250, 124, 123,
-    5, 202, 38, 147, 118, 126, 255, 82, 85, 212, 207, 206, 59, 227, 47, 16, 58, 17, 182, 189, 28, 42,
-    223, 183, 170, 213, 119, 248, 152, 2, 44, 154, 163, 70, 221, 153, 101, 155, 167, 43, 172, 9,
-    129, 22, 39, 253, 19, 98, 108, 110, 79, 113, 224, 232, 178, 185, 112, 104, 218, 246, 97, 228,
-    251, 34, 242, 193, 238, 210, 144, 12, 191, 179, 162, 241, 81, 51, 145, 235, 249, 14, 239, 107,
-    49, 192, 214, 31, 181, 199, 106, 157, 184, 84, 204, 176, 115, 121, 50, 45, 127, 4, 150, 254,
-    138, 236, 205, 93, 222, 114, 67, 29, 24, 72, 243, 141, 128, 195, 78, 66, 215, 61, 156, 180
-};
- 
-float FadeValue(float t)
+    return (DataIndex < MAX_DATA_SIZE) ? EvalContext.DataBuffer[DataIndex] : 0;
+}
+
+float3 GetFieldVector(FFieldEvalContext EvalContext, int DataIndex)
 {
-	return t * t * t * (t * (t * 6 - 15) + 10);
+	return ((DataIndex+2) < MAX_DATA_SIZE) ? float3(EvalContext.DataBuffer[DataIndex], EvalContext.DataBuffer[DataIndex+1], EvalContext.DataBuffer[DataIndex+2]) : float3(0,0,0);
 }
- 
-float LerpValue(float t, float a, float b)
+
+void SetFieldValue(inout FFieldEvalContext EvalContext, int DataIndex, float FieldValue)
 {
-	return a + t * (b - a);
+	if(DataIndex < MAX_DATA_SIZE)
+	{
+		EvalContext.DataBuffer[DataIndex] = FieldValue;
+	}
 }
- 
-float GradientDirection(int hash, float x, float y, float z)
+
+void InitEvalContext(inout FFieldEvalContext EvalContext, in float TimeSeconds)
 {
-    // CONVERT LO 4 BITS OF HASH CODE INTO 12 GRADIENT DIRECTIONS.
-	int h = hash & 15;
-	float u = h < 8 ? x : y;
-	float v = h < 4 ? y : h == 12 || h == 14 ? x : z;
-	return ((h & 1) == 0 ? u : -u) + ((h & 2) == 0 ? v : -v);
+	[unroll]
+	for (int DataIndex = 0; DataIndex < MAX_DATA_SIZE; ++DataIndex)
+	{
+		EvalContext.DataBuffer[DataIndex] = 0.0;
+	}
+    EvalContext.DataIndex = 0;
+	EvalContext.TimeSeconds = TimeSeconds;
 }
- 
-float PerlinNoise(float3 pos)
+
+/* -----------------------------------------------------------------
+ * Perlin noise
+ * -----------------------------------------------------------------
+ */
+
+float3 ComputeHashValue(uint3 v)
 {
-    // FIND UNIT CUBE THAT CONTAINS POINT.
-	int X = (int) floor(pos.x) & 255;
-	int Y = (int) floor(pos.y) & 255;
-	int Z = (int) floor(pos.z) & 255;
-   
-    // FIND RELATIVE X,Y,Z OF POINT IN CUBE.
-	float x = pos.x - floor(pos.x);
-	float y = pos.y - floor(pos.y);
-	float z = pos.z - floor(pos.z);
-   
-    // COMPUTE FADE CURVES FOR EACH OF X,Y,Z.
-	float u = FadeValue(x);
-	float v = FadeValue(y);
-	float w = FadeValue(z);
-   
-    // HASH COORDINATES OF THE 8 CUBE CORNERS
-	int A = p[X] + Y;
-	int AA = p[A] + Z;
-	int AB = p[A + 1] + Z;
-	int B = p[X + 1] + Y;
-	int BA = p[B] + Z;
-	int BB = p[B + 1] + Z;
- 
-    // AND ADD BLENDED RESULTS FROM  8 CORNERS OF CUBE
-	return clamp(0.97 * LerpValue(w, LerpValue(v, LerpValue(u, GradientDirection(p[AA], x, y, z), // CORNER 0      
-           GradientDirection(p[BA], x - 1, y, z)), // CORNER 1                              
-           LerpValue(u, GradientDirection(p[AB], x, y - 1, z), // CORNER 2                        
-           GradientDirection(p[BB], x - 1, y - 1, z))), // CORNER 3                      
-           LerpValue(v, LerpValue(u, GradientDirection(p[AA + 1], x, y, z - 1), // CORNER 4            
-           GradientDirection(p[BA + 1], x - 1, y, z - 1)), // CORNER 5                  
-           LerpValue(u, GradientDirection(p[AB + 1], x, y - 1, z - 1), // CORNER 6  
-           GradientDirection(p[BB + 1], x - 1, y - 1, z - 1)))), -1.0, 1.0); // CORNER 7
+	// PCG Step
+	v = v * 1664525u + 1013904223u;
+	v.x += v.y * v.z; v.y += v.z * v.x; v.z += v.x * v.y;
+	v ^= v >> 16u;
+	v.x += v.y * v.z; v.y += v.z * v.x; v.z += v.x * v.y;
+    
+	// Efficiently map to [-1.0, 1.0]
+	// 0xFFFFFFu is the mask for 24-bit precision (mantissa of a float)
+	return (float3(v & 0xFFFFFFu) / float(0xFFFFFFu)) * 2.0 - 1.0;
+}
+
+float ComputePerlinNoise(float3 pos)
+{
+	float3 floorPos = floor(pos);
+	float3 f = frac(pos);
+    
+	// Use uint3 for the base grid coordinate
+	uint3 uPos = uint3(abs(floorPos)); // abs handles negative world coords safely
+
+	// Quintic interpolation
+	float3 u = f * f * f * (f * (f * 6.0 - 15.0) + 10.0);
+
+	// Evaluate gradients at 8 corners using direct uint3 offsets
+	// This avoids all casting inside the hot loop
+	float n000 = dot(ComputeHashValue(uPos + uint3(0,0,0)), f - float3(0,0,0));
+	float n100 = dot(ComputeHashValue(uPos + uint3(1,0,0)), f - float3(1,0,0));
+	float n010 = dot(ComputeHashValue(uPos + uint3(0,1,0)), f - float3(0,1,0));
+	float n110 = dot(ComputeHashValue(uPos + uint3(1,1,0)), f - float3(1,1,0));
+	float n001 = dot(ComputeHashValue(uPos + uint3(0,0,1)), f - float3(0,0,1));
+	float n101 = dot(ComputeHashValue(uPos + uint3(1,0,1)), f - float3(1,0,1));
+	float n011 = dot(ComputeHashValue(uPos + uint3(0,1,1)), f - float3(0,1,1));
+	float n111 = dot(ComputeHashValue(uPos + uint3(1,1,1)), f - float3(1,1,1));
+
+	// Nested Lerps for trilinear interpolation
+	return lerp(
+		lerp(lerp(n000, n100, u.x), lerp(n010, n110, u.x), u.y),
+		lerp(lerp(n001, n101, u.x), lerp(n011, n111, u.x), u.y), 
+		u.z
+	);
 }
 
-uint3 RandomVector(int3 p)
+/* -----------------------------------------------------------------
+ * Random Vector
+ * -----------------------------------------------------------------
+ */
+uint3 ComputeRandomVector(int3 pos)
 {
 	// taking a signed int then reinterpreting as unsigned gives good behavior for negatives
-	uint3 v = uint3(p);
+	uint3 v = uint3(pos);
 
 	// Linear congruential step. These LCG constants are from Numerical Recipies
 	// For additional #'s, PCG would do multiple LCG steps and scramble each on output
@@ -165,21 +161,21 @@ uint3 RandomVector(int3 p)
 
 #define NONE_TARGET 0
 #define DYNAMIC_STATE 1
-#define	LINEAR_FORCE 2
-#define	EXTERNAL_STRAIN	3
-#define	KILL_PARTICLES 4
-#define	LINEAR_VELOCITY 5
-#define	ANGULAR_VELOCITY 6
-#define	ANGULAR_TORQUE 7
-#define	INTERNAL_STRAIN 8
-#define	DISABLED_THRESHOLD 9
-#define	SLEEPING_THRESHOLD 10
-#define	POSITION_STATIC 11
-#define	POSITION_ANIMATED 12
-#define	POSITION_TARGET 13
-#define	DYNAMIC_CONSTRAINT 14
-#define	COLLISION_GROUP 15
-#define	ACTIVATE_DISABLED 16
+#define LINEAR_FORCE 2
+#define EXTERNAL_STRAIN    3
+#define KILL_PARTICLES 4
+#define LINEAR_VELOCITY 5
+#define ANGULAR_VELOCITY 6
+#define ANGULAR_TORQUE 7
+#define INTERNAL_STRAIN 8
+#define DISABLED_THRESHOLD 9
+#define SLEEPING_THRESHOLD 10
+#define POSITION_STATIC 11
+#define POSITION_ANIMATED 12
+#define POSITION_TARGET 13
+#define DYNAMIC_CONSTRAINT 14
+#define COLLISION_GROUP 15
+#define ACTIVATE_DISABLED 16
 
 #define NONE_NODE 0
 #define UNIFORM_INTEGER 1
@@ -202,7 +198,7 @@ uint3 RandomVector(int3 p)
 #define SET_IFF_NOT_INTERIOR 1
 #define SET_IFF_NOT_EXTERIOR 2
 
-#define NONE_FALLOFF 0 
+#define NONE_FALLOFF 0
 #define FALLOFF_LINEAR 1
 #define FALLOFF_INVERSE 2
 #define FALLOFF_SQUARED 3
@@ -226,18 +222,53 @@ uint3 RandomVector(int3 p)
  * -----------------------------------------------------------------
  */
 
-float3 FieldRotateVectorByQuat(in float3 Vector, in float4 Quat)
+float3 RotateVectorByQuat(in float3 Vector, in float4 Quat)
 {
 	float3 T = 2.0 * cross(Quat.xyz, Vector);
 	return Vector + Quat.w * T + cross(Quat.xyz, T);
 }
 
-float3 FieldUnRotateVectorByQuat(in float3 Vector, in float4 Quat)
+float3 UnRotateVectorByQuat(in float3 Vector, in float4 Quat)
 {
 	float3 T = 2.0 * cross(Quat.xyz, Vector);
 	return Vector - Quat.w * T + cross(Quat.xyz, T);
 }
 
+/* -----------------------------------------------------------------
+ * Common Falloff functions
+ * -----------------------------------------------------------------
+ */
+
+float ComputeFalloffValue(int FieldFalloff, float FieldMagnitude, float DeltaDistance)
+{
+	float Result = 0.0;
+    
+	// Modern HLSL compilers handle switch statements very efficiently
+	switch (FieldFalloff)
+	{
+		case NONE_FALLOFF: 
+			Result = 1.0;
+			break;
+		case FALLOFF_LINEAR: 
+			Result = DeltaDistance;
+			break;
+		case 3: 
+			Result = DeltaDistance * DeltaDistance;
+			break;
+		case FALLOFF_INVERSE: 
+			Result = 2.0 * (1.0 - 1.0 / (DeltaDistance + 1.0));
+			break;
+		case FALLOFF_LOGARITHMIC: 
+			// log2(x) is a native hardware instruction, much faster than log(x)/log(2)
+			Result = log2(DeltaDistance + 1.0);
+			break;
+		default:
+			break;
+	}
+    
+	return Result * FieldMagnitude;
+}
+
 /* -----------------------------------------------------------------
  * Uniform Integer Field
  * -----------------------------------------------------------------
@@ -245,9 +276,9 @@ float3 FieldUnRotateVectorByQuat(in float3 Vector, in float4 Quat)
 
 #define UNIFORM_INTEGER_MAGNITUDE 0
 
-void EvaluateUniformInteger(in float3 SamplePosition, in int NodeOffset, inout int LocalIndex)
+void EvaluateUniformInteger(float3 SamplePosition, int NodeOffset, inout FFieldEvalContext EvalContext, PHYSICS_FIELD_BUFFERS)
 {
-	SharedDatas[GFieldGroupThreadId][LocalIndex++] = NodesParams[NodeOffset + UNIFORM_INTEGER_MAGNITUDE];
+    SetFieldValue(EvalContext, EvalContext.DataIndex++, PHYSICS_FIELD_NODES_PARAMS[NodeOffset + UNIFORM_INTEGER_MAGNITUDE]);
 }
 
 /* -----------------------------------------------------------------
@@ -263,21 +294,28 @@ void EvaluateUniformInteger(in float3 SamplePosition, in int NodeOffset, inout i
 #define RADIAL_MASK_INTEGER_EXTERIOR_VALUE 5
 #define RADIAL_MASK_INTEGER_SET_MASK_CONDITION 6
 
-void EvaluateRadialMaskInteger(in float3 SamplePosition, in int NodeOffset, inout int LocalIndex)
-{
-	const float3 DeltaPosition = float3(NodesParams[NodeOffset + RADIAL_MASK_INTEGER_POSITIONX],
-										 NodesParams[NodeOffset + RADIAL_MASK_INTEGER_POSITIONY],
-										 NodesParams[NodeOffset + RADIAL_MASK_INTEGER_POSITIONZ]) - SamplePosition;
-	const float DistanceSquared = dot(DeltaPosition, DeltaPosition);
-	const int DeltaResult = (DistanceSquared < NodesParams[NodeOffset + RADIAL_MASK_INTEGER_RADIUS] *
-											 NodesParams[NodeOffset + RADIAL_MASK_INTEGER_RADIUS]) ?
-											 (int)NodesParams[NodeOffset + RADIAL_MASK_INTEGER_INTERIOR_VALUE]:
-											 (int)NodesParams[NodeOffset + RADIAL_MASK_INTEGER_EXTERIOR_VALUE];
-	const int MaskCondition = (int)NodesParams[NodeOffset + RADIAL_MASK_INTEGER_SET_MASK_CONDITION];
-	
-	SharedDatas[GFieldGroupThreadId][LocalIndex++] = (MaskCondition == SET_ALWAYS) ? DeltaResult :
-					 ((MaskCondition == SET_IFF_NOT_INTERIOR) && (SharedDatas[GFieldGroupThreadId][LocalIndex] != NodesParams[NodeOffset + RADIAL_MASK_INTEGER_INTERIOR_VALUE])) ? DeltaResult :
-					 ((MaskCondition == SET_IFF_NOT_EXTERIOR) && (SharedDatas[GFieldGroupThreadId][LocalIndex] != NodesParams[NodeOffset + RADIAL_MASK_INTEGER_EXTERIOR_VALUE])) ? DeltaResult : SharedDatas[GFieldGroupThreadId][LocalIndex];
+void EvaluateRadialMaskInteger(float3 SamplePosition, int NodeOffset, inout FFieldEvalContext EvalContext, PHYSICS_FIELD_BUFFERS)
+{
+    const float MaskRadius = PHYSICS_FIELD_NODES_PARAMS[NodeOffset + RADIAL_MASK_INTEGER_RADIUS];
+    const float3 CenterPosition = float3(PHYSICS_FIELD_NODES_PARAMS[NodeOffset + RADIAL_MASK_INTEGER_POSITIONX], 
+    									 PHYSICS_FIELD_NODES_PARAMS[NodeOffset + RADIAL_MASK_INTEGER_POSITIONY], 
+    									 PHYSICS_FIELD_NODES_PARAMS[NodeOffset + RADIAL_MASK_INTEGER_POSITIONZ]);
+    const int InteriorValue = (int)PHYSICS_FIELD_NODES_PARAMS[NodeOffset + RADIAL_MASK_INTEGER_INTERIOR_VALUE];
+    const int ExteriorValue = (int)PHYSICS_FIELD_NODES_PARAMS[NodeOffset + RADIAL_MASK_INTEGER_EXTERIOR_VALUE];
+    const int MaskCondition = (int)PHYSICS_FIELD_NODES_PARAMS[NodeOffset + RADIAL_MASK_INTEGER_SET_MASK_CONDITION];
+    
+    const float3 DeltaPosition = SamplePosition - CenterPosition;
+    const float DistanceSquared = dot(DeltaPosition, DeltaPosition);
+    const float RadiusSquared = MaskRadius * MaskRadius;
+    
+    const int DeltaResult = DistanceSquared < RadiusSquared ? InteriorValue : ExteriorValue;
+    const int CurrentValue = (int)GetFieldValue(EvalContext, EvalContext.DataIndex);
+    
+    const float ResultScalar = (MaskCondition == SET_ALWAYS) ? DeltaResult :
+        ((MaskCondition == SET_IFF_NOT_INTERIOR && CurrentValue != InteriorValue) ? DeltaResult :
+        ((MaskCondition == SET_IFF_NOT_EXTERIOR && CurrentValue != ExteriorValue) ? DeltaResult : CurrentValue));
+    
+    SetFieldValue(EvalContext, EvalContext.DataIndex++, ResultScalar);
 }
 
 /* -----------------------------------------------------------------
@@ -287,42 +325,9 @@ void EvaluateRadialMaskInteger(in float3 SamplePosition, in int NodeOffset, inou
 
 #define UNIFORM_SCALAR_MAGNITUDE 0
 
-void EvaluateUniformScalar(in float3 SamplePosition, in int NodeOffset, inout int LocalIndex)
-{
-	SharedDatas[GFieldGroupThreadId][LocalIndex++] = NodesParams[NodeOffset + UNIFORM_SCALAR_MAGNITUDE];
-}
-
-/* -----------------------------------------------------------------
- * Common Falloff functions
- * -----------------------------------------------------------------
- */
-
-void SetFalloffValue(in int FieldFalloff, in float FieldMagnitude, in float DeltaDistance, out float OutNodeResult)
+void EvaluateUniformScalar(float3 SamplePosition, int NodeOffset, inout FFieldEvalContext EvalContext, PHYSICS_FIELD_BUFFERS)
 {
-	if (FieldFalloff == NONE_FALLOFF)
-	{
-		OutNodeResult = FieldMagnitude;
-	}
-	else if (FieldFalloff == FALLOFF_LINEAR)
-	{
-		OutNodeResult = FieldMagnitude * DeltaDistance;
-	}
-	else if (FieldFalloff == FALLOFF_SQUARED)
-	{
-		OutNodeResult = FieldMagnitude * DeltaDistance * DeltaDistance;
-	}
-	else if (FieldFalloff == FALLOFF_INVERSE)
-	{
-		OutNodeResult = FieldMagnitude * 2.0 * (1.0 - 1.0 / (DeltaDistance+1.0));
-	}
-	else if (FieldFalloff == FALLOFF_LOGARITHMIC)
-	{
-		OutNodeResult = FieldMagnitude * log(DeltaDistance + 1.0) / log(2.0);
-	}
-	else
-	{
-		OutNodeResult = 0.0;
-	}
+    SetFieldValue(EvalContext, EvalContext.DataIndex++, PHYSICS_FIELD_NODES_PARAMS[NodeOffset + UNIFORM_SCALAR_MAGNITUDE]);
 }
 
 /* -----------------------------------------------------------------
@@ -342,51 +347,52 @@ void SetFalloffValue(in int FieldFalloff, in float FieldMagnitude, in float Delt
 
 static const float WAVE_PI = 3.14159265f;
 
-void EvaluateWaveScalar(in float3 SamplePosition, in int NodeOffset, inout int LocalIndex, in float MinLength)
+void EvaluateWaveScalar(float3 SamplePosition, int NodeOffset, float MinLength, inout FFieldEvalContext EvalContext, PHYSICS_FIELD_BUFFERS)
 {
-	const float Distance = length(float3(NodesParams[NodeOffset + WAVE_SCALAR_POSITIONX],
-										 NodesParams[NodeOffset + WAVE_SCALAR_POSITIONY],
-										 NodesParams[NodeOffset + WAVE_SCALAR_POSITIONZ]) - SamplePosition);
-	
-	const float Wavelength = NodesParams[NodeOffset + WAVE_SCALAR_WAVELENGTH];
-	const float Velocity = Wavelength / NodesParams[NodeOffset + WAVE_SCALAR_PERIOD];
+    const float Magnitude = PHYSICS_FIELD_NODES_PARAMS[NodeOffset + WAVE_SCALAR_MAGNITUDE];
+    const float3 CenterPosition = float3(PHYSICS_FIELD_NODES_PARAMS[NodeOffset + WAVE_SCALAR_POSITIONX], 
+    									 PHYSICS_FIELD_NODES_PARAMS[NodeOffset + WAVE_SCALAR_POSITIONY], 
+    									 PHYSICS_FIELD_NODES_PARAMS[NodeOffset + WAVE_SCALAR_POSITIONZ]);
 	
-	const float Wavenumber = 2.0 * WAVE_PI / max(Wavelength, MinLength);
-	const float DeltaTime = max(TimeSeconds - NodesParams[NodeOffset + WAVE_SCALAR_TIME], 0.0);
-	const float Radius = Velocity * DeltaTime;
-	const float Phase = Wavenumber * (Distance - Radius);
+    const float Wavelength = PHYSICS_FIELD_NODES_PARAMS[NodeOffset + WAVE_SCALAR_WAVELENGTH];
+    const float WavePeriod = max(PHYSICS_FIELD_NODES_PARAMS[NodeOffset + WAVE_SCALAR_PERIOD], 0.0001);
+    const float StartTime = PHYSICS_FIELD_NODES_PARAMS[NodeOffset + WAVE_SCALAR_TIME];
+    const int WaveFunction = (int)PHYSICS_FIELD_NODES_PARAMS[NodeOffset + WAVE_SCALAR_FUNCTION];
+    const int FalloffType = (int)PHYSICS_FIELD_NODES_PARAMS[NodeOffset + WAVE_SCALAR_FALLOFF];
 	
-	float Result = 0.0;
+    const float DistanceValue = length(SamplePosition - CenterPosition);
 	
-	const int Function = (int) NodesParams[NodeOffset + WAVE_SCALAR_FUNCTION];
-	const float Magnitude = NodesParams[NodeOffset + WAVE_SCALAR_MAGNITUDE];
+	const float WaveVelocity = Wavelength * rcp(WavePeriod);
+	const float WaveNumber = 2.0 * WAVE_PI * rcp(max(Wavelength, MinLength));
+	const float DeltaTime = max(EvalContext.TimeSeconds - StartTime, 0.0);
+	const float WaveRadius = WaveVelocity * DeltaTime;
+	const float WavePhase = WaveNumber * (DistanceValue - WaveRadius);
 	
-	[branch]
-	if (Function == WAVE_GAUSSIAN)
-	{
-		Result = Magnitude * exp(-Phase * Phase);
-	}
-	else if (Function == WAVE_COSINE)
-	{
-		Result = Magnitude * cos(Phase);
-	}
-	else if (Function == WAVE_FALLOFF)
-	{
-		const int Falloff = (int) NodesParams[NodeOffset + WAVE_SCALAR_FALLOFF];
-		
-		if ((Distance < Radius) && (Radius > 0.0))
-		{
-			const float Fraction = (1.0 - Distance / Radius);
-			SetFalloffValue(Falloff, Magnitude, Fraction, Result);
-		}
-	}
-	else if (Function == WAVE_DECAY)
-	{
-		const float Decay = DeltaTime / NodesParams[NodeOffset + WAVE_SCALAR_PERIOD];
-		Result = Magnitude * exp(-Decay * Decay);
-
-	}
-	SharedDatas[GFieldGroupThreadId][LocalIndex++] = Result;
+	float ResultScalar = 0.0f;
+	switch (WaveFunction)
+	{
+		case WAVE_GAUSSIAN:
+			ResultScalar = Magnitude * exp(-WavePhase * WavePhase);
+			break;
+		case WAVE_COSINE:
+			ResultScalar = Magnitude * cos(WavePhase);
+			break;
+		case WAVE_FALLOFF:
+			{
+				const bool bHasFalloff = (DistanceValue < WaveRadius) && (WaveRadius > 0.0);
+				ResultScalar = bHasFalloff ? ComputeFalloffValue(FalloffType, Magnitude, (1.0 - DistanceValue * rcp(WaveRadius))) : 0.0;
+				break;
+			}
+		case WAVE_DECAY:
+			{
+				const float WaveDecay = DeltaTime * rcp(WavePeriod);
+				ResultScalar = Magnitude * exp(-WaveDecay * WaveDecay);
+				break;
+			}
+		default : break;
+	}
+    
+    SetFieldValue(EvalContext, EvalContext.DataIndex++, ResultScalar);
 }
 
 /* -----------------------------------------------------------------
@@ -404,26 +410,30 @@ void EvaluateWaveScalar(in float3 SamplePosition, in int NodeOffset, inout int L
 #define RADIAL_FALLOFF_POSITIONZ 7
 #define RADIAL_FALLOFF_FALLOFF 8
 
-void EvaluateRadialFalloffScalar(in float3 SamplePosition, in int NodeOffset, inout int LocalIndex)
+void EvaluateRadialFalloffScalar(float3 SamplePosition, int NodeOffset, inout FFieldEvalContext EvalContext, PHYSICS_FIELD_BUFFERS)
 {
-	float ScalarDatas = NodesParams[NodeOffset + RADIAL_FALLOFF_DEFAULT];
-	const float DeltaRange = NodesParams[NodeOffset + RADIAL_FALLOFF_MAX_RANGE] - NodesParams[NodeOffset + RADIAL_FALLOFF_MIN_RANGE];
-
-	const float3 DeltaPosition = float3(NodesParams[NodeOffset + RADIAL_FALLOFF_POSITIONX],
-									    NodesParams[NodeOffset + RADIAL_FALLOFF_POSITIONY],
-										NodesParams[NodeOffset + RADIAL_FALLOFF_POSITIONZ]) - SamplePosition;
-	const float LocalDistance = length(DeltaPosition);
-	const float FieldRadius = NodesParams[NodeOffset + RADIAL_FALLOFF_RADIUS];
+	const float Magnitude = PHYSICS_FIELD_NODES_PARAMS[NodeOffset + RADIAL_FALLOFF_MAGNITUDE];
+	const float MinRange = PHYSICS_FIELD_NODES_PARAMS[NodeOffset + RADIAL_FALLOFF_MIN_RANGE];
+	const float MaxRange = PHYSICS_FIELD_NODES_PARAMS[NodeOffset + RADIAL_FALLOFF_MAX_RANGE];
+	const float DefaultValue = PHYSICS_FIELD_NODES_PARAMS[NodeOffset + RADIAL_FALLOFF_DEFAULT];
+	const float FalloffRadius = PHYSICS_FIELD_NODES_PARAMS[NodeOffset + RADIAL_FALLOFF_RADIUS];
+	
+	const float3 CenterPosition = float3(PHYSICS_FIELD_NODES_PARAMS[NodeOffset + RADIAL_FALLOFF_POSITIONX], 
+		PHYSICS_FIELD_NODES_PARAMS[NodeOffset + RADIAL_FALLOFF_POSITIONY], PHYSICS_FIELD_NODES_PARAMS[NodeOffset + RADIAL_FALLOFF_POSITIONZ]);
+	const int FalloffType = (int)PHYSICS_FIELD_NODES_PARAMS[NodeOffset + RADIAL_FALLOFF_FALLOFF];
 	
-	if (FieldRadius > 0.0 && LocalDistance < FieldRadius)
+	const float3 DeltaPosition = SamplePosition - CenterPosition;
+	const float DistanceSquared = dot(DeltaPosition, DeltaPosition);
+	const float RadiusSquared = FalloffRadius * FalloffRadius;
+    
+	float ResultScalar = DefaultValue;
+	if (RadiusSquared > 0.0 && DistanceSquared < RadiusSquared)
 	{
-		const float DeltaDistance = 1.0 - LocalDistance / FieldRadius;
-
-		SetFalloffValue((int)NodesParams[NodeOffset + RADIAL_FALLOFF_FALLOFF],
-						1.0, DeltaDistance, ScalarDatas);
-		ScalarDatas = NodesParams[NodeOffset + RADIAL_FALLOFF_MAGNITUDE] * (NodesParams[NodeOffset + RADIAL_FALLOFF_MIN_RANGE] + ScalarDatas * DeltaRange);
+		const float FalloffValue = ComputeFalloffValue(FalloffType, 1.0, 1.0 - sqrt(DistanceSquared) * rcp(FalloffRadius));
+		ResultScalar = Magnitude * (MinRange + FalloffValue * (MaxRange - MinRange));
 	}
-	SharedDatas[GFieldGroupThreadId][LocalIndex++] = ScalarDatas;
+    
+	SetFieldValue(EvalContext, EvalContext.DataIndex++, ResultScalar);
 }
 
 /* -----------------------------------------------------------------
@@ -444,33 +454,36 @@ void EvaluateRadialFalloffScalar(in float3 SamplePosition, in int NodeOffset, in
 #define PLANE_FALLOFF_NORMALZ 10
 #define PLANE_FALLOFF_FALLOFF 11
 
-void EvaluatePlaneFalloffScalar(in float3 SamplePosition, in int NodeOffset, inout int LocalIndex)
+void EvaluatePlaneFalloffScalar(float3 SamplePosition, int NodeOffset, inout FFieldEvalContext EvalContext, PHYSICS_FIELD_BUFFERS)
 {
-	float ScalarDatas = NodesParams[NodeOffset + PLANE_FALLOFF_DEFAULT];
-	const float DeltaRange = NodesParams[NodeOffset + PLANE_FALLOFF_MAX_RANGE] - NodesParams[NodeOffset + PLANE_FALLOFF_MIN_RANGE];
-
-	const float3 PlaneBase = float3(NodesParams[NodeOffset + PLANE_FALLOFF_POSITIONX],
-								    NodesParams[NodeOffset + PLANE_FALLOFF_POSITIONY],
-								    NodesParams[NodeOffset + PLANE_FALLOFF_POSITIONZ]);
-
-	const float3 PlaneNormal = float3(NodesParams[NodeOffset + PLANE_FALLOFF_NORMALX],
-							          NodesParams[NodeOffset + PLANE_FALLOFF_NORMALY],
-								      NodesParams[NodeOffset + PLANE_FALLOFF_NORMALZ]);
-
-	const float PlaneOffset = dot(PlaneBase, PlaneNormal);
-
+	const float Magnitude = PHYSICS_FIELD_NODES_PARAMS[NodeOffset + PLANE_FALLOFF_MAGNITUDE];
+	const float MinRange = PHYSICS_FIELD_NODES_PARAMS[NodeOffset + PLANE_FALLOFF_MIN_RANGE];
+	const float MaxRange = PHYSICS_FIELD_NODES_PARAMS[NodeOffset + PLANE_FALLOFF_MAX_RANGE];
+	const float DefaultValue = PHYSICS_FIELD_NODES_PARAMS[NodeOffset + PLANE_FALLOFF_DEFAULT];
+	const float PlaneDistance = PHYSICS_FIELD_NODES_PARAMS[NodeOffset + PLANE_FALLOFF_DISTANCE];
+	
+	const float3 PlanePosition = float3(PHYSICS_FIELD_NODES_PARAMS[NodeOffset + PLANE_FALLOFF_POSITIONX], 
+									    PHYSICS_FIELD_NODES_PARAMS[NodeOffset + PLANE_FALLOFF_POSITIONY], 
+									    PHYSICS_FIELD_NODES_PARAMS[NodeOffset + PLANE_FALLOFF_POSITIONZ]);
+	
+	const float3 PlaneNormal = float3(PHYSICS_FIELD_NODES_PARAMS[NodeOffset + PLANE_FALLOFF_NORMALX], 
+									  PHYSICS_FIELD_NODES_PARAMS[NodeOffset + PLANE_FALLOFF_NORMALY], 
+									  PHYSICS_FIELD_NODES_PARAMS[NodeOffset + PLANE_FALLOFF_NORMALZ]);
+	
+	const int FalloffType = (int)PHYSICS_FIELD_NODES_PARAMS[NodeOffset + PLANE_FALLOFF_FALLOFF];
+	
+	const float PlaneOffset = dot(PlanePosition, PlaneNormal);
 	const float LocalDistance = dot(PlaneNormal, SamplePosition) - PlaneOffset;
-	const float PlaneDistance = NodesParams[NodeOffset + PLANE_FALLOFF_DISTANCE];
-
+    
+	float ResultScalar = DefaultValue;
 	if (PlaneDistance > 0.0 && LocalDistance > -PlaneDistance && LocalDistance < 0.0)
 	{
-		const float DeltaDistance = 1.0 + LocalDistance / PlaneDistance;
-		
-		SetFalloffValue((int)NodesParams[NodeOffset + PLANE_FALLOFF_FALLOFF],
-						1.0, DeltaDistance, ScalarDatas);
-		ScalarDatas = NodesParams[NodeOffset + PLANE_FALLOFF_MAGNITUDE] * (NodesParams[NodeOffset + PLANE_FALLOFF_MIN_RANGE] + ScalarDatas * DeltaRange);
+		const float InverseDistance = rcp(PlaneDistance);
+		const float FalloffValue = ComputeFalloffValue(FalloffType, 1.0, 1.0 + LocalDistance * InverseDistance);
+		ResultScalar = Magnitude * (MinRange + FalloffValue * (MaxRange - MinRange));
 	}
-	SharedDatas[GFieldGroupThreadId][LocalIndex++] = ScalarDatas;
+    
+	SetFieldValue(EvalContext, EvalContext.DataIndex++, ResultScalar);
 }
 
 /* -----------------------------------------------------------------
@@ -493,48 +506,43 @@ void EvaluatePlaneFalloffScalar(in float3 SamplePosition, in int NodeOffset, ino
 #define BOX_FALLOFF_SCALEY 12
 #define BOX_FALLOFF_SCALEZ 13
 #define BOX_FALLOFF_FALLOFF 14
- 
-void EvaluateBoxFalloffScalar(in float3 SamplePosition, in int NodeOffset, inout int LocalIndex)
+
+void EvaluateBoxFalloffScalar(float3 SamplePosition, int NodeOffset, inout FFieldEvalContext EvalContext, PHYSICS_FIELD_BUFFERS)
 {
-	float ScalarDatas = NodesParams[NodeOffset + BOX_FALLOFF_DEFAULT];
-	const float HalfBox = 50.0;
+	const float Magnitude = PHYSICS_FIELD_NODES_PARAMS[NodeOffset + BOX_FALLOFF_MAGNITUDE];
+	const float MinRange = PHYSICS_FIELD_NODES_PARAMS[NodeOffset + BOX_FALLOFF_MIN_RANGE];
+	const float MaxRange = PHYSICS_FIELD_NODES_PARAMS[NodeOffset + BOX_FALLOFF_MAX_RANGE];
+	const float DefaultValue = PHYSICS_FIELD_NODES_PARAMS[NodeOffset + BOX_FALLOFF_DEFAULT];
 	
-	const float DeltaRange = NodesParams[NodeOffset + BOX_FALLOFF_MAX_RANGE] - NodesParams[NodeOffset + BOX_FALLOFF_MIN_RANGE];
+	const float4 BoxRotation = float4(PHYSICS_FIELD_NODES_PARAMS[NodeOffset + BOX_FALLOFF_ROTATIONX], PHYSICS_FIELD_NODES_PARAMS[NodeOffset + BOX_FALLOFF_ROTATIONY],
+							          PHYSICS_FIELD_NODES_PARAMS[NodeOffset + BOX_FALLOFF_ROTATIONZ], PHYSICS_FIELD_NODES_PARAMS[NodeOffset + BOX_FALLOFF_ROTATIONW]);
 
-	const float4 BoxRotation = float4(NodesParams[NodeOffset + BOX_FALLOFF_ROTATIONX],
-								      NodesParams[NodeOffset + BOX_FALLOFF_ROTATIONY],
-								      NodesParams[NodeOffset + BOX_FALLOFF_ROTATIONZ],
-									  NodesParams[NodeOffset + BOX_FALLOFF_ROTATIONW]);
+	const float3 BoxTranslation = float3(PHYSICS_FIELD_NODES_PARAMS[NodeOffset + BOX_FALLOFF_TRANSLATIONX],
+									     PHYSICS_FIELD_NODES_PARAMS[NodeOffset + BOX_FALLOFF_TRANSLATIONY],
+										 PHYSICS_FIELD_NODES_PARAMS[NodeOffset + BOX_FALLOFF_TRANSLATIONZ]);
 
-	const float3 BoxTranslation = float3(NodesParams[NodeOffset + BOX_FALLOFF_TRANSLATIONX],
-							             NodesParams[NodeOffset + BOX_FALLOFF_TRANSLATIONY],
-								         NodesParams[NodeOffset + BOX_FALLOFF_TRANSLATIONZ]);
-
-	const float3 BoxScale = float3(NodesParams[NodeOffset + BOX_FALLOFF_SCALEX],
-							       NodesParams[NodeOffset + BOX_FALLOFF_SCALEY],
-								   NodesParams[NodeOffset + BOX_FALLOFF_SCALEZ]);
-
-	const float3 InverseScale = float3(BoxScale.x != 0.0 ? 1.0 / BoxScale.x : 0.0,
-										BoxScale.y != 0.0 ? 1.0 / BoxScale.y : 0.0,
-										BoxScale.z != 0.0 ? 1.0 / BoxScale.z : 0.0);
-
-	const float3 LocalPosition = FieldUnRotateVectorByQuat(SamplePosition - BoxTranslation, BoxRotation) * InverseScale;
+	const float3 BoxScale = float3(PHYSICS_FIELD_NODES_PARAMS[NodeOffset + BOX_FALLOFF_SCALEX],
+								   PHYSICS_FIELD_NODES_PARAMS[NodeOffset + BOX_FALLOFF_SCALEY],
+								   PHYSICS_FIELD_NODES_PARAMS[NodeOffset + BOX_FALLOFF_SCALEZ]);
+	
+	const int FalloffType = (int)PHYSICS_FIELD_NODES_PARAMS[NodeOffset + BOX_FALLOFF_FALLOFF];
 	
-	const float3 DeltaPosition = abs(LocalPosition) - float3(HalfBox, HalfBox, HalfBox);
-	const int ClosestAxis = ((DeltaPosition.x > DeltaPosition.y) && (DeltaPosition.x > DeltaPosition.z)) ? 0 : (DeltaPosition.y > DeltaPosition.z) ? 1 : 2;
+	const float3 InverseScale = rcp(max(BoxScale, 0.0001));
+	const float3 LocalPosition = UnRotateVectorByQuat(SamplePosition - BoxTranslation, BoxRotation) * InverseScale;
+	const float3 DeltaPosition = abs(LocalPosition) - 50.0;
+	
+	const float MaxComponent = max(max(DeltaPosition.x, DeltaPosition.y), DeltaPosition.z);
 	const float OutsideDistance = length(max(DeltaPosition, 0.0));
-
-	const float LocalDistance = OutsideDistance + min(DeltaPosition[ClosestAxis], 0.0);
-
+	const float LocalDistance = OutsideDistance + min(MaxComponent, 0.0);
+    
+	float ResultScalar = DefaultValue;
 	if (LocalDistance < 0.0)
 	{
-		const float DeltaDistance = -LocalDistance / HalfBox;
-		
-		SetFalloffValue((int)NodesParams[NodeOffset + BOX_FALLOFF_FALLOFF],
-						1.0,DeltaDistance, ScalarDatas);
-		ScalarDatas = NodesParams[NodeOffset + BOX_FALLOFF_MAGNITUDE] * (NodesParams[NodeOffset + BOX_FALLOFF_MIN_RANGE] + ScalarDatas * DeltaRange);
+		const float FalloffValue = ComputeFalloffValue(FalloffType, 1.0, -LocalDistance * 0.02);
+		ResultScalar = Magnitude * (MinRange + FalloffValue * (MaxRange - MinRange));
 	}
-	SharedDatas[GFieldGroupThreadId][LocalIndex++] = ScalarDatas;
+    
+	SetFieldValue(EvalContext, EvalContext.DataIndex++, ResultScalar);
 }
 
 /* -----------------------------------------------------------------
@@ -554,38 +562,33 @@ void EvaluateBoxFalloffScalar(in float3 SamplePosition, in int NodeOffset, inout
 #define NOISE_SCALEX 9
 #define NOISE_SCALEY 10
 #define NOISE_SCALEZ 11
- 
-void EvaluateNoiseScalar(in float3 SamplePosition, in int NodeOffset, inout int LocalIndex, in float MinScale)
+
+void EvaluateNoiseScalar(float3 SamplePosition, int NodeOffset, float MinScale, inout FFieldEvalContext EvalContext, PHYSICS_FIELD_BUFFERS)
 {
-	const float4 NoiseRotation = float4(NodesParams[NodeOffset + NOISE_ROTATIONX],
-								        NodesParams[NodeOffset + NOISE_ROTATIONY],
-								        NodesParams[NodeOffset + NOISE_ROTATIONZ],
-									    NodesParams[NodeOffset + NOISE_ROTATIONW]);
+	const float MinRange = PHYSICS_FIELD_NODES_PARAMS[NodeOffset + NOISE_MIN_RANGE];
+	const float MaxRange = PHYSICS_FIELD_NODES_PARAMS[NodeOffset + NOISE_MAX_RANGE];
 	
-	const float3 NoiseTranslation = float3(NodesParams[NodeOffset + NOISE_TRANSLATIONX],
-								           NodesParams[NodeOffset + NOISE_TRANSLATIONY],
-								           NodesParams[NodeOffset + NOISE_TRANSLATIONZ]);
-
-	float3 NoiseScale = float3(NodesParams[NodeOffset + NOISE_SCALEX],
-							         NodesParams[NodeOffset + NOISE_SCALEY],
-								     NodesParams[NodeOffset + NOISE_SCALEZ]);
+	const float4 NoiseRotation = float4(PHYSICS_FIELD_NODES_PARAMS[NodeOffset + NOISE_ROTATIONX], PHYSICS_FIELD_NODES_PARAMS[NodeOffset + NOISE_ROTATIONY],
+									    PHYSICS_FIELD_NODES_PARAMS[NodeOffset + NOISE_ROTATIONZ], PHYSICS_FIELD_NODES_PARAMS[NodeOffset + NOISE_ROTATIONW]);
 	
-	NoiseScale = 0.5 * max(NoiseScale, float3(MinScale, MinScale, MinScale));
+	const float3 NoiseTranslation = float3(PHYSICS_FIELD_NODES_PARAMS[NodeOffset + NOISE_TRANSLATIONX], 
+										   PHYSICS_FIELD_NODES_PARAMS[NodeOffset + NOISE_TRANSLATIONY], 
+										   PHYSICS_FIELD_NODES_PARAMS[NodeOffset + NOISE_TRANSLATIONZ]);
 	
-	const float3 InverseScale = float3(NoiseScale.x != 0.0 ? 1.0 / NoiseScale.x : 0.0,
-										NoiseScale.y != 0.0 ? 1.0 / NoiseScale.y : 0.0,
-										NoiseScale.z != 0.0 ? 1.0 / NoiseScale.z : 0.0);
-
-	float3 LocalPosition = fmod(FieldUnRotateVectorByQuat(SamplePosition - NoiseTranslation, NoiseRotation), NoiseScale) * InverseScale;
-	LocalPosition = (clamp(LocalPosition, float3(-1.0,-1.0,-1.0), float3(1.0,1.0,1.0)) * 0.5 + float3(0.5,0.5,0.5)) * 255;
-
-	const float NoiseValue = PerlinNoise(LocalPosition) * 0.5 + 0.5; 
-	SharedDatas[GFieldGroupThreadId][LocalIndex++] = NoiseValue * (NodesParams[NodeOffset + NOISE_MAX_RANGE] -
-								  NodesParams[NodeOffset + NOISE_MIN_RANGE]) 
-				+ NodesParams[NodeOffset + NOISE_MIN_RANGE];
+	const float3 NoiseScale = 0.5 * max(float3(PHYSICS_FIELD_NODES_PARAMS[NodeOffset + NOISE_SCALEX], 
+											   PHYSICS_FIELD_NODES_PARAMS[NodeOffset + NOISE_SCALEY],  
+											   PHYSICS_FIELD_NODES_PARAMS[NodeOffset + NOISE_SCALEZ]), MinScale);
+    
+	const float3 InverseScale = rcp(max(NoiseScale, 0.0001));
+	float3 LocalPosition = fmod(UnRotateVectorByQuat(SamplePosition - NoiseTranslation, NoiseRotation), NoiseScale) * InverseScale;
+	LocalPosition = clamp(LocalPosition, -1.0, 1.0) * 127.5 + 127.5;
+    
+	const float NoiseValue = ComputePerlinNoise(LocalPosition) * 0.5 + 0.5;
+	const float DeltaRange = MaxRange - MinRange;
+    
+	SetFieldValue(EvalContext, EvalContext.DataIndex++, NoiseValue * DeltaRange + MinRange);
 }
 
-
 /* -----------------------------------------------------------------
  * Uniform Vector Field
  * -----------------------------------------------------------------
@@ -595,17 +598,36 @@ void EvaluateNoiseScalar(in float3 SamplePosition, in int NodeOffset, inout int
 #define UNIFORM_VECTOR_DIRECTIONX 1
 #define UNIFORM_VECTOR_DIRECTIONY 2
 #define UNIFORM_VECTOR_DIRECTIONZ 3
- 
-void EvaluateUniformVector(in float3 SamplePosition, in int NodeOffset, inout int LocalIndex)
+
+void EvaluateUniformVector(int NodeOffset, inout FFieldEvalContext EvalContext, PHYSICS_FIELD_BUFFERS)
+{
+	const float Magnitude = PHYSICS_FIELD_NODES_PARAMS[NodeOffset + UNIFORM_VECTOR_MAGNITUDE];
+	const float3 Direction = float3(PHYSICS_FIELD_NODES_PARAMS[NodeOffset + UNIFORM_VECTOR_DIRECTIONX], 
+		PHYSICS_FIELD_NODES_PARAMS[NodeOffset + UNIFORM_VECTOR_DIRECTIONY], PHYSICS_FIELD_NODES_PARAMS[NodeOffset + UNIFORM_VECTOR_DIRECTIONZ]);
+	
+	const float3 ResultVector = Magnitude * Direction;
+    
+	SetFieldValue(EvalContext, EvalContext.DataIndex++, ResultVector.x);
+	SetFieldValue(EvalContext, EvalContext.DataIndex++, ResultVector.y);
+	SetFieldValue(EvalContext, EvalContext.DataIndex++, ResultVector.z);
+}
+
+/* -----------------------------------------------------------------
+ * Utils to set a radial direction
+ * -----------------------------------------------------------------
+ */
+
+void SetRadialDirection(float3 RadialDirection, float Magnitude, inout FFieldEvalContext EvalContext)
 {
-	const float3 VectorDatas = NodesParams[NodeOffset + UNIFORM_VECTOR_MAGNITUDE] * float3(
-										NodesParams[NodeOffset + UNIFORM_VECTOR_DIRECTIONX],
-								        NodesParams[NodeOffset + UNIFORM_VECTOR_DIRECTIONY],
-								        NodesParams[NodeOffset + UNIFORM_VECTOR_DIRECTIONZ]);
-
-	SharedDatas[GFieldGroupThreadId][LocalIndex++] = VectorDatas.x;
-	SharedDatas[GFieldGroupThreadId][LocalIndex++] = VectorDatas.y;
-	SharedDatas[GFieldGroupThreadId][LocalIndex++] = VectorDatas.z;
+	const float LengthSquared = dot(RadialDirection, RadialDirection);
+    
+	// Use rsqrt for normalization
+	const float InverseLength = LengthSquared > 0.0 ? rsqrt(LengthSquared) : 0.0;
+	const float3 ResultVector = Magnitude * RadialDirection * InverseLength;
+    
+	SetFieldValue(EvalContext, EvalContext.DataIndex++, ResultVector.x);
+	SetFieldValue(EvalContext, EvalContext.DataIndex++, ResultVector.y);
+	SetFieldValue(EvalContext, EvalContext.DataIndex++, ResultVector.z);
 }
 
 /* -----------------------------------------------------------------
@@ -617,20 +639,15 @@ void EvaluateUniformVector(in float3 SamplePosition, in int NodeOffset, inout in
 #define RADIAL_VECTOR_POSITIONX 1
 #define RADIAL_VECTOR_POSITIONY 2
 #define RADIAL_VECTOR_POSITIONZ 3
- 
-void EvaluateRadialVector(in float3 SamplePosition, in int NodeOffset, inout int LocalIndex)
+
+void EvaluateRadialVector(float3 SamplePosition, int NodeOffset, inout FFieldEvalContext EvalContext, PHYSICS_FIELD_BUFFERS)
 {
-	const float3 RadialDirection = SamplePosition - float3(
-										NodesParams[NodeOffset + RADIAL_VECTOR_POSITIONX],
-								        NodesParams[NodeOffset + RADIAL_VECTOR_POSITIONY],
-								        NodesParams[NodeOffset + RADIAL_VECTOR_POSITIONZ]);
-	const float DirectionLength = length(RadialDirection);
-	const float3 VectorDatas = (DirectionLength != 0.0) ? NodesParams[NodeOffset + RADIAL_VECTOR_MAGNITUDE] *
-					RadialDirection / DirectionLength : float3(0, 0, 0);
-
-	SharedDatas[GFieldGroupThreadId][LocalIndex++] = VectorDatas.x;
-	SharedDatas[GFieldGroupThreadId][LocalIndex++] = VectorDatas.y;
-	SharedDatas[GFieldGroupThreadId][LocalIndex++] = VectorDatas.z;
+	const float Magnitude = PHYSICS_FIELD_NODES_PARAMS[NodeOffset + RADIAL_VECTOR_MAGNITUDE];
+	const float3 Center = float3(PHYSICS_FIELD_NODES_PARAMS[NodeOffset + RADIAL_VECTOR_POSITIONX], 
+		PHYSICS_FIELD_NODES_PARAMS[NodeOffset + RADIAL_VECTOR_POSITIONY], PHYSICS_FIELD_NODES_PARAMS[NodeOffset + RADIAL_VECTOR_POSITIONZ]);
+    
+	const float3 RadialDirection = SamplePosition - Center;
+	SetRadialDirection(RadialDirection, Magnitude, EvalContext);
 }
 
 /* -----------------------------------------------------------------
@@ -639,18 +656,13 @@ void EvaluateRadialVector(in float3 SamplePosition, in int NodeOffset, inout int
  */
 
 #define RANDOM_VECTOR_MAGNITUDE 0
- 
-void EvaluateRandomVector(in float3 SamplePosition, in int NodeOffset, inout int LocalIndex)
-{
-	const float3 RadialDirection = float3(RandomVector(int3(SamplePosition))) / 0xffff - 0.5;//float3(Rand3DPCG16(int3(SamplePosition))) / 0xffff - 0.5;
-
-	const float DirectionLength = length(RadialDirection);
-	const float3 VectorDatas = (DirectionLength != 0.0) ? NodesParams[NodeOffset + RANDOM_VECTOR_MAGNITUDE] *
-					RadialDirection / DirectionLength : float3(0, 0, 0);
 
-	SharedDatas[GFieldGroupThreadId][LocalIndex++] = VectorDatas.x;
-	SharedDatas[GFieldGroupThreadId][LocalIndex++] = VectorDatas.y;
-	SharedDatas[GFieldGroupThreadId][LocalIndex++] = VectorDatas.z;
+void EvaluateRandomVector(float3 SamplePosition, int NodeOffset, inout FFieldEvalContext EvalContext, PHYSICS_FIELD_BUFFERS)
+{
+	const float Magnitude = PHYSICS_FIELD_NODES_PARAMS[NodeOffset+RANDOM_VECTOR_MAGNITUDE];
+	
+	const float3 RadialDirection = float3(ComputeRandomVector(int3(SamplePosition))) * rcp(65535.0) - 0.5;
+	SetRadialDirection(RadialDirection, Magnitude, EvalContext);
 }
 
 /* -----------------------------------------------------------------
@@ -663,39 +675,37 @@ void EvaluateRandomVector(in float3 SamplePosition, in int NodeOffset, inout int
 #define SUM_SCALAR_LEFT 2
 #define SUM_SCALAR_OPERATION 3
 
-void EvaluateSumScalar(in float3 SamplePosition, in int NodeOffset, inout int LocalIndex)
+void EvaluateSumScalar(int NodeOffset, inout FFieldEvalContext EvalContext, PHYSICS_FIELD_BUFFERS)
 {
-	const int HasRight = (int)NodesParams[NodeOffset + SUM_SCALAR_RIGHT];
-	const int HasLeft = (int)NodesParams[NodeOffset + SUM_SCALAR_LEFT];
-
-	LocalIndex -= HasRight + HasLeft;
-
-	const int FieldOperation = (int)NodesParams[NodeOffset + SUM_SCALAR_OPERATION];
-	const bool UnitOperation = (FieldOperation == MULTIPLY_OP) || (FieldOperation == DIVIDE_OP);
-
-	int OffsetIndex = LocalIndex;
-	const float ScalarRight = HasRight ? SharedDatas[GFieldGroupThreadId][OffsetIndex] : UnitOperation ? 1.0 : 0.0;
-	OffsetIndex += HasRight;
-	const float ScalarLeft = HasLeft ? SharedDatas[GFieldGroupThreadId][OffsetIndex] : UnitOperation ? 1.0 : 0.0;
-
-	float ScalarDatas = 0.0;
-	if (FieldOperation == MULTIPLY_OP)
-	{
-		ScalarDatas = ScalarRight * ScalarLeft;
-	}
-	else if (FieldOperation == DIVIDE_OP)
-	{
-		ScalarDatas = ScalarLeft / ScalarRight;
-	}
-	else if (FieldOperation == ADD_OP)
-	{
-		ScalarDatas = ScalarRight + ScalarLeft;
-	}
-	else if (FieldOperation == SUBTRACT_OP)
-	{
-		ScalarDatas = ScalarLeft - ScalarRight;
-	}
-	SharedDatas[GFieldGroupThreadId][LocalIndex++] = ScalarDatas * NodesParams[NodeOffset + SUM_SCALAR_MAGNITUDE];
+	const float Magnitude = PHYSICS_FIELD_NODES_PARAMS[NodeOffset + SUM_SCALAR_MAGNITUDE];
+	const int HasRight = (int)PHYSICS_FIELD_NODES_PARAMS[NodeOffset + SUM_SCALAR_RIGHT];
+	const int HasLeft = (int)PHYSICS_FIELD_NODES_PARAMS[NodeOffset + SUM_SCALAR_LEFT];
+	const int Operation = (int)PHYSICS_FIELD_NODES_PARAMS[NodeOffset + SUM_SCALAR_OPERATION];
+    
+	EvalContext.DataIndex -= HasRight + HasLeft;
+	
+	const float DefaultValue = (Operation == MULTIPLY_OP || Operation == DIVIDE_OP) ? 1.0 : 0.0;
+	const float RightValue = HasRight ? GetFieldValue(EvalContext, EvalContext.DataIndex) : DefaultValue;
+	const float LeftValue = HasLeft ? GetFieldValue(EvalContext, EvalContext.DataIndex + HasRight) : DefaultValue;
+	
+	float ResultValue = 0.0;
+	switch (Operation)
+	{
+		case MULTIPLY_OP: 
+			ResultValue = RightValue * LeftValue;
+			break;
+		case DIVIDE_OP:
+			ResultValue = LeftValue * rcp(RightValue);
+			break;
+		case ADD_OP: 
+			ResultValue = RightValue + LeftValue;
+			break;
+		case SUBTRACT_OP:
+			ResultValue = LeftValue - RightValue;
+			break;
+		default : break;
+	}
+	SetFieldValue(EvalContext, EvalContext.DataIndex++, ResultValue * Magnitude);
 }
 
 /* -----------------------------------------------------------------
@@ -709,72 +719,48 @@ void EvaluateSumScalar(in float3 SamplePosition, in int NodeOffset, inout int Lo
 #define SUM_VECTOR_LEFT 3
 #define SUM_VECTOR_OPERATION 4
 
-void EvaluateSumVector(in float3 SamplePosition, in int NodeOffset, inout int LocalIndex)
+void EvaluateSumVector(int NodeOffset, inout FFieldEvalContext EvalContext, PHYSICS_FIELD_BUFFERS)
 {
-	const int HasScalar = (int)NodesParams[NodeOffset + SUM_VECTOR_SCALAR];
-	const int HasRight = (int)NodesParams[NodeOffset + SUM_VECTOR_RIGHT];
-	const int HasLeft = (int)NodesParams[NodeOffset + SUM_VECTOR_LEFT];
-
-	LocalIndex -= HasScalar + HasRight * 3 + HasLeft * 3;
-
-	const int FieldOperation = (int)NodesParams[NodeOffset + SUM_VECTOR_OPERATION];
-	const bool UnitOperation = (FieldOperation == MULTIPLY_OP) || (FieldOperation == DIVIDE_OP);
-
-	int OffsetIndex = LocalIndex;
-	const float ScalarMagnitude = HasScalar ? SharedDatas[GFieldGroupThreadId][OffsetIndex] : 1.0;
-	OffsetIndex += HasScalar;
-	const float3 VectorRight = HasRight ? float3(SharedDatas[GFieldGroupThreadId][OffsetIndex], SharedDatas[GFieldGroupThreadId][OffsetIndex + 1], SharedDatas[GFieldGroupThreadId][OffsetIndex + 2]) : UnitOperation ? float3(1, 1, 1) : float3(0, 0, 0);
-	OffsetIndex += HasRight * 3;
-	const float3 VectorLeft = HasLeft ? float3(SharedDatas[GFieldGroupThreadId][OffsetIndex], SharedDatas[GFieldGroupThreadId][OffsetIndex + 1], SharedDatas[GFieldGroupThreadId][OffsetIndex + 2]) : UnitOperation ? float3(1, 1, 1) : float3(0, 0, 0);
-
-	float3 VectorDatas = float3(0.0, 0.0, 0.0);
-	if (FieldOperation == MULTIPLY_OP)
-	{
-		VectorDatas = VectorRight * VectorLeft;
-	}
-	else if (FieldOperation == DIVIDE_OP)
-	{
-		VectorDatas = VectorLeft / VectorRight;
-	}
-	else if (FieldOperation == ADD_OP)
-	{
-		VectorDatas = VectorRight + VectorLeft;
-	}
-	else if (FieldOperation == SUBTRACT_OP)
-	{
-		VectorDatas = VectorLeft - VectorRight;
-	}
-
-	VectorDatas *= ScalarMagnitude * NodesParams[NodeOffset + SUM_VECTOR_MAGNITUDE];
-
-	SharedDatas[GFieldGroupThreadId][LocalIndex++] = VectorDatas.x;
-	SharedDatas[GFieldGroupThreadId][LocalIndex++] = VectorDatas.y;
-	SharedDatas[GFieldGroupThreadId][LocalIndex++] = VectorDatas.z;
-}
-
-/* -----------------------------------------------------------------
- * Conversion Scalar
- * -----------------------------------------------------------------
- */
-#define CONVERSION_INPUT 0
-
-void EvaluateConversionScalar(in float3 SamplePosition, in int NodeOffset, inout int LocalIndex)
-{
-	// conversion in the same LocalIndex
-}
-
-/* -----------------------------------------------------------------
- * Conversion Integer
- * -----------------------------------------------------------------
- */
-
-void EvaluateConversionInteger(in float3 SamplePosition, in int NodeOffset, inout int LocalIndex)
-{
-	// conversion in the same LocalIndex
+    const float Magnitude = PHYSICS_FIELD_NODES_PARAMS[NodeOffset + SUM_VECTOR_MAGNITUDE];
+    const int HasScalar = (int)PHYSICS_FIELD_NODES_PARAMS[NodeOffset + SUM_VECTOR_SCALAR];
+    const int HasRight = (int)PHYSICS_FIELD_NODES_PARAMS[NodeOffset + SUM_VECTOR_RIGHT];
+    const int HasLeft = (int)PHYSICS_FIELD_NODES_PARAMS[NodeOffset + SUM_VECTOR_LEFT];
+    const int Operation = (int)PHYSICS_FIELD_NODES_PARAMS[NodeOffset + SUM_VECTOR_OPERATION];
+    
+    EvalContext.DataIndex -= HasScalar + HasRight * 3 + HasLeft * 3;
+	
+    const float3 DefaultVec = (Operation == MULTIPLY_OP || Operation == DIVIDE_OP) ? 1.0 : 0.0;
+    const float ScalarValue = HasScalar ? GetFieldValue(EvalContext, EvalContext.DataIndex) : 1.0;
+	
+    const float3 RightVector = HasRight ? GetFieldVector(EvalContext, EvalContext.DataIndex + HasScalar) : DefaultVec;
+    const float3 LeftVector = HasLeft ? GetFieldVector(EvalContext, EvalContext.DataIndex + HasScalar + HasRight * 3) : DefaultVec;
+	
+	float3 ResultVector = 0.0;
+	switch (Operation)
+	{
+		case MULTIPLY_OP: 
+			ResultVector = RightVector * LeftVector;
+			break;
+		case DIVIDE_OP:
+			ResultVector = LeftVector * rcp(RightVector);
+			break;
+		case ADD_OP: 
+			ResultVector = RightVector + LeftVector;
+			break;
+		case SUBTRACT_OP:
+			ResultVector = LeftVector - RightVector;
+			break;
+		default : break;
+	}
+    
+    ResultVector *= ScalarValue * Magnitude;
+    SetFieldValue(EvalContext, EvalContext.DataIndex++, ResultVector.x);
+    SetFieldValue(EvalContext, EvalContext.DataIndex++, ResultVector.y);
+    SetFieldValue(EvalContext, EvalContext.DataIndex++, ResultVector.z);
 }
 
 /* -----------------------------------------------------------------
- * Culling Integer
+ * Culling Scalar
  * -----------------------------------------------------------------
  */
 
@@ -782,255 +768,199 @@ void EvaluateConversionInteger(in float3 SamplePosition, in int NodeOffset, inou
 #define CULLING_INPUT 1
 #define CULLING_OPERATION 2
 
-void EvaluateCullingInteger(in float3 SamplePosition, in int NodeOffset, inout int LocalIndex)
-{
-	const int HasScalar = (int)NodesParams[NodeOffset + CULLING_SCALAR];
-	const int HasInput = (int)NodesParams[NodeOffset + CULLING_INPUT];
-
-	LocalIndex -= HasScalar + HasInput;
-
-	int OffsetIndex = LocalIndex;
-	const float CullingScalar = HasScalar ? SharedDatas[GFieldGroupThreadId][OffsetIndex] : 0.0;
-	OffsetIndex += (int)CullingScalar;
-	const int InputInteger = HasInput ? (int) SharedDatas[GFieldGroupThreadId][OffsetIndex] : 0;
-
-	const int CullingOperation = (int)NodesParams[NodeOffset + CULLING_OPERATION];
-	const bool ValidSample = ((CullingOperation == CULLING_OUTSIDE) && (CullingScalar != 0.0)) ||
-							 ((CullingOperation == CULLING_INSIDE) && (CullingScalar == 0.0));
-
-	SharedDatas[GFieldGroupThreadId][LocalIndex++] = (float) (ValidSample ? InputInteger : 0);
+void EvaluateCullingScalar(int NodeOffset, inout FFieldEvalContext EvalContext, PHYSICS_FIELD_BUFFERS)
+{
+	const int HasScalar = (int)PHYSICS_FIELD_NODES_PARAMS[NodeOffset + CULLING_SCALAR];
+	const int HasInput = (int)PHYSICS_FIELD_NODES_PARAMS[NodeOffset + CULLING_INPUT];
+	const int Operation = (int)PHYSICS_FIELD_NODES_PARAMS[NodeOffset + CULLING_OPERATION];
+    
+	EvalContext.DataIndex -= HasScalar + HasInput;
+    
+	const float CullingScalar = HasScalar ? GetFieldValue(EvalContext, EvalContext.DataIndex) : 0.0;
+	const bool ValidCulling = (Operation == CULLING_OUTSIDE && CullingScalar != 0.0) || 
+							  (Operation == CULLING_INSIDE && CullingScalar == 0.0);
+	
+	const float InputScalar = (HasInput && ValidCulling) ? GetFieldValue(EvalContext, EvalContext.DataIndex + HasScalar) : 0.0;
+	SetFieldValue(EvalContext, EvalContext.DataIndex++, InputScalar);
 }
 
 /* -----------------------------------------------------------------
- * Culling Scalar
+ * Culling Vector
  * -----------------------------------------------------------------
  */
 
-void EvaluateCullingScalar(in float3 SamplePosition, in int NodeOffset, inout int LocalIndex)
-{
-	const int HasScalar = (int)NodesParams[NodeOffset + CULLING_SCALAR];
-	const int HasInput = (int) NodesParams[NodeOffset + CULLING_INPUT];
-
-	LocalIndex -= HasScalar + HasInput;
-
-	int OffsetIndex = LocalIndex;
-	const float CullingScalar = HasScalar ? SharedDatas[GFieldGroupThreadId][OffsetIndex] : 0.0;
-	OffsetIndex += (int)CullingScalar;
-	const float InputScalar = HasInput ? SharedDatas[GFieldGroupThreadId][OffsetIndex] : 0.0;
-
-	const int CullingOperation = (int)NodesParams[NodeOffset + CULLING_OPERATION];
-	const bool ValidSample = ((CullingOperation == CULLING_OUTSIDE) && (CullingScalar != 0.0)) ||
-							 ((CullingOperation == CULLING_INSIDE) && (CullingScalar == 0.0));
-
-	SharedDatas[GFieldGroupThreadId][LocalIndex++] = ValidSample ? InputScalar : 0.0;
+void EvaluateCullingVector(int NodeOffset, inout FFieldEvalContext EvalContext, PHYSICS_FIELD_BUFFERS)
+{
+	const int HasScalar = (int)PHYSICS_FIELD_NODES_PARAMS[NodeOffset + CULLING_SCALAR];
+	const int HasInput = (int)PHYSICS_FIELD_NODES_PARAMS[NodeOffset + CULLING_INPUT];
+	const int Operation = (int)PHYSICS_FIELD_NODES_PARAMS[NodeOffset + CULLING_OPERATION];
+    
+	EvalContext.DataIndex -= HasScalar + HasInput * 3;
+    
+	const float CullingScalar = HasScalar ? GetFieldValue(EvalContext, EvalContext.DataIndex) : 0.0;
+	const bool ValidCulling = (Operation == CULLING_OUTSIDE && CullingScalar != 0.0) || 
+							  (Operation == CULLING_INSIDE && CullingScalar == 0.0);
+	
+	const float3 InputVector = (HasInput && ValidCulling) ? GetFieldVector(EvalContext, EvalContext.DataIndex + HasScalar) : 0.0;
+	SetFieldValue(EvalContext, EvalContext.DataIndex++, InputVector.x);
+	SetFieldValue(EvalContext, EvalContext.DataIndex++, InputVector.y);
+	SetFieldValue(EvalContext, EvalContext.DataIndex++, InputVector.z);
 }
 
 /* -----------------------------------------------------------------
- * Culling Vector
+ * Culling Integer
  * -----------------------------------------------------------------
  */
 
-void EvaluateCullingVector(in float3 SamplePosition, in int NodeOffset, inout int LocalIndex)
+void EvaluateCullingInteger(int NodeOffset, inout FFieldEvalContext EvalContext, PHYSICS_FIELD_BUFFERS)
 {
-	const int HasScalar = (int)NodesParams[NodeOffset + CULLING_SCALAR];
-	const int HasInput = (int)NodesParams[NodeOffset + CULLING_INPUT];
-
-	LocalIndex -= HasScalar + HasInput * 3;
-
-	int OffsetIndex = LocalIndex;
-	const float CullingScalar = HasScalar ? SharedDatas[GFieldGroupThreadId][OffsetIndex] : 0.0;
-	OffsetIndex += (int)CullingScalar;
-	const float3 InputVector = HasInput ? float3(SharedDatas[GFieldGroupThreadId][OffsetIndex], SharedDatas[GFieldGroupThreadId][OffsetIndex + 1], SharedDatas[GFieldGroupThreadId][OffsetIndex + 2]) : float3(0, 0, 0);
-
-	const int CullingOperation = (int)NodesParams[NodeOffset + CULLING_OPERATION];
-	const bool ValidSample = ((CullingOperation == CULLING_OUTSIDE) && (CullingScalar != 0.0)) ||
-							 ((CullingOperation == CULLING_INSIDE) && (CullingScalar == 0.0));
-
-	SharedDatas[GFieldGroupThreadId][LocalIndex++] = ValidSample ? InputVector.x : 0.0;
-	SharedDatas[GFieldGroupThreadId][LocalIndex++] = ValidSample ? InputVector.y : 0.0;
-	SharedDatas[GFieldGroupThreadId][LocalIndex++] = ValidSample ? InputVector.z : 0.0;
+	EvaluateCullingScalar(NodeOffset, EvalContext, PHYSICS_FIELD_ARGS);
 }
 
 /* -----------------------------------------------------------------
- * Nodes evaluation
+ * Main Evaluation Dispatchers - Using switch for better performance
  * -----------------------------------------------------------------
  */
 
-void EvaluateFieldNodeInteger(in float3 SamplePosition, in int NodeOffset, inout int LocalIndex)
-{
-	const int NodeType = (int)NodesParams[NodeOffset];
-
-	[branch]
-	if (NodeType == UNIFORM_INTEGER)
-	{
-		EvaluateUniformInteger(SamplePosition, NodeOffset + 1, LocalIndex);
-	}
-	else if (NodeType == RADIAL_MASK_INTEGER)
-	{
-		EvaluateRadialMaskInteger(SamplePosition, NodeOffset + 1, LocalIndex);
-	}
-	else if (NodeType == CONVERSION_FIELD)
-	{
-		EvaluateConversionInteger(SamplePosition, NodeOffset + 1, LocalIndex);
-	}
-	else if (NodeType == CULLING_FIELD)
-	{
-		EvaluateCullingInteger(SamplePosition, NodeOffset + 1, LocalIndex);
+void EvaluateFieldNodeScalar(float3 SamplePosition, int NodeType, int NodeOffset, float MinScale, float MinLength, inout FFieldEvalContext EvalContext, PHYSICS_FIELD_BUFFERS)
+{
+    switch(NodeType)
+    {
+        case UNIFORM_SCALAR:
+			EvaluateUniformScalar(SamplePosition, NodeOffset, EvalContext, PHYSICS_FIELD_ARGS);
+            break;
+        case WAVE_SCALAR:
+			EvaluateWaveScalar(SamplePosition, NodeOffset, MinLength, EvalContext, PHYSICS_FIELD_ARGS);
+            break;
+        case RADIAL_FALLOFF_SCALAR:
+			EvaluateRadialFalloffScalar(SamplePosition, NodeOffset, EvalContext, PHYSICS_FIELD_ARGS);
+            break;
+        case PLANE_FALLOFF_SCALAR:
+			EvaluatePlaneFalloffScalar(SamplePosition, NodeOffset, EvalContext, PHYSICS_FIELD_ARGS);
+            break;
+        case BOX_FALLOFF_SCALAR:
+			EvaluateBoxFalloffScalar(SamplePosition, NodeOffset, EvalContext, PHYSICS_FIELD_ARGS);
+            break;
+        case NOISE_SCALAR:
+			EvaluateNoiseScalar(SamplePosition, NodeOffset, MinScale, EvalContext, PHYSICS_FIELD_ARGS);
+            break;
+        case SUM_SCALAR:
+			EvaluateSumScalar(NodeOffset, EvalContext, PHYSICS_FIELD_ARGS);
+            break;
+        case CONVERSION_FIELD:
+            // No-op: conversion happens in-place
+            break;
+        case CULLING_FIELD:
+			EvaluateCullingScalar(NodeOffset, EvalContext, PHYSICS_FIELD_ARGS);
+            break;
+		default : break;
+    }
+}
+
+void EvaluateFieldNodeInteger(float3 SamplePosition, int NodeType, int NodeOffset, inout FFieldEvalContext EvalContext, PHYSICS_FIELD_BUFFERS)
+{
+    switch(NodeType)
+    {
+        case UNIFORM_INTEGER:
+			EvaluateUniformInteger(SamplePosition, NodeOffset, EvalContext, PHYSICS_FIELD_ARGS);
+            break;
+        case RADIAL_MASK_INTEGER:
+			EvaluateRadialMaskInteger(SamplePosition, NodeOffset, EvalContext, PHYSICS_FIELD_ARGS);
+            break;
+        case CONVERSION_FIELD:
+            // No-op: conversion happens in-place
+            break;
+        case CULLING_FIELD:
+			EvaluateCullingInteger(NodeOffset, EvalContext, PHYSICS_FIELD_ARGS);
+            break;
+        default : break;
+    }
+}
+
+void EvaluateFieldNodeVector(float3 SamplePosition, int NodeType, int NodeOffset, inout FFieldEvalContext EvalContext, PHYSICS_FIELD_BUFFERS)
+{
+    switch(NodeType)
+    {
+        case UNIFORM_VECTOR:
+			EvaluateUniformVector(NodeOffset, EvalContext, PHYSICS_FIELD_ARGS);
+            break;
+        case RADIAL_VECTOR:
+			EvaluateRadialVector(SamplePosition, NodeOffset, EvalContext, PHYSICS_FIELD_ARGS);
+            break;
+        case RANDOM_VECTOR:
+			EvaluateRandomVector(SamplePosition, NodeOffset, EvalContext, PHYSICS_FIELD_ARGS);
+            break;
+        case SUM_VECTOR:
+			EvaluateSumVector(NodeOffset, EvalContext, PHYSICS_FIELD_ARGS);
+            break;
+        case CULLING_FIELD:
+			EvaluateCullingVector(NodeOffset, EvalContext, PHYSICS_FIELD_ARGS);
+            break;
+    	default : break;
+    }
+}
+
+void EvaluateFieldNode(float3 SamplePosition, int NodeType, int NodeOffset, int DataType, float MinScale, float MinLength, inout FFieldEvalContext EvalContext, PHYSICS_FIELD_BUFFERS)
+{
+    switch(DataType)
+    {
+        case SCALAR_TYPE:
+			EvaluateFieldNodeScalar(SamplePosition, NodeType, NodeOffset, MinScale, MinLength, EvalContext, PHYSICS_FIELD_ARGS);
+            break;
+        case INTEGER_TYPE:
+			EvaluateFieldNodeInteger(SamplePosition, NodeType, NodeOffset, EvalContext, PHYSICS_FIELD_ARGS);
+            break;
+        case VECTOR_TYPE:
+			EvaluateFieldNodeVector(SamplePosition, NodeType, NodeOffset, EvalContext, PHYSICS_FIELD_ARGS);
+            break;
+		default : break;
+    }
+}
+
+void SampleFieldData(float3 SamplePosition, int TargetType, float MinScale, float MinLength, float TimeSeconds,
+	PHYSICS_FIELD_BUFFERS, inout FFieldEvalContext EvalContext)
+{
+	InitEvalContext(EvalContext, TimeSeconds);
+    
+    const int NodesBegin = PHYSICS_FIELD_TARGETS_OFFSETS[TargetType];
+    const int NodesEnd = PHYSICS_FIELD_TARGETS_OFFSETS[TargetType + 1];
+    
+    // Loop through nodes
+    for (int NodeIndex = NodesBegin; NodeIndex < NodesEnd; ++NodeIndex)
+    {
+        const int NodeOffset = PHYSICS_FIELD_NODES_OFFSETS[NodeIndex];
+        const int DataType = (int)PHYSICS_FIELD_NODES_PARAMS[NodeOffset];
+        const int NodeType = (int)PHYSICS_FIELD_NODES_PARAMS[NodeOffset + 1];
+    	
+		EvaluateFieldNode(SamplePosition, NodeType, NodeOffset + 2, DataType, MinScale, MinLength, EvalContext, PHYSICS_FIELD_ARGS);
 	}
 }
 
-void EvaluateFieldNodeScalar(in float3 SamplePosition, in int NodeOffset, inout int LocalIndex, in float MinScale, in float MinLength)
-{
-	const int NodeType = (int)NodesParams[NodeOffset];
-
-	[branch]
-	if (NodeType == UNIFORM_SCALAR)
-	{
-		EvaluateUniformScalar(SamplePosition, NodeOffset + 1, LocalIndex);
-	}
-	else if (NodeType == WAVE_SCALAR)
-	{
-		EvaluateWaveScalar(SamplePosition, NodeOffset + 1, LocalIndex, MinLength);
-	}
-	else if (NodeType == RADIAL_FALLOFF_SCALAR)
-	{
-		EvaluateRadialFalloffScalar(SamplePosition, NodeOffset + 1, LocalIndex);
-	}
-	else if (NodeType == PLANE_FALLOFF_SCALAR)
-	{
-		EvaluatePlaneFalloffScalar(SamplePosition, NodeOffset + 1, LocalIndex);
-	}
-	else if (NodeType == BOX_FALLOFF_SCALAR)
-	{
-		EvaluateBoxFalloffScalar(SamplePosition, NodeOffset + 1, LocalIndex);
-	}
-	else if (NodeType == NOISE_SCALAR)
-	{
-		EvaluateNoiseScalar(SamplePosition, NodeOffset + 1, LocalIndex, MinScale);
-	}
-	else if (NodeType == SUM_SCALAR)
-	{
-		EvaluateSumScalar(SamplePosition, NodeOffset + 1, LocalIndex);
-	}
-	else if (NodeType == CONVERSION_FIELD)
-	{
-		EvaluateConversionScalar(SamplePosition, NodeOffset + 1, LocalIndex);
-	}
-	else if (NodeType == CULLING_FIELD)
-	{
-		EvaluateCullingScalar(SamplePosition, NodeOffset + 1, LocalIndex);
-	}
-}
-
-void EvaluateFieldNodeVector(in float3 SamplePosition, in int NodeOffset, inout int LocalIndex)
-{
-	const int NodeType = (int)NodesParams[NodeOffset];
-
-	[branch]
-	if (NodeType == UNIFORM_VECTOR)
-	{
-		EvaluateUniformVector(SamplePosition, NodeOffset + 1, LocalIndex);
-	}
-	else if (NodeType == RADIAL_VECTOR)
-	{
-		EvaluateRadialVector(SamplePosition, NodeOffset + 1, LocalIndex);
-	}
-	else if (NodeType == RANDOM_VECTOR)
-	{
-		EvaluateRandomVector(SamplePosition, NodeOffset + 1, LocalIndex);
-	}
-	else if (NodeType == SUM_VECTOR)
-	{
-		EvaluateSumVector(SamplePosition, NodeOffset + 1, LocalIndex);
-	}
-	else if (NodeType == CULLING_FIELD)
-	{
-		EvaluateCullingVector(SamplePosition, NodeOffset + 1, LocalIndex);
-	}
-}
-
-void EvaluateFieldNodeDatas(in float3 SamplePosition, in int NodeOffset, inout int LocalIndex, in float MinScale, in float MinLength)
-{
-	const int DatasType = (int)NodesParams[NodeOffset];
-
-	[branch]
-	if (DatasType == SCALAR_TYPE)
-	{
-		EvaluateFieldNodeScalar(SamplePosition, NodeOffset + 1, LocalIndex, MinScale, MinLength);
-	}
-	else if (DatasType == INTEGER_TYPE)
-	{
-		EvaluateFieldNodeInteger(SamplePosition, NodeOffset + 1, LocalIndex);
-	}
-	else if (DatasType == VECTOR_TYPE)
-	{
-		EvaluateFieldNodeVector(SamplePosition, NodeOffset + 1,  LocalIndex);
-	}
-}
-
-void SampleFieldDatas(in float3 SamplePosition, inout int LocalIndex, in int TargetType, in float MinScale, in float MinLength)
-{
-	const int NodesBegin = TargetsOffsets[TargetType];
-	const int NodesEnd = TargetsOffsets[TargetType + 1];
-	const int NumNodes = NodesBegin - NodesEnd;
-
-	[unroll]
-	for (int DatasIndex = 0; DatasIndex < MAX_DATAS; ++DatasIndex)
-	{
-		SharedDatas[GFieldGroupThreadId][DatasIndex] = 0.0;
-	}
-	[branch]
-	if (NumNodes != 0)
-	{
-		for (int NodeIndex = NodesBegin; NodeIndex < NodesEnd; ++NodeIndex)
-		{
-			const int NodeOffset = NodesOffsets[NodeIndex];
-			EvaluateFieldNodeDatas(SamplePosition, NodeOffset, LocalIndex, MinScale, MinLength);
-		}
-	}
-}
+/* -----------------------------------------------------------------
+ * Public API
+ * -----------------------------------------------------------------
+ */
 
-float3 PhysicsField_EvalPhysicsVectorField(in float3 SamplePosition, in int TargetIndex)
+float3 PhysicsField_EvalPhysicsVectorField(float3 SamplePosition, int TargetIndex, PHYSICS_FIELD_BUFFERS, in float TimeSeconds)
 {
-	float3 VectorField = float3(0,0,0);
-
-	int LocalIndex = 0;
-	SampleFieldDatas(SamplePosition, LocalIndex, TargetIndex, 0.0, 0.0);
-
-	if(LocalIndex == 3)
-	{
-		VectorField.x = SharedDatas[GFieldGroupThreadId][0];
-		VectorField.y = SharedDatas[GFieldGroupThreadId][1];
-		VectorField.z = SharedDatas[GFieldGroupThreadId][2];
-	}
-	return VectorField;
+    FFieldEvalContext EvalContext;
+	SampleFieldData(SamplePosition, TargetIndex, 0.0, 0.0, TimeSeconds, PHYSICS_FIELD_ARGS, EvalContext);
+    
+    return EvalContext.DataIndex == 3 ? float3(GetFieldValue(EvalContext, 0), GetFieldValue(EvalContext, 1), GetFieldValue(EvalContext, 2)) : 0.0;
 }
 
-float PhysicsField_EvalPhysicsScalarField(in float3 SamplePosition, in int TargetIndex)
+float PhysicsField_EvalPhysicsScalarField(float3 SamplePosition, int TargetIndex, PHYSICS_FIELD_BUFFERS, in float TimeSeconds)
 {
-	float ScalarField = 0.0;
-
-	int LocalIndex = 0;
-	SampleFieldDatas(SamplePosition, LocalIndex, TargetIndex, 0.0, 0.0);
-
-	if(LocalIndex == 1)
-	{
-		ScalarField = SharedDatas[GFieldGroupThreadId][0];
-	}
-	return ScalarField;
+    FFieldEvalContext EvalContext;
+	SampleFieldData(SamplePosition, TargetIndex, 0.0, 0.0, TimeSeconds, PHYSICS_FIELD_ARGS, EvalContext);
+    
+    return EvalContext.DataIndex == 1 ? GetFieldValue(EvalContext, 0) : 0.0;
 }
 
-int PhysicsField_EvalPhysicsIntegerField(in float3 SamplePosition, in int TargetIndex)
+int PhysicsField_EvalPhysicsIntegerField(float3 SamplePosition, int TargetIndex, PHYSICS_FIELD_BUFFERS, in float TimeSeconds)
 {
-	int IntegerField = 0;
-
-	int LocalIndex = 0;
-	SampleFieldDatas(SamplePosition, LocalIndex, TargetIndex, 0.0, 0.0);
-
-	if(LocalIndex == 1)
-	{
-		IntegerField = SharedDatas[GFieldGroupThreadId][0];
-	}
-	return IntegerField;
-}
-
+    FFieldEvalContext EvalContext;
+	SampleFieldData(SamplePosition, TargetIndex, 0.0, 0.0, TimeSeconds, PHYSICS_FIELD_ARGS, EvalContext);
+    
+    return EvalContext.DataIndex == 1 ? (int)GetFieldValue(EvalContext, 0) : 0;
+}
\ No newline at end of file
diff --git a/Engine/Shaders/Private/PhysicsFieldSampler.ush b/Engine/Shaders/Private/PhysicsFieldSampler.ush
--- a/Engine/Shaders/Private/PhysicsFieldSampler.ush
+++ b/Engine/Shaders/Private/PhysicsFieldSampler.ush
@@ -7,12 +7,16 @@
 #pragma once
 
 #include "PhysicsFieldShared.ush"
+#include "PhysicsFieldEval.ush"
 
 #if FEATURE_LEVEL >= FEATURE_LEVEL_SM5
 
 #if IS_MATERIAL_SHADER
     // for materials, these are in the view UB
 	#define PhysicsFieldClipmapBuffer			View.PhysicsFieldClipmapBuffer
+	#define PhysicsFieldNodesParams				View.PhysicsFieldNodesParams
+	#define PhysicsFieldNodesOffsets			View.PhysicsFieldNodesOffsets
+	#define PhysicsFieldTargetsOffsets			View.PhysicsFieldTargetsOffsets
 	#define PhysicsFieldClipmapCenter			View.PhysicsFieldClipmapCenter
 	#define PhysicsFieldClipmapDistance			View.PhysicsFieldClipmapDistance
 	#define PhysicsFieldClipmapResolution		View.PhysicsFieldClipmapResolution
@@ -20,8 +24,13 @@
 	#define PhysicsFieldClipmapCount			View.PhysicsFieldClipmapCount
 	#define PhysicsFieldTargetCount				View.PhysicsFieldTargetCount
 	#define PhysicsFieldTargets					View.PhysicsFieldTargets
+	#define PhysicsFieldTimeSeconds				View.GameTime
+	#define PhysicsFieldTargetMappings			View.PhysicsFieldTargetMappings
 #else
 	Buffer<float> PhysicsFieldClipmapBuffer;
+	Buffer<float> PhysicsFieldNodesParams;
+	Buffer<int> PhysicsFieldNodesOffsets;
+	Buffer<int> PhysicsFieldTargetsOffsets;
 	float3		PhysicsFieldClipmapCenter;
 	float		PhysicsFieldClipmapDistance;
 	int			PhysicsFieldClipmapResolution;
@@ -29,6 +38,8 @@
 	int			PhysicsFieldClipmapCount;
 	int			PhysicsFieldTargetCount;
 	int4		PhysicsFieldTargets[MAX_PHYSICS_FIELD_TARGETS]; // x (Vector), y (Scalar), z (Integer), w (Padding)
+	float		PhysicsFieldTimeSeconds;
+	int4		PhysicsFieldTargetMappings[MAX_PHYSICS_FIELD_TARGETS]; 
 #endif
 
 float3 MatPhysicsField_SamplePhysicsVectorField(float3 WorldPosition, int VectorTarget)
@@ -50,4 +61,19 @@ int MatPhysicsField_SamplePhysicsIntegerField(float3 WorldPosition, int IntegerT
 	
 }
 
+float3 MatPhysicsField_EvalPhysicsVectorField(float3 WorldPosition, int VectorTarget)
+{
+	return PhysicsField_EvalPhysicsVectorField(WorldPosition, PhysicsFieldTargetMappings[VectorTarget].x, PhysicsFieldNodesParams, PhysicsFieldNodesOffsets, PhysicsFieldTargetsOffsets, PhysicsFieldTimeSeconds);
+}
+
+float MatPhysicsField_EvalPhysicsScalarField(float3 WorldPosition, int ScalarTarget)
+{
+	return PhysicsField_EvalPhysicsScalarField(WorldPosition, PhysicsFieldTargetMappings[ScalarTarget].y, PhysicsFieldNodesParams, PhysicsFieldNodesOffsets, PhysicsFieldTargetsOffsets, PhysicsFieldTimeSeconds);
+}
+
+int MatPhysicsField_EvalPhysicsIntegerField(float3 WorldPosition, int IntegerTarget)
+{
+	return PhysicsField_EvalPhysicsIntegerField(WorldPosition, PhysicsFieldTargetMappings[IntegerTarget].z, PhysicsFieldNodesParams, PhysicsFieldNodesOffsets, PhysicsFieldTargetsOffsets, PhysicsFieldTimeSeconds);
+}
+
 #endif // FEATURE_LEVEL >= FEATURE_LEVEL_SM5
diff --git a/Engine/Shaders/Private/PhysicsFieldVisualizer.usf b/Engine/Shaders/Private/PhysicsFieldVisualizer.usf
--- a/Engine/Shaders/Private/PhysicsFieldVisualizer.usf
+++ b/Engine/Shaders/Private/PhysicsFieldVisualizer.usf
@@ -13,7 +13,7 @@
 #include "PhysicsFieldEval.ush"
 
 RWTexture2D<float4> OutputTexture;
-float2 OutputResolution; 
+float2 OutputResolution;
 
 Buffer<float4> BoundsMin;
 Buffer<float4> BoundsMax;
@@ -128,7 +128,6 @@ void PrintVectorField()
 void MainCS(uint3 DispatchThreadId : SV_DispatchThreadID,
     uint3 GroupThreadId : SV_GroupThreadID)
 {
-	GFieldGroupThreadId = GroupThreadId.x * 16 + GroupThreadId.y * 4 + GroupThreadId.z;
 	const FFontColor AABBColor = GetVectorTargetColor();
 	if(DispatchThreadId.x < VOXEL_RESOLUTION * BoundsSize && DispatchThreadId.y < VOXEL_RESOLUTION && DispatchThreadId.z < VOXEL_RESOLUTION)
 	{
@@ -152,7 +151,7 @@ void MainCS(uint3 DispatchThreadId : SV_DispatchThreadID,
 #endif
 
 #if PERMUTATION_EVAL == 0
-			float3 FieldVector = PhysicsField_EvalPhysicsVectorField(VoxelCenter, GlobalTarget);
+			float3 FieldVector = PhysicsField_EvalPhysicsVectorField(VoxelCenter, GlobalTarget, View.PhysicsFieldNodesParams, View.PhysicsFieldNodesOffsets, View.PhysicsFieldTargetsOffsets, View.GameTime);
 #endif
 
 			const float VectorLength = length(FieldVector);
@@ -285,8 +284,6 @@ void PrintIntegerField()
 void MainCS(uint3 DispatchThreadId : SV_DispatchThreadID,
     uint3 GroupThreadId : SV_GroupThreadID)
 {
-	GFieldGroupThreadId = GroupThreadId.x * 8 + GroupThreadId.y + GroupThreadId.z;
-
 	#if PERMUTATION_FIELD == 2
 	const FFontColor AABBColor = GetIntegerTargetColor();
 	#endif
@@ -321,7 +318,7 @@ void MainCS(uint3 DispatchThreadId : SV_DispatchThreadID,
 #endif
 
 #if PERMUTATION_EVAL == 0
-					const int FieldInteger = PhysicsField_EvalPhysicsIntegerField(WorldPosition, GlobalTarget);
+					const int FieldInteger = PhysicsField_EvalPhysicsIntegerField(WorldPosition, GlobalTarget, View.PhysicsFieldNodesParams, View.PhysicsFieldNodesOffsets, View.PhysicsFieldTargetsOffsets, View.GameTime);
 #endif
 					RayDensity += FieldInteger * StepSize;
 #endif
@@ -334,7 +331,7 @@ void MainCS(uint3 DispatchThreadId : SV_DispatchThreadID,
 #endif
 
 #if PERMUTATION_EVAL == 0
-					float FieldScalar = PhysicsField_EvalPhysicsScalarField(WorldPosition, GlobalTarget);
+					float FieldScalar = PhysicsField_EvalPhysicsScalarField(WorldPosition, GlobalTarget, View.PhysicsFieldNodesParams, View.PhysicsFieldNodesOffsets, View.PhysicsFieldTargetsOffsets, View.GameTime);
 #endif
 					
 					RayDensity += FieldScalar * StepSize;
@@ -370,7 +367,6 @@ void MainCS(uint3 DispatchThreadId : SV_DispatchThreadID,
 void MainCS(uint3 DispatchThreadId : SV_DispatchThreadID,
     uint3 GroupThreadId : SV_GroupThreadID)
 {
-	GFieldGroupThreadId = GroupThreadId.x * 8 + GroupThreadId.y + GroupThreadId.z;
 	if(DispatchThreadId.x < OutputResolution.x && DispatchThreadId.y < OutputResolution.y)
 	{
 		float3 RayBegin, RayEnd, RayDirection;
diff --git a/Engine/Source/Runtime/Engine/Private/Materials/HLSLMaterialTranslator.cpp b/Engine/Source/Runtime/Engine/Private/Materials/HLSLMaterialTranslator.cpp
--- a/Engine/Source/Runtime/Engine/Private/Materials/HLSLMaterialTranslator.cpp
+++ b/Engine/Source/Runtime/Engine/Private/Materials/HLSLMaterialTranslator.cpp
@@ -12232,7 +12232,7 @@ int32 FHLSLMaterialTranslator::DistanceFieldApproxAO(int32 PositionArg, EPositio
 		*CoerceParameter(MaxDistance, MCT_Float));
 }
 
-int32 FHLSLMaterialTranslator::SamplePhysicsField(int32 PositionArg, EPositionOrigin PositionOrigin, const int32 OutputType, const int32 TargetIndex)
+int32 FHLSLMaterialTranslator::ProcessPhysicsField(int32 PositionArg, EPositionOrigin PositionOrigin, const int32 OutputType, const int32 TargetIndex, const bool bSampleField)
 {
 	if (FeatureLevel == ERHIFeatureLevel::ES3_1)
 	{
@@ -12271,15 +12271,18 @@ int32 FHLSLMaterialTranslator::SamplePhysicsField(int32 PositionArg, EPositionOr
 	// LWC_TODO: LWC aware physics field
 	if (OutputType == EFieldOutputType::Field_Output_Vector)
 	{
-		return AddCodeChunk(MCT_Float3, TEXT("MatPhysicsField_SamplePhysicsVectorField(%s,%d)"), *CoerceParameter(PositionArg, MCT_Float3), static_cast<uint8>(TargetIndex));
+		return bSampleField ? AddCodeChunk(MCT_Float3, TEXT("MatPhysicsField_SamplePhysicsVectorField(%s,%d)"), *CoerceParameter(PositionArg, MCT_Float3), static_cast<uint8>(TargetIndex)) : 
+							  AddCodeChunk(MCT_Float3, TEXT("MatPhysicsField_EvalPhysicsVectorField(%s,%d)"), *CoerceParameter(PositionArg, MCT_Float3), static_cast<uint8>(TargetIndex));
 	}
 	else if (OutputType == EFieldOutputType::Field_Output_Scalar)
 	{
-		return AddCodeChunk(MCT_Float, TEXT("MatPhysicsField_SamplePhysicsScalarField(%s,%d)"), *CoerceParameter(PositionArg, MCT_Float3), static_cast<uint8>(TargetIndex));
+		return bSampleField ? AddCodeChunk(MCT_Float, TEXT("MatPhysicsField_SamplePhysicsScalarField(%s,%d)"), *CoerceParameter(PositionArg, MCT_Float3), static_cast<uint8>(TargetIndex)) : 
+		                      AddCodeChunk(MCT_Float, TEXT("MatPhysicsField_EvalPhysicsScalarField(%s,%d)"), *CoerceParameter(PositionArg, MCT_Float3), static_cast<uint8>(TargetIndex));
 	}
 	else if (OutputType == EFieldOutputType::Field_Output_Integer)
 	{
-		return AddCodeChunk(MCT_Float, TEXT("MatPhysicsField_SamplePhysicsIntegerField(%s,%d)"), *CoerceParameter(PositionArg, MCT_Float3), static_cast<uint8>(TargetIndex));
+		return bSampleField ? AddCodeChunk(MCT_Float, TEXT("MatPhysicsField_SamplePhysicsIntegerField(%s,%d)"), *CoerceParameter(PositionArg, MCT_Float3), static_cast<uint8>(TargetIndex)) :
+							  AddCodeChunk(MCT_Float, TEXT("MatPhysicsField_EvalPhysicsIntegerField(%s,%d)"), *CoerceParameter(PositionArg, MCT_Float3), static_cast<uint8>(TargetIndex));
 	}
 	else
 	{
diff --git a/Engine/Source/Runtime/Engine/Private/Materials/HLSLMaterialTranslator.h b/Engine/Source/Runtime/Engine/Private/Materials/HLSLMaterialTranslator.h
--- a/Engine/Source/Runtime/Engine/Private/Materials/HLSLMaterialTranslator.h
+++ b/Engine/Source/Runtime/Engine/Private/Materials/HLSLMaterialTranslator.h
@@ -1187,7 +1187,7 @@ protected:
 	virtual int32 DistanceToNearestSurface(int32 PositionArg, EPositionOrigin PositionOrigin) override;
 	virtual int32 DistanceFieldGradient(int32 PositionArg, EPositionOrigin PositionOrigin) override;
 	virtual int32 DistanceFieldApproxAO(int32 PositionArg, EPositionOrigin PositionOrigin, int32 NormalArg, int32 BaseDistanceArg, int32 RadiusArg, uint32 NumSteps, float StepScale) override;
-	virtual int32 SamplePhysicsField(int32 PositionArg, EPositionOrigin PositionOrigin, const int32 OutputType, const int32 TargetIndex) override;
+	virtual int32 ProcessPhysicsField(int32 PositionArg, EPositionOrigin PositionOrigin, const int32 OutputType, const int32 TargetIndex, const bool bSampleField) override;
 	virtual int32 AtmosphericFogColor(int32 WorldPosition, EPositionOrigin PositionOrigin) override;
 	virtual int32 AtmosphericLightVector() override;
 	virtual int32 AtmosphericLightColor() override;
diff --git a/Engine/Source/Runtime/Engine/Private/Materials/MaterialExpressions.cpp b/Engine/Source/Runtime/Engine/Private/Materials/MaterialExpressions.cpp
--- a/Engine/Source/Runtime/Engine/Private/Materials/MaterialExpressions.cpp
+++ b/Engine/Source/Runtime/Engine/Private/Materials/MaterialExpressions.cpp
@@ -19237,7 +19237,7 @@ int32 UMaterialExpressionSamplePhysicsVectorField::Compile(class FMaterialCompil
 		PositionArg = Compiler->WorldPosition(UE::MaterialTranslatorUtils::GetWorldPositionTypeWithOrigin(WorldPositionOriginType));
 	}
 
-	return Compiler->SamplePhysicsField(PositionArg, WorldPositionOriginType, EFieldOutputType::Field_Output_Vector, static_cast<uint8>(FieldTarget));
+	return Compiler->ProcessPhysicsField(PositionArg, WorldPositionOriginType, EFieldOutputType::Field_Output_Vector, static_cast<uint8>(FieldTarget), true);
 }
 
 void UMaterialExpressionSamplePhysicsVectorField::GetCaption(TArray<FString>& OutCaptions) const
@@ -19245,6 +19245,31 @@ void UMaterialExpressionSamplePhysicsVectorField::GetCaption(TArray<FString>& Ou
 	OutCaptions.Add(TEXT("SamplePhysicsVectorField"));
 }
 
+///////////////////////////////////////////////////////////////////////////////
+// UMaterialExpressionEvalPhysicsVectorField
+///////////////////////////////////////////////////////////////////////////////
+
+int32 UMaterialExpressionEvalPhysicsVectorField::Compile(class FMaterialCompiler* Compiler, int32 OutputIndex)
+{
+	int32 PositionArg = INDEX_NONE;
+
+	if (WorldPosition.GetTracedInput().Expression)
+	{
+		PositionArg = WorldPosition.Compile(Compiler);
+	}
+	else
+	{
+		PositionArg = Compiler->WorldPosition(UE::MaterialTranslatorUtils::GetWorldPositionTypeWithOrigin(WorldPositionOriginType));
+	}
+
+	return Compiler->ProcessPhysicsField(PositionArg, WorldPositionOriginType, EFieldOutputType::Field_Output_Vector, static_cast<uint8>(FieldTarget), false);
+}
+
+void UMaterialExpressionEvalPhysicsVectorField::GetCaption(TArray<FString>& OutCaptions) const
+{
+	OutCaptions.Add(TEXT("EvalPhysicsVectorField"));
+}
+
 ///////////////////////////////////////////////////////////////////////////////
 // UMaterialExpressionSamplePhysicsScalarField
 ///////////////////////////////////////////////////////////////////////////////
@@ -19286,7 +19311,7 @@ int32 UMaterialExpressionSamplePhysicsScalarField::Compile(class FMaterialCompil
 		PositionArg = Compiler->WorldPosition(UE::MaterialTranslatorUtils::GetWorldPositionTypeWithOrigin(WorldPositionOriginType));
 	}
 
-	return Compiler->SamplePhysicsField(PositionArg, WorldPositionOriginType, EFieldOutputType::Field_Output_Scalar, static_cast<uint8>(FieldTarget));
+	return Compiler->ProcessPhysicsField(PositionArg, WorldPositionOriginType, EFieldOutputType::Field_Output_Scalar, static_cast<uint8>(FieldTarget), true);
 }
 
 void UMaterialExpressionSamplePhysicsScalarField::GetCaption(TArray<FString>& OutCaptions) const
@@ -19294,6 +19319,31 @@ void UMaterialExpressionSamplePhysicsScalarField::GetCaption(TArray<FString>& Ou
 	OutCaptions.Add(TEXT("SamplePhysicsScalarField"));
 }
 
+//////////////////////////////////////////////////////////////////////////////
+// UMaterialExpressionEvalPhysicsScalarField
+///////////////////////////////////////////////////////////////////////////////
+
+int32 UMaterialExpressionEvalPhysicsScalarField::Compile(class FMaterialCompiler* Compiler, int32 OutputIndex)
+{
+	int32 PositionArg = INDEX_NONE;
+
+	if (WorldPosition.GetTracedInput().Expression)
+	{
+		PositionArg = WorldPosition.Compile(Compiler);
+	}
+	else
+	{
+		PositionArg = Compiler->WorldPosition(UE::MaterialTranslatorUtils::GetWorldPositionTypeWithOrigin(WorldPositionOriginType));
+	}
+
+	return Compiler->ProcessPhysicsField(PositionArg, WorldPositionOriginType, EFieldOutputType::Field_Output_Scalar, static_cast<uint8>(FieldTarget), false);
+}
+
+void UMaterialExpressionEvalPhysicsScalarField::GetCaption(TArray<FString>& OutCaptions) const
+{
+	OutCaptions.Add(TEXT("EvalPhysicsScalarField"));
+}
+
 ///////////////////////////////////////////////////////////////////////////////
 // UMaterialExpressionSamplePhysicsIntegerField
 ///////////////////////////////////////////////////////////////////////////////
@@ -19335,13 +19385,38 @@ int32 UMaterialExpressionSamplePhysicsIntegerField::Compile(class FMaterialCompi
 		PositionArg = Compiler->WorldPosition(UE::MaterialTranslatorUtils::GetWorldPositionTypeWithOrigin(WorldPositionOriginType));
 	}
 
-	return Compiler->SamplePhysicsField(PositionArg, WorldPositionOriginType, EFieldOutputType::Field_Output_Integer, static_cast<uint8>(FieldTarget));
+	return Compiler->ProcessPhysicsField(PositionArg, WorldPositionOriginType, EFieldOutputType::Field_Output_Integer, static_cast<uint8>(FieldTarget), true);
 }
 
 void UMaterialExpressionSamplePhysicsIntegerField::GetCaption(TArray<FString>& OutCaptions) const
 {
 	OutCaptions.Add(TEXT("SamplePhysicsIntegerField"));
 }
+
+///////////////////////////////////////////////////////////////////////////////
+// UMaterialExpressionEvalPhysicsIntegerField
+///////////////////////////////////////////////////////////////////////////////
+
+int32 UMaterialExpressionEvalPhysicsIntegerField::Compile(class FMaterialCompiler* Compiler, int32 OutputIndex)
+{
+	int32 PositionArg = INDEX_NONE;
+
+	if (WorldPosition.GetTracedInput().Expression)
+	{
+		PositionArg = WorldPosition.Compile(Compiler);
+	}
+	else
+	{
+		PositionArg = Compiler->WorldPosition(UE::MaterialTranslatorUtils::GetWorldPositionTypeWithOrigin(WorldPositionOriginType));
+	}
+
+	return Compiler->ProcessPhysicsField(PositionArg, WorldPositionOriginType, EFieldOutputType::Field_Output_Integer, static_cast<uint8>(FieldTarget), false);
+}
+
+void UMaterialExpressionEvalPhysicsIntegerField::GetCaption(TArray<FString>& OutCaptions) const
+{
+	OutCaptions.Add(TEXT("EvalPhysicsIntegerField"));
+}
 #endif // WITH_EDITOR
 
 
diff --git a/Engine/Source/Runtime/Engine/Private/Materials/MaterialExpressionsIR.cpp b/Engine/Source/Runtime/Engine/Private/Materials/MaterialExpressionsIR.cpp
--- a/Engine/Source/Runtime/Engine/Private/Materials/MaterialExpressionsIR.cpp
+++ b/Engine/Source/Runtime/Engine/Private/Materials/MaterialExpressionsIR.cpp
@@ -4576,7 +4576,7 @@ void UMaterialExpressionBindlessSwitch::Build(MIR::FEmitter& Em)
 	Em.Output(0, OutputValue);
 }
 
-static FValueRef EmitSamplePhysicsField(MIR::FEmitter& Em, FExpressionInput* PositionInput, EPositionOrigin PositionOriginType, EFieldOutputType OutputType, const int32 TargetIndex)
+static FValueRef EmitProcessPhysicsField(MIR::FEmitter& Em, FExpressionInput* PositionInput, EPositionOrigin PositionOriginType, EFieldOutputType OutputType, const int32 TargetIndex, const bool bSampleField)
 {
 	FValueRef PositionValue = Em.TryInput(PositionInput);
 	if (!PositionValue)
@@ -4618,15 +4618,18 @@ static FValueRef EmitSamplePhysicsField(MIR::FEmitter& Em, FExpressionInput* Pos
 	// LWC_TODO: LWC aware physics field
 	if (OutputType == Field_Output_Vector)
 	{
-		return Em.InlineHLSL(MIR::FType::MakeFloatVector(3), TEXTVIEW("MatPhysicsField_SamplePhysicsVectorField($0, $1)"), { PositionValue, TargetIndexValue });
+		return bSampleField ? Em.InlineHLSL(MIR::FType::MakeFloatVector(3), TEXTVIEW("MatPhysicsField_SamplePhysicsVectorField($0, $1)"), { PositionValue, TargetIndexValue }) :
+								Em.InlineHLSL(MIR::FType::MakeFloatVector(3), TEXTVIEW("MatPhysicsField_EvalPhysicsVectorField($0, $1)"), { PositionValue, TargetIndexValue });
 	}
 	else if (OutputType == Field_Output_Scalar)
 	{
-		return Em.InlineHLSL(MIR::FType::MakeFloatScalar(), TEXTVIEW("MatPhysicsField_SamplePhysicsScalarField($0, $1)"), { PositionValue, TargetIndexValue });
+		return bSampleField ? Em.InlineHLSL(MIR::FType::MakeFloatScalar(), TEXTVIEW("MatPhysicsField_SamplePhysicsScalarField($0, $1)"), { PositionValue, TargetIndexValue }) :
+								Em.InlineHLSL(MIR::FType::MakeFloatScalar(), TEXTVIEW("MatPhysicsField_EvalPhysicsScalarField($0, $1)"), { PositionValue, TargetIndexValue });
 	}
 	else if (OutputType == Field_Output_Integer)
 	{
-		return Em.InlineHLSL(MIR::FType::MakeFloatScalar(), TEXTVIEW("MatPhysicsField_SamplePhysicsIntegerField($0, $1)"), { PositionValue, TargetIndexValue });
+		return bSampleField ? Em.InlineHLSL(MIR::FType::MakeFloatScalar(), TEXTVIEW("MatPhysicsField_SamplePhysicsIntegerField($0, $1)"), { PositionValue, TargetIndexValue }) :
+								Em.InlineHLSL(MIR::FType::MakeFloatScalar(), TEXTVIEW("MatPhysicsField_EvalPhysicsIntegerField($0, $1)"), { PositionValue, TargetIndexValue }) ;
 	}
 	else
 	{
@@ -4637,17 +4640,32 @@ static FValueRef EmitSamplePhysicsField(MIR::FEmitter& Em, FExpressionInput* Pos
 
 void UMaterialExpressionSamplePhysicsIntegerField::Build(MIR::FEmitter& Em)
 {
-	Em.Output(0, EmitSamplePhysicsField(Em, &WorldPosition, WorldPositionOriginType, Field_Output_Integer, FieldTarget));
+	Em.Output(0, EmitProcessPhysicsField(Em, &WorldPosition, WorldPositionOriginType, Field_Output_Integer, FieldTarget, true));
 }
 
 void UMaterialExpressionSamplePhysicsScalarField::Build(MIR::FEmitter& Em)
 {
-	Em.Output(0, EmitSamplePhysicsField(Em, &WorldPosition, WorldPositionOriginType, Field_Output_Scalar, FieldTarget));
+	Em.Output(0, EmitProcessPhysicsField(Em, &WorldPosition, WorldPositionOriginType, Field_Output_Scalar, FieldTarget, true));
 }
 
 void UMaterialExpressionSamplePhysicsVectorField::Build(MIR::FEmitter& Em)
 {
-	Em.Output(0, EmitSamplePhysicsField(Em, &WorldPosition, WorldPositionOriginType, Field_Output_Vector, FieldTarget));
+	Em.Output(0, EmitProcessPhysicsField(Em, &WorldPosition, WorldPositionOriginType, Field_Output_Vector, FieldTarget, true));
+}
+
+void UMaterialExpressionEvalPhysicsIntegerField::Build(MIR::FEmitter& Em)
+{
+	Em.Output(0, EmitProcessPhysicsField(Em, &WorldPosition, WorldPositionOriginType, Field_Output_Integer, FieldTarget, false));
+}
+
+void UMaterialExpressionEvalPhysicsScalarField::Build(MIR::FEmitter& Em)
+{
+	Em.Output(0, EmitProcessPhysicsField(Em, &WorldPosition, WorldPositionOriginType, Field_Output_Scalar, FieldTarget, false));
+}
+
+void UMaterialExpressionEvalPhysicsVectorField::Build(MIR::FEmitter& Em)
+{
+	Em.Output(0, EmitProcessPhysicsField(Em, &WorldPosition, WorldPositionOriginType, Field_Output_Vector, FieldTarget, false));
 }
 
 void UMaterialExpressionSobol::Build(MIR::FEmitter& Em)
diff --git a/Engine/Source/Runtime/Engine/Private/PhysicsField/PhysicsFieldComponent.cpp b/Engine/Source/Runtime/Engine/Private/PhysicsField/PhysicsFieldComponent.cpp
--- a/Engine/Source/Runtime/Engine/Private/PhysicsField/PhysicsFieldComponent.cpp
+++ b/Engine/Source/Runtime/Engine/Private/PhysicsField/PhysicsFieldComponent.cpp
@@ -853,7 +853,7 @@ void UPhysicsFieldComponent::SendRenderDynamicData_Concurrent()
 {
 	SCOPE_CYCLE_COUNTER(STAT_PhysicsFields_SendRenderData);
 
-	Super::SendRenderTransform_Concurrent();
+	Super::SendRenderDynamicData_Concurrent();
 
 	static IConsoleVariable* TransientLifetimeCVar = IConsoleManager::Get().FindConsoleVariable(TEXT("r.PhysicsField.Rendering.TransientLifetime"));
 	const float TransientLifetime = TransientLifetimeCVar ? TransientLifetimeCVar->GetFloat() : 0.0f;
diff --git a/Engine/Source/Runtime/Engine/Private/SceneManagement.cpp b/Engine/Source/Runtime/Engine/Private/SceneManagement.cpp
--- a/Engine/Source/Runtime/Engine/Private/SceneManagement.cpp
+++ b/Engine/Source/Runtime/Engine/Private/SceneManagement.cpp
@@ -1247,7 +1247,11 @@ FViewUniformShaderParameters::FViewUniformShaderParameters()
 
 	SkyIrradianceEnvironmentMap = GIdentityPrimitiveBuffer.SkyIrradianceEnvironmentMapSRV;
 
+	// Physics field
 	PhysicsFieldClipmapBuffer = GWhiteVertexBufferWithSRV->ShaderResourceViewRHI;
+	PhysicsFieldNodesParams = GWhiteVertexBufferWithSRV->ShaderResourceViewRHI;
+	PhysicsFieldNodesOffsets = GWhiteVertexBufferWithSRV->ShaderResourceViewRHI;
+	PhysicsFieldTargetsOffsets = GWhiteVertexBufferWithSRV->ShaderResourceViewRHI;
 
 	// Water
 	WaterIndirection = GWhiteVertexBufferWithSRV->ShaderResourceViewRHI;
diff --git a/Engine/Source/Runtime/Engine/Public/MaterialCompiler.h b/Engine/Source/Runtime/Engine/Public/MaterialCompiler.h
--- a/Engine/Source/Runtime/Engine/Public/MaterialCompiler.h
+++ b/Engine/Source/Runtime/Engine/Public/MaterialCompiler.h
@@ -645,7 +645,7 @@ public:
 	virtual int32 DistanceToNearestSurface(int32 PositionArg, EPositionOrigin PositionOrigin) = 0;
 	virtual int32 DistanceFieldGradient(int32 PositionArg, EPositionOrigin PositionOrigin) = 0;
 	virtual int32 DistanceFieldApproxAO(int32 PositionArg, EPositionOrigin PositionOrigin, int32 NormalArg, int32 BaseDistanceArg, int32 RadiusArg, uint32 NumSteps, float StepScale) = 0;
-	virtual int32 SamplePhysicsField(int32 PositionArg, EPositionOrigin PositionOrigin, const int32 OutputType, const int32 TargetIndex) = 0;
+	virtual int32 ProcessPhysicsField(int32 PositionArg, EPositionOrigin PositionOrigin, const int32 OutputType, const int32 TargetIndex, const bool bSampleField) = 0;
 	virtual int32 DepthOfFieldFunction(int32 Depth, int32 FunctionValueIndex) = 0;
 	virtual int32 PostVolumeUserFlagTestFunction(int32 Input) = 0;
 	virtual int32 AtmosphericFogColor(int32 WorldPosition, EPositionOrigin PositionOrigin) = 0;
@@ -1221,7 +1221,7 @@ public:
 	virtual int32 DistanceToNearestSurface(int32 PositionArg, EPositionOrigin PositionOrigin) override { return Compiler->DistanceToNearestSurface(PositionArg, PositionOrigin); }
 	virtual int32 DistanceFieldGradient(int32 PositionArg, EPositionOrigin PositionOrigin) override { return Compiler->DistanceFieldGradient(PositionArg, PositionOrigin); }
 	virtual int32 DistanceFieldApproxAO(int32 PositionArg, EPositionOrigin PositionOrigin, int32 NormalArg, int32 BaseDistanceArg, int32 RadiusArg, uint32 NumSteps, float StepScale) override { return Compiler->DistanceFieldApproxAO(PositionArg, PositionOrigin, NormalArg, BaseDistanceArg, RadiusArg, NumSteps, StepScale); }
-	virtual int32 SamplePhysicsField(int32 PositionArg, EPositionOrigin PositionOrigin, const int32 OutputType, const int32 TargetIndex)  override { return Compiler->SamplePhysicsField(PositionArg, PositionOrigin, OutputType, TargetIndex); }
+	virtual int32 ProcessPhysicsField(int32 PositionArg, EPositionOrigin PositionOrigin, const int32 OutputType, const int32 TargetIndex, const bool bSampleField)  override { return Compiler->ProcessPhysicsField(PositionArg, PositionOrigin, OutputType, TargetIndex, bSampleField); }
 	virtual int32 PerInstanceRandom() override { return Compiler->PerInstanceRandom(); }
 	virtual int32 PerInstanceFadeAmount() override { return Compiler->PerInstanceFadeAmount(); }
 	virtual int32 PerInstanceCustomData(int32 DataIndex, int32 DefaultValueIndex) override { return Compiler->PerInstanceCustomData(DataIndex, DefaultValueIndex); }
diff --git a/Engine/Source/Runtime/Engine/Public/Materials/MaterialExpressionSamplePhysicsField.h b/Engine/Source/Runtime/Engine/Public/Materials/MaterialExpressionSamplePhysicsField.h
--- a/Engine/Source/Runtime/Engine/Public/Materials/MaterialExpressionSamplePhysicsField.h
+++ b/Engine/Source/Runtime/Engine/Public/Materials/MaterialExpressionSamplePhysicsField.h
@@ -10,7 +10,7 @@
 #include "MaterialExpressionSamplePhysicsField.generated.h"
 
 /**
- * Material expresions to sample the global field
+ * Material expresions to sample the global physics field
  */
 
 UCLASS()
@@ -43,6 +43,22 @@ public:
 	//~ End UMaterialExpression Interface
 };
 
+UCLASS()
+class UMaterialExpressionEvalPhysicsVectorField : public UMaterialExpressionSamplePhysicsVectorField
+{
+	GENERATED_BODY()
+
+public:
+
+	//~ Begin UMaterialExpression Interface
+#if WITH_EDITOR
+	virtual void Build(MIR::FEmitter& Emitter) override;
+	virtual int32 Compile(class FMaterialCompiler* Compiler, int32 OutputIndex) override;
+	virtual void GetCaption(TArray<FString>& OutCaptions) const override;
+#endif
+	//~ End UMaterialExpression Interface
+};
+
 UCLASS()
 class UMaterialExpressionSamplePhysicsScalarField : public UMaterialExpression
 {
@@ -73,6 +89,22 @@ public:
 	//~ End UMaterialExpression Interface
 };
 
+UCLASS()
+class UMaterialExpressionEvalPhysicsScalarField : public UMaterialExpressionSamplePhysicsScalarField
+{
+	GENERATED_BODY()
+
+public:
+
+	//~ Begin UMaterialExpression Interface
+#if WITH_EDITOR
+	virtual void Build(MIR::FEmitter& Emitter) override;
+	virtual int32 Compile(class FMaterialCompiler* Compiler, int32 OutputIndex) override;
+	virtual void GetCaption(TArray<FString>& OutCaptions) const override;
+#endif
+	//~ End UMaterialExpression Interface
+};
+
 UCLASS()
 class UMaterialExpressionSamplePhysicsIntegerField : public UMaterialExpression
 {
@@ -102,3 +134,20 @@ public:
 #endif
 	//~ End UMaterialExpression Interface
 };
+
+
+UCLASS()
+class UMaterialExpressionEvalPhysicsIntegerField : public UMaterialExpressionSamplePhysicsIntegerField
+{
+	GENERATED_BODY()
+
+public:
+	
+	//~ Begin UMaterialExpression Interface
+#if WITH_EDITOR
+	virtual void Build(MIR::FEmitter& Emitter) override;
+	virtual int32 Compile(class FMaterialCompiler* Compiler, int32 OutputIndex) override;
+	virtual void GetCaption(TArray<FString>& OutCaptions) const override;
+#endif
+	//~ End UMaterialExpression Interface
+};
diff --git a/Engine/Source/Runtime/Engine/Public/SceneView.h b/Engine/Source/Runtime/Engine/Public/SceneView.h
--- a/Engine/Source/Runtime/Engine/Public/SceneView.h
+++ b/Engine/Source/Runtime/Engine/Public/SceneView.h
@@ -1114,6 +1114,7 @@ enum ETranslucencyVolumeCascade
 	VIEW_UNIFORM_BUFFER_MEMBER(int, PhysicsFieldClipmapCount) \
 	VIEW_UNIFORM_BUFFER_MEMBER(int, PhysicsFieldTargetCount) \
 	VIEW_UNIFORM_BUFFER_MEMBER_ARRAY(FIntVector4, PhysicsFieldTargets, MAX_PHYSICS_FIELD_TARGETS) \
+	VIEW_UNIFORM_BUFFER_MEMBER_ARRAY(FIntVector4, PhysicsFieldTargetMappings, MAX_PHYSICS_FIELD_TARGETS) \
 	VIEW_UNIFORM_BUFFER_MEMBER_PER_VIEW(uint32, GPUSceneViewId) \
 	VIEW_UNIFORM_BUFFER_MEMBER(float, ViewResolutionFraction) \
 	VIEW_UNIFORM_BUFFER_MEMBER(float, SubSurfaceColorAsTransmittanceAtDistanceInMeters) \
@@ -1280,7 +1281,11 @@ BEGIN_GLOBAL_SHADER_PARAMETER_STRUCT_WITH_CONSTRUCTOR(FViewUniformShaderParamete
 
 	SHADER_PARAMETER_UAV(RWStructuredBuffer<uint>, VTFeedbackBuffer)
 
+	// Physics Field
 	SHADER_PARAMETER_SRV(Buffer<float>, PhysicsFieldClipmapBuffer)
+	SHADER_PARAMETER_SRV(Buffer<float>, PhysicsFieldNodesParams)
+	SHADER_PARAMETER_SRV(Buffer<int>, PhysicsFieldNodesOffsets)
+	SHADER_PARAMETER_SRV(Buffer<int>, PhysicsFieldTargetsOffsets)
 
 	// Blue noise texture
 	SHADER_PARAMETER_TEXTURE(Texture2D<float>, BlueNoiseScalarTexture)
diff --git a/Engine/Source/Runtime/Renderer/Private/PhysicsFieldRendering.cpp b/Engine/Source/Runtime/Renderer/Private/PhysicsFieldRendering.cpp
--- a/Engine/Source/Runtime/Renderer/Private/PhysicsFieldRendering.cpp
+++ b/Engine/Source/Runtime/Renderer/Private/PhysicsFieldRendering.cpp
@@ -68,10 +68,6 @@ class FPhysicsFieldRayMarchingCS : public FGlobalShader
 		SHADER_PARAMETER_STRUCT_REF(FViewUniformShaderParameters, ViewUniformBuffer)
 		SHADER_PARAMETER_SRV(Buffer<float4>, BoundsMin)
 		SHADER_PARAMETER_SRV(Buffer<float4>, BoundsMax)
-		SHADER_PARAMETER_SRV(Buffer<float>, NodesParams)
-		SHADER_PARAMETER_SRV(Buffer<int>, NodesOffsets)
-		SHADER_PARAMETER_SRV(Buffer<int>, TargetsOffsets)
-		SHADER_PARAMETER(float, TimeSeconds)
 		SHADER_PARAMETER(uint32, BoundsOffset)
 		SHADER_PARAMETER(uint32, BoundsSize)
 		SHADER_PARAMETER(uint32, LocalTarget)
@@ -118,11 +114,6 @@ static FPhysicsFieldRayMarchingCS::FParameters* CreateShaderParameters(FRDGBuild
 	Parameters->BoundsOffset = PhysicsFieldResource->FieldInfos.BoundsOffsets[TargetType];
 	Parameters->BoundsSize = PhysicsFieldResource->FieldInfos.BoundsOffsets[TargetType + 1] - PhysicsFieldResource->FieldInfos.BoundsOffsets[TargetType];;
 
-	Parameters->NodesParams = PhysicsFieldResource->NodesParams.SRV;
-	Parameters->NodesOffsets = PhysicsFieldResource->NodesOffsets.SRV;
-	Parameters->TargetsOffsets = PhysicsFieldResource->TargetsOffsets.SRV;
-	Parameters->TimeSeconds = PhysicsFieldResource->FieldInfos.TimeSeconds;
-
 	Parameters->LocalTarget = TargetIndex;
 	Parameters->GlobalTarget = TargetType;
 	Parameters->PrintOffset = PrintOffset;
diff --git a/Engine/Source/Runtime/Renderer/Private/SceneRendering.cpp b/Engine/Source/Runtime/Renderer/Private/SceneRendering.cpp
--- a/Engine/Source/Runtime/Renderer/Private/SceneRendering.cpp
+++ b/Engine/Source/Runtime/Renderer/Private/SceneRendering.cpp
@@ -1246,10 +1246,16 @@ void SetupPhysicsFieldUniformBufferParameters(const FScene* Scene, FEngineShowFl
 		if (FieldResource->FieldInfos.bBuildClipmap)
 		{
 			ViewUniformShaderParameters.PhysicsFieldClipmapBuffer = FieldResource->ClipmapBuffer.SRV.GetReference();
+			ViewUniformShaderParameters.PhysicsFieldNodesParams = FieldResource->NodesParams.SRV.GetReference();
+			ViewUniformShaderParameters.PhysicsFieldNodesOffsets = FieldResource->NodesOffsets.SRV.GetReference();
+			ViewUniformShaderParameters.PhysicsFieldTargetsOffsets = FieldResource->TargetsOffsets.SRV.GetReference();
 		}
 		else
 		{
 			ViewUniformShaderParameters.PhysicsFieldClipmapBuffer = GWhiteVertexBufferWithSRV->ShaderResourceViewRHI;
+			ViewUniformShaderParameters.PhysicsFieldNodesParams = GWhiteVertexBufferWithSRV->ShaderResourceViewRHI;
+			ViewUniformShaderParameters.PhysicsFieldNodesOffsets = GWhiteVertexBufferWithSRV->ShaderResourceViewRHI;
+			ViewUniformShaderParameters.PhysicsFieldTargetsOffsets = GWhiteVertexBufferWithSRV->ShaderResourceViewRHI;
 		}
 		ViewUniformShaderParameters.PhysicsFieldClipmapCenter = (FVector3f)FieldResource->FieldInfos.ClipmapCenter;
 		ViewUniformShaderParameters.PhysicsFieldClipmapDistance = FieldResource->FieldInfos.ClipmapDistance;
@@ -1257,18 +1263,31 @@ void SetupPhysicsFieldUniformBufferParameters(const FScene* Scene, FEngineShowFl
 		ViewUniformShaderParameters.PhysicsFieldClipmapExponent = FieldResource->FieldInfos.ClipmapExponent;
 		ViewUniformShaderParameters.PhysicsFieldClipmapCount = FieldResource->FieldInfos.ClipmapCount;
 		ViewUniformShaderParameters.PhysicsFieldTargetCount = FieldResource->FieldInfos.TargetCount;
+
+		const TArray<EFieldPhysicsType> VectorTypes = GetFieldTargetTypes(EFieldOutputType::Field_Output_Vector);
+		const TArray<EFieldPhysicsType> ScalarTypes = GetFieldTargetTypes(EFieldOutputType::Field_Output_Scalar);
+		const TArray<EFieldPhysicsType> IntegerTypes = GetFieldTargetTypes(EFieldOutputType::Field_Output_Integer);
+
 		for (int32 Index = 0; Index < MAX_PHYSICS_FIELD_TARGETS; ++Index)
 		{
 			ViewUniformShaderParameters.PhysicsFieldTargets[Index].X = FieldResource->FieldInfos.VectorTargets[Index];
 			ViewUniformShaderParameters.PhysicsFieldTargets[Index].Y = FieldResource->FieldInfos.ScalarTargets[Index];
 			ViewUniformShaderParameters.PhysicsFieldTargets[Index].Z = FieldResource->FieldInfos.IntegerTargets[Index];
 			ViewUniformShaderParameters.PhysicsFieldTargets[Index].W = 0; // Padding
+
+			ViewUniformShaderParameters.PhysicsFieldTargetMappings[Index].X = Index < VectorTypes.Num() ? VectorTypes[Index] : 0;
+			ViewUniformShaderParameters.PhysicsFieldTargetMappings[Index].Y = Index < ScalarTypes.Num() ? ScalarTypes[Index] : 0;
+			ViewUniformShaderParameters.PhysicsFieldTargetMappings[Index].Z = Index < IntegerTypes.Num() ? IntegerTypes[Index] : 0;
+			ViewUniformShaderParameters.PhysicsFieldTargetMappings[Index].W = 0;
 		}
 	}
 	else
 	{
 		TStaticArray<UE::Core::TAlignedElement<FIntVector4, 16>, MAX_PHYSICS_FIELD_TARGETS> EmptyTargets = {};
 		ViewUniformShaderParameters.PhysicsFieldClipmapBuffer = GWhiteVertexBufferWithSRV->ShaderResourceViewRHI;
+		ViewUniformShaderParameters.PhysicsFieldNodesParams = GWhiteVertexBufferWithSRV->ShaderResourceViewRHI;
+		ViewUniformShaderParameters.PhysicsFieldNodesOffsets = GWhiteVertexBufferWithSRV->ShaderResourceViewRHI;
+		ViewUniformShaderParameters.PhysicsFieldTargetsOffsets = GWhiteVertexBufferWithSRV->ShaderResourceViewRHI;
 		ViewUniformShaderParameters.PhysicsFieldClipmapCenter = FVector3f::ZeroVector;
 		ViewUniformShaderParameters.PhysicsFieldClipmapDistance = 1.0;
 		ViewUniformShaderParameters.PhysicsFieldClipmapResolution = 2;
@@ -1276,10 +1295,10 @@ void SetupPhysicsFieldUniformBufferParameters(const FScene* Scene, FEngineShowFl
 		ViewUniformShaderParameters.PhysicsFieldClipmapCount = 1;
 		ViewUniformShaderParameters.PhysicsFieldTargetCount = 0;
 		ViewUniformShaderParameters.PhysicsFieldTargets = EmptyTargets;
+		ViewUniformShaderParameters.PhysicsFieldTargetMappings = EmptyTargets;
 	}
 }
 
-
 FIntPoint FViewInfo::GetSecondaryViewRectSize() const
 {
 	return FIntPoint(
