Substrate - temp fix for compilation issue

#rnx

#Synced-CL 50935098

diff --git a/Engine/Shaders/Private/Substrate/SubstrateEvaluation.ush b/Engine/Shaders/Private/Substrate/SubstrateEvaluation.ush
--- a/Engine/Shaders/Private/Substrate/SubstrateEvaluation.ush
+++ b/Engine/Shaders/Private/Substrate/SubstrateEvaluation.ush
@@ -2168,7 +2168,7 @@ FSubstrateEnvLightResult SubstrateEvaluateForEnvLight(FSubstrateBSDFContext BSDF
 			{
 				FBxDFEnergyTermsRGB SpecularEnergyTerms = ComputeGGXSpecEnergyTermsRGB_Env(SafeRoughness, BSDFContext.Context.NoV, F0, F90);
 				FBxDFEnergyTermsRGB DiffuseEnergyTerms = ComputeDiffuseEnergyTermsRGB(SafeRoughness, BSDFContext.Context.NoV);
-				MSScale = ComputeEnergyConservation(SpecularEnergyTerms);
+				MSScale = 1;// ComputeEnergyConservation(SpecularEnergyTerms);
 
 				SpecEnvBRDF = MSScale * SpecularEnergyTerms.E;
 				DiffEnvBRDF = DiffuseEnergyTerms.E;
