Substrate - toon - fix for MSSCale to be float3 and compile as everything is RGB in this case.
After temp fix 50937629.

#tests pc editor engine test
#rnx
#Synced-CL 50935098

diff --git a/Engine/Shaders/Private/Substrate/SubstrateEvaluation.ush b/Engine/Shaders/Private/Substrate/SubstrateEvaluation.ush
--- a/Engine/Shaders/Private/Substrate/SubstrateEvaluation.ush
+++ b/Engine/Shaders/Private/Substrate/SubstrateEvaluation.ush
@@ -2163,12 +2163,12 @@ FSubstrateEnvLightResult SubstrateEvaluateForEnvLight(FSubstrateBSDFContext BSDF
 			float3 DiffEnvBRDF = 0.0f;
 			float3 SpecEnvBRDF = 0.0f;
 			float DirectionalAlbedo_SpecularTransmission = 1.0f;
-			FBxDFEnergyType MSScale = 1.0f;
+			float3 MSScale = 1.0f;
 			if (ToonEvalSettings.bEnergyConservation)
 			{
 				FBxDFEnergyTermsRGB SpecularEnergyTerms = ComputeGGXSpecEnergyTermsRGB_Env(SafeRoughness, BSDFContext.Context.NoV, F0, F90);
 				FBxDFEnergyTermsRGB DiffuseEnergyTerms = ComputeDiffuseEnergyTermsRGB(SafeRoughness, BSDFContext.Context.NoV);
-				MSScale = 1;// ComputeEnergyConservation(SpecularEnergyTerms);
+				MSScale = ComputeEnergyConservation(SpecularEnergyTerms);
 
 				SpecEnvBRDF = MSScale * SpecularEnergyTerms.E;
 				DiffEnvBRDF = DiffuseEnergyTerms.E;
