[MegaLights] Remove double add of lighting in MATERIAL_SHADINGMODEL_THIN_TRANSLUCENT codepath.

#Synced-CL 50941167

diff --git a/Engine/Shaders/Private/BasePassPixelShader.usf b/Engine/Shaders/Private/BasePassPixelShader.usf
--- a/Engine/Shaders/Private/BasePassPixelShader.usf
+++ b/Engine/Shaders/Private/BasePassPixelShader.usf
@@ -1442,8 +1442,6 @@ void FPixelShaderInOut_MainPS(
 						float3 TLVDiffuseLighting;
 						float3 TLVSpecularLighting;
 						GetTranslucencyVolumeLighting(MaterialParameters, PixelMaterialInputs, BasePassInterpolants, GBuffer, IndirectIrradiance, TLVDiffuseLighting, TLVSpecularLighting);
-						DiffuseColor += TLVDiffuseLighting;
-						ColorSeparateSpecular += TLVSpecularLighting;
 					
 						#if MATERIAL_SHADINGMODEL_THIN_TRANSLUCENT
 							DiffuseColor += TLVDiffuseLighting;
