Substrate - Toon - fixed function type cast warning as error.

#tests PC editor dx12 - tested fix in engine test.
#rnx
#Synced-CL 50966613

diff --git a/Engine/Shaders/Private/ToonProfileCommon.ush b/Engine/Shaders/Private/ToonProfileCommon.ush
--- a/Engine/Shaders/Private/ToonProfileCommon.ush
+++ b/Engine/Shaders/Private/ToonProfileCommon.ush
@@ -28,8 +28,7 @@ float3 EvaluateToonProfileSpecularRamp(in float ToonProfileId, in float Specular
 
 float EvaluateToonProfileSpecularIndirectRamp(in float ToonProfileId, in float Level)
 {
-	// frac for looping ramp over intensity
-	return EvaluateToonProfileRamp(ToonProfileId, Level, TOON_PROFILE_SPEC_IND_RAMP_Y);
+	return EvaluateToonProfileRamp(ToonProfileId, Level, TOON_PROFILE_SPEC_IND_RAMP_Y).x;
 }
 
 float EvaluateToonProfileShadowRamp(in float ToonProfileId, in float Transmittance)
