Path Tracing: Convert payload shading model ID to a path tracer specific set of values

This makes it more explicit that the path tracer has its own way of tracking the different shading models. This also avoids future clashes as the engine's SHADINGMODELID_* set of defines expands.

#Synced-CL 50953311

diff --git a/Engine/Plugins/Experimental/GPULightmass/Shaders/Private/LightmapPathTracing.usf b/Engine/Plugins/Experimental/GPULightmass/Shaders/Private/LightmapPathTracing.usf
--- a/Engine/Plugins/Experimental/GPULightmass/Shaders/Private/LightmapPathTracing.usf
+++ b/Engine/Plugins/Experimental/GPULightmass/Shaders/Private/LightmapPathTracing.usf
@@ -369,7 +369,7 @@ void PathTracingKernel(
 			Payload.SubsurfaceColor = 0;
 			Payload.BSDFOpacity = 1;
 			Payload.TransparencyColor = 0;
-			Payload.ShadingModelID = SHADINGMODELID_DEFAULT_LIT;
+			Payload.PathTracingShadingModelID = PATH_TRACING_SHADINGMODELID_BASIC;
 			Payload.PrimitiveLightingChannelMask = 0b111;
 			Payload.HitT = 1.0f;
 			Payload.SetFrontFace();
diff --git a/Engine/Shaders/Private/PathTracing/Material/PathTracingMaterialSampling.ush b/Engine/Shaders/Private/PathTracing/Material/PathTracingMaterialSampling.ush
--- a/Engine/Shaders/Private/PathTracing/Material/PathTracingMaterialSampling.ush
+++ b/Engine/Shaders/Private/PathTracing/Material/PathTracingMaterialSampling.ush
@@ -24,32 +24,32 @@ float3 EstimateMaterialAlbedo(FPathTracingPayload Payload)
 {
 	// TODO: make this more accurate by estimating the directional albedo?
 	float3 Albedo = 0;
-	switch (Payload.ShadingModelID)
+	switch (Payload.PathTracingShadingModelID)
 	{
 #if SIMPLIFIED_MATERIAL_SHADER
-		case SHADINGMODELID_TWOSIDED_FOLIAGE:
+		case PATH_TRACING_SHADINGMODELID_FOLIAGE:
 			Albedo = Payload.BaseColor + Payload.SubsurfaceColor;
 			break;
 #else
-		case SHADINGMODELID_SUBSTRATE:
+		case PATH_TRACING_SHADINGMODELID_BASIC:
 		{
 			Albedo = Payload.DiffuseColor + Payload.SubsurfaceColor + Payload.SpecularColor + Payload.FuzzAmount * Payload.FuzzColor;
 			Albedo *= Payload.GetMaxLobeWeight();
 			break;
 		}
-		case SHADINGMODELID_SOLID_GLASS:
+		case PATH_TRACING_SHADINGMODELID_SOLID_GLASS:
 		{
 			Albedo = 1.0 + Payload.FuzzAmount * Payload.FuzzColor;
 			Albedo *= Payload.GetMaxLobeWeight();
 			break;
 		}
-		case SHADINGMODELID_TWOSIDED_FOLIAGE:
+		case PATH_TRACING_SHADINGMODELID_FOLIAGE:
 		{
 			Albedo = Payload.DiffuseColor + Payload.GetFoliageTransmissionColor() + Payload.SpecularColor + Payload.FuzzAmount * Payload.FuzzColor;
 			Albedo *= Payload.GetMaxLobeWeight();
 			break;
 		}
-		case SHADINGMODELID_THIN_TRANSLUCENT:
+		case PATH_TRACING_SHADINGMODELID_THIN_GLASS:
 		{
 			const float F0 = F0RGBToF0(Payload.SpecularColor);
 			Albedo = Payload.HasRefraction() ? lerp(Payload.GetTransmittanceColor(), 1.0, F0) : F0;
@@ -57,16 +57,16 @@ float3 EstimateMaterialAlbedo(FPathTracingPayload Payload)
 			Albedo *= Payload.GetMaxLobeWeight();
 			break;
 		}
-		case SHADINGMODELID_CLEAR_COAT:
+		case PATH_TRACING_SHADINGMODELID_CLEARCOAT:
 		{
 			Albedo = Payload.DiffuseColor + Payload.SubsurfaceColor + Payload.SpecularColor + Payload.RoughnessData.z * CLEAR_COAT_F0 + Payload.FuzzAmount * Payload.FuzzColor;
 			Albedo *= Payload.GetMaxLobeWeight();
 			break;
 		}
-		case SHADINGMODELID_EYE:
+		case PATH_TRACING_SHADINGMODELID_EYE:
 			Albedo = Payload.DiffuseColor + Payload.SubsurfaceColor + Payload.SpecularColor;
 			break;
-		case SHADINGMODELID_HAIR:
+		case PATH_TRACING_SHADINGMODELID_HAIR:
 		// SUBSTRATE_TODO: Add more cases here?
 #endif
 		default:
@@ -95,9 +95,9 @@ void AccumulateAlbedoNormal(inout FPathTracingPayload Payload, float3 PathThroug
 	float NormThroughput = Luminance(PathThroughput);
 	// figure out how much the specular lobes should contribute to the total (relative to the current PathRoughness)
 	// if we are above the MinRoughness, use the PathRoughness as are lower bound instead so that multiple bounces don't accumulate too much albedo
-	switch (Payload.ShadingModelID)
+	switch (Payload.PathTracingShadingModelID)
 	{
-		case SHADINGMODELID_SUBSTRATE:
+		case PATH_TRACING_SHADINGMODELID_BASIC:
 		{
 			const float Roughness0 = Payload.RoughnessData.x;
 			const float Roughness1 = Payload.RoughnessData.y;
@@ -107,7 +107,7 @@ void AccumulateAlbedoNormal(inout FPathTracingPayload Payload, float3 PathThroug
 			Normal += NormThroughput * Payload.WorldNormal;
 			break;
 		}
-		case SHADINGMODELID_SOLID_GLASS:
+		case PATH_TRACING_SHADINGMODELID_SOLID_GLASS:
 		{
 			const float Roughness0 = Payload.RoughnessData.x;
 			const float Roughness1 = Payload.RoughnessData.y;
@@ -118,7 +118,7 @@ void AccumulateAlbedoNormal(inout FPathTracingPayload Payload, float3 PathThroug
 			Normal += NormThroughput * Payload.WorldNormal;
 			break;
 		}
-		case SHADINGMODELID_TWOSIDED_FOLIAGE:
+		case PATH_TRACING_SHADINGMODELID_FOLIAGE:
 		{
 			const float Roughness0 = Payload.RoughnessData.x;
 			const float Roughness1 = Payload.RoughnessData.y;
@@ -128,7 +128,7 @@ void AccumulateAlbedoNormal(inout FPathTracingPayload Payload, float3 PathThroug
 			Normal += NormThroughput * Payload.WorldNormal;
 			break;
 		}
-		case SHADINGMODELID_THIN_TRANSLUCENT:
+		case PATH_TRACING_SHADINGMODELID_THIN_GLASS:
 		{
 			const float F0 = F0RGBToF0(Payload.SpecularColor);
 			const float3 GlassAlbedo = Payload.HasRefraction() ? lerp(Payload.GetTransmittanceColor(), 1.0, F0) : F0;
@@ -140,7 +140,7 @@ void AccumulateAlbedoNormal(inout FPathTracingPayload Payload, float3 PathThroug
 			Normal += NormThroughput * RoughnessWeight * Payload.WorldNormal;
 
 		}
-		case SHADINGMODELID_CLEAR_COAT:
+		case PATH_TRACING_SHADINGMODELID_CLEARCOAT:
 		{
 			const float Roughness0 = Payload.RoughnessData.x;
 			const float Roughness1 = Payload.RoughnessData.y;
@@ -151,7 +151,7 @@ void AccumulateAlbedoNormal(inout FPathTracingPayload Payload, float3 PathThroug
 			Normal += NormThroughput * Payload.WorldNormal;
 			break;
 		}
-		case SHADINGMODELID_EYE:
+		case PATH_TRACING_SHADINGMODELID_EYE:
 		{
 			const float RoughnessWeight = ROUGHNESS_WEIGHT(Payload.GetEyeRoughness());
 			Albedo += PathThroughput * (Payload.DiffuseColor + RoughnessWeight * Payload.SpecularColor + Payload.SubsurfaceColor);
@@ -184,9 +184,9 @@ float GetAverageRoughness(FPathTracingPayload Payload)
 	return 1.0;
 #else
 	float AvgRoughness = Payload.RoughnessData.x;
-	switch (Payload.ShadingModelID)
+	switch (Payload.PathTracingShadingModelID)
 	{
-		case SHADINGMODELID_SUBSTRATE:
+		case PATH_TRACING_SHADINGMODELID_BASIC:
 		{
 			AvgRoughness = max(AvgRoughness, Payload.RoughnessData.y);
 			// Only non-glass slabs contain diffuse
@@ -194,33 +194,33 @@ float GetAverageRoughness(FPathTracingPayload Payload)
 			AvgRoughness = lerp(1.0, AvgRoughness, DiffuseFraction);
 			break;
 		}
-		case SHADINGMODELID_SOLID_GLASS:
+		case PATH_TRACING_SHADINGMODELID_SOLID_GLASS:
 		{
 			// NOTE: this case does not contain a diffuse lobe yet
 			AvgRoughness = max(AvgRoughness, Payload.RoughnessData.y);
 			break;
 		}
-		case SHADINGMODELID_THIN_TRANSLUCENT:
+		case PATH_TRACING_SHADINGMODELID_THIN_GLASS:
 		{
 			AvgRoughness = max(AvgRoughness, Payload.RoughnessData.y);
 			float DiffuseFraction = LobeSelectionProb(Payload.DiffuseColor + Payload.FuzzAmount * Payload.FuzzColor, Payload.SpecularColor);
 			AvgRoughness = lerp(1.0, AvgRoughness, DiffuseFraction);
 			break;
 		}
-		case SHADINGMODELID_CLEAR_COAT:
+		case PATH_TRACING_SHADINGMODELID_CLEARCOAT:
 		{
 			AvgRoughness = max(AvgRoughness, Payload.RoughnessData.y);
 			float DiffuseFraction = F0RGBToMetallic(Payload.SpecularColor + Payload.RoughnessData.z * CLEAR_COAT_F0);
 			AvgRoughness = lerp(1.0, AvgRoughness, DiffuseFraction);
 			break;
 		}
-		case SHADINGMODELID_EYE:
+		case PATH_TRACING_SHADINGMODELID_EYE:
 		{
 			float DiffuseFraction = F0RGBToMetallic(Payload.SpecularColor);
 			AvgRoughness = lerp(1.0, AvgRoughness, DiffuseFraction);
 			break;
 		}
-		case SHADINGMODELID_MEDIUM:
+		case PATH_TRACING_SHADINGMODELID_MEDIUM:
 			return 1;
 		// SUBSTRATE_TODO: do we need special cases for the other slab types?
 	}
@@ -248,9 +248,9 @@ void AdjustPayloadAfterUnpack(inout FPathTracingPayload HitPayload, bool bApplyO
 		HitPayload.SpecularColor = HitPayload.SpecularColor * View.SpecularOverrideParameter.w + View.SpecularOverrideParameter.xyz;
 	}
 
-	switch (HitPayload.ShadingModelID)
+	switch (HitPayload.PathTracingShadingModelID)
 	{
-		case SHADINGMODELID_SUBSTRATE:
+		case PATH_TRACING_SHADINGMODELID_BASIC:
 		{
 			// NOTE: if radius is 0 the sss color will get flipped back into diffuse
 			HitPayload.SubsurfaceColor = HitPayload.DiffuseColor;
@@ -258,7 +258,7 @@ void AdjustPayloadAfterUnpack(inout FPathTracingPayload HitPayload, bool bApplyO
 			AdjustDiffuseSSSContribution(HitPayload.DiffuseColor, HitPayload.SubsurfaceColor, HitPayload.MeanFreePath);
 			break;
 		}
-		case SHADINGMODELID_TWOSIDED_FOLIAGE:
+		case PATH_TRACING_SHADINGMODELID_FOLIAGE:
 		{
 			float3 TransmissionColor = HitPayload.GetFoliageTransmissionColor();
 			if (bApplyOverrides)
@@ -269,7 +269,7 @@ void AdjustPayloadAfterUnpack(inout FPathTracingPayload HitPayload, bool bApplyO
 			HitPayload.SetFoliageTransmissionColor(TransmissionColor);
 			break;
 		}
-		case SHADINGMODELID_CLEAR_COAT:
+		case PATH_TRACING_SHADINGMODELID_CLEARCOAT:
 		{
 			if (bApplyOverrides)
 			{
@@ -277,7 +277,7 @@ void AdjustPayloadAfterUnpack(inout FPathTracingPayload HitPayload, bool bApplyO
 			}
 			break;
 		}
-		case SHADINGMODELID_EYE:
+		case PATH_TRACING_SHADINGMODELID_EYE:
 		{
 			// eye still needs a diffuse component for the iris center
 			float IrisMask = HitPayload.GetEyeIrisMask();
@@ -307,14 +307,14 @@ FSSSRandomWalkInfo GetMaterialSSSInfo(FPathTracingPayload Payload, float3 V_Worl
 		Result.Color = Payload.SubsurfaceColor;
 		Result.Radius = Payload.MeanFreePath;
 		Result.G = Payload.PhaseG;
-		if (Payload.ShadingModelID == SHADINGMODELID_SUBSTRATE)
+		if (Payload.PathTracingShadingModelID == PATH_TRACING_SHADINGMODELID_BASIC)
 		{
 			FSubstrateSingleSlabData Data = (FSubstrateSingleSlabData)0;
 			Data.PrepareBSDF(V_World, Payload, true);
 			Result.Weight = Payload.WeightV * Payload.BSDFOpacity * Data.DiffuseWeight;
 			Result.Prob = Data.LobePdf.x;
 		}
-		else if (Payload.ShadingModelID == SHADINGMODELID_EYE)
+		else if (Payload.PathTracingShadingModelID == PATH_TRACING_SHADINGMODELID_EYE)
 		{
 			const float NoV = saturate(dot(Payload.WorldNormal, V_World));
 			const float3 SpecE = ComputeGGXSpecEnergyTermsRGB(Payload.GetEyeRoughness(), NoV, Payload.SpecularColor).E;
@@ -350,25 +350,25 @@ FMaterialSample SampleMaterial(
 #if SIMPLIFIED_MATERIAL_SHADER
 	return Simplified_SampleMaterial(Payload, RandSample);
 #else
-	switch (Payload.ShadingModelID)
+	switch (Payload.PathTracingShadingModelID)
 	{
-		case SHADINGMODELID_UNLIT:
+		case PATH_TRACING_SHADINGMODELID_VOLUMETRIC:
 			return Volumetric_SampleMaterial(V_World, Payload, RandSample);
-		case SHADINGMODELID_SUBSTRATE:
+		case PATH_TRACING_SHADINGMODELID_BASIC:
 			return Substrate_SampleMaterial(V_World, Payload, RandSample);
-		case SHADINGMODELID_SOLID_GLASS:
+		case PATH_TRACING_SHADINGMODELID_SOLID_GLASS:
 			return SubstrateSolidGlass_SampleMaterial(V_World, Payload, RandSample);
-		case SHADINGMODELID_TWOSIDED_FOLIAGE:
+		case PATH_TRACING_SHADINGMODELID_FOLIAGE:
 			return SubstrateFoliage_SampleMaterial(V_World, Payload, RandSample);
-		case SHADINGMODELID_THIN_TRANSLUCENT:
+		case PATH_TRACING_SHADINGMODELID_THIN_GLASS:
 			return SubstrateThinGlass_SampleMaterial(V_World, Payload, RandSample);
-		case SHADINGMODELID_CLEAR_COAT:
+		case PATH_TRACING_SHADINGMODELID_CLEARCOAT:
 			return SubstrateClearCoat_SampleMaterial(V_World, Payload, RandSample);
-		case SHADINGMODELID_EYE:
+		case PATH_TRACING_SHADINGMODELID_EYE:
 			return Eye_SampleMaterial(V_World, Payload, RandSample);
-		case SHADINGMODELID_HAIR:
+		case PATH_TRACING_SHADINGMODELID_HAIR:
 			return Hair_SampleMaterial(V_World, Payload, RandSample);
-		case SHADINGMODELID_MEDIUM:
+		case PATH_TRACING_SHADINGMODELID_MEDIUM:
 			return Medium_SampleMaterial(V_World, Payload, RandSample);
 		default:
 			return Lambert_SampleMaterial(Payload, RandSample);
@@ -387,25 +387,25 @@ FMaterialEval EvalMaterial(
 #if SIMPLIFIED_MATERIAL_SHADER
 	return Simplified_EvalMaterial(L_World, Payload);
 #else
-	switch (Payload.ShadingModelID)
+	switch (Payload.PathTracingShadingModelID)
 	{
-		case SHADINGMODELID_UNLIT:
+		case PATH_TRACING_SHADINGMODELID_VOLUMETRIC:
 			return Volumetric_EvalMaterial(V_World, L_World, Payload, DiffuseSpecularScale);
-		case SHADINGMODELID_SUBSTRATE:
+		case PATH_TRACING_SHADINGMODELID_BASIC:
 			return Substrate_EvalMaterial(V_World, L_World, Payload, DiffuseSpecularScale);
-		case SHADINGMODELID_SOLID_GLASS:
+		case PATH_TRACING_SHADINGMODELID_SOLID_GLASS:
 			return SubstrateSolidGlass_EvalMaterial(V_World, L_World, Payload, DiffuseSpecularScale);
-		case SHADINGMODELID_TWOSIDED_FOLIAGE:
+		case PATH_TRACING_SHADINGMODELID_FOLIAGE:
 			return SubstrateFoliage_EvalMaterial(V_World, L_World, Payload, DiffuseSpecularScale);
-		case SHADINGMODELID_THIN_TRANSLUCENT:
+		case PATH_TRACING_SHADINGMODELID_THIN_GLASS:
 			return SubstrateThinGlass_EvalMaterial(V_World, L_World, Payload, DiffuseSpecularScale);
-		case SHADINGMODELID_CLEAR_COAT:
+		case PATH_TRACING_SHADINGMODELID_CLEARCOAT:
 			return SubstrateClearCoat_EvalMaterial(V_World, L_World, Payload, DiffuseSpecularScale);
-		case SHADINGMODELID_EYE:
+		case PATH_TRACING_SHADINGMODELID_EYE:
 			return Eye_EvalMaterial(V_World, L_World, Payload, DiffuseSpecularScale);
-		case SHADINGMODELID_HAIR:
+		case PATH_TRACING_SHADINGMODELID_HAIR:
 			return Hair_EvalMaterial(V_World, L_World, Payload, DiffuseSpecularScale);
-		case SHADINGMODELID_MEDIUM:
+		case PATH_TRACING_SHADINGMODELID_MEDIUM:
 			return Medium_EvalMaterial(V_World, L_World, Payload, DiffuseSpecularScale);
 		default:
 			return Lambert_EvalMaterial(L_World, Payload, DiffuseSpecularScale);
diff --git a/Engine/Shaders/Private/PathTracing/PathTracingCommon.ush b/Engine/Shaders/Private/PathTracing/PathTracingCommon.ush
--- a/Engine/Shaders/Private/PathTracing/PathTracingCommon.ush
+++ b/Engine/Shaders/Private/PathTracing/PathTracingCommon.ush
@@ -38,15 +38,20 @@
 #define PATH_TRACING_PAYLOAD_NORMAL_BITS  		32
 #define PATH_TRACING_PAYLOAD_TANGENT_ANGLE_BITS 16 						// Could be dropped to 12 or even 10 if we need more room in the future
 
-// shading model extensions beyond the default count
-// TODO: instead of re-using SHADINGMODELID_* values, make a distint enum for the path tracer
-#define SHADINGMODELID_MEDIUM  			(SHADINGMODELID_SUBSTRATE_TOON) // alias with unused value in the path tracer
-#define SHADINGMODELID_LAMBERT          (SHADINGMODELID_NUM    )        // pure lambert response (used for lighting after SSS random walk)
-#define SHADINGMODELID_SOLID_GLASS 		(SHADINGMODELID_NUM + 1) 		// use a sparse bit
 
-HLSL_STATIC_ASSERT(SHADINGMODELID_MEDIUM      <= SHADINGMODELID_MASK, "");
-HLSL_STATIC_ASSERT(SHADINGMODELID_LAMBERT     <= SHADINGMODELID_MASK, "");
-HLSL_STATIC_ASSERT(SHADINGMODELID_SOLID_GLASS <= SHADINGMODELID_MASK, "");
+// The Path Tracer does not directly follow the SHADINGMODELID_* defines because it classifies models slightly differently and maps legacy materials (pre-substrate) onto their substrate equivalents
+#define PATH_TRACING_SHADINGMODELID_BASIC     	0 // basic slab with SSS
+#define PATH_TRACING_SHADINGMODELID_FOLIAGE   	1 // basic slab with diffuse transmission lobe
+#define PATH_TRACING_SHADINGMODELID_CLEARCOAT 	2 // basic slab where the second specular lobe is layered over the first instead of mixed (legacy clearcoat)
+#define PATH_TRACING_SHADINGMODELID_SOLID_GLASS 3 // slab with refraction
+#define PATH_TRACING_SHADINGMODELID_THIN_GLASS 	4 // slab with thin-refraction
+#define PATH_TRACING_SHADINGMODELID_EYE         5 // special purpose model for eyes
+#define PATH_TRACING_SHADINGMODELID_HAIR        6 // special purpose model for hair
+#define PATH_TRACING_SHADINGMODELID_MEDIUM      7 // represent the result of scattering in a volume (phase function)
+#define PATH_TRACING_SHADINGMODELID_VOLUMETRIC  8 // surface shading model for particles that collects light from all directions
+#define PATH_TRACING_SHADINGMODELID_LAMBERT     9 // pure lambert response (used for lighting after SSS random walk)
+#define PATH_TRACING_SHADINGMODELID_MASK        0xf // 4 bits reserved
+
 
 #define PATH_TRACING_ABSORPTION_SCALE  0.01        // inverse distance at which we hit the transmission color (Beer's law)
 
@@ -249,7 +254,7 @@ struct FPackedPathTracingPayload : FMinimalPayload
 #endif
 
 	float UnpackBSDFOpacity() { return f16tof32(PackedData[0] & 0xFFFF); }
-	uint GetShadingModelID() { return (PackedData[0] >> 16) & SHADINGMODELID_MASK; }
+	uint GetPathTracingShadingModelID() { return (PackedData[0] >> 16) & PATH_TRACING_SHADINGMODELID_MASK; }
 	uint GetFlags() { return (PackedData[0] >> 20) & 0x1FF; }
 	void SetFlag(uint Flag) { PackedData[0] |= Flag << 20; }
 	void RemoveFlag(uint Flag) { PackedData[0] &= ~(Flag << 20); }
@@ -669,7 +674,7 @@ struct FPathTracingPayload : FMinimalPayload
 
 	float  BSDFOpacity;              // how much should we weigh down the bsdf by? (this also accounts for the probability of picking a particular lobe in multi-lobed materials)
 	float3 TransparencyColor;        // how much do we see straight through the surface?
-	uint   ShadingModelID;
+	uint   PathTracingShadingModelID;
 	uint   Flags;
 	uint   PrimitiveLightingChannelMask;
 
@@ -731,7 +736,7 @@ struct FPathTracingPayload : FMinimalPayload
 	float3 GetSubsurfaceColor() { return SubsurfaceColor; }
 	bool IsMaterialTransmissive()
 	{
-		return ShadingModelID == SHADINGMODELID_TWOSIDED_FOLIAGE;
+		return PathTracingShadingModelID == PATH_TRACING_SHADINGMODELID_FOLIAGE;
 	}
 #else // IS_PAYLOAD_ENABLED(RT_PAYLOAD_TYPE_GPULIGHTMASS)
 
@@ -811,29 +816,44 @@ struct FPathTracingPayload : FMinimalPayload
 
 	bool IsMaterialSolidGlass()
 	{
-		return ShadingModelID == SHADINGMODELID_SOLID_GLASS;
+		return PathTracingShadingModelID == PATH_TRACING_SHADINGMODELID_SOLID_GLASS;
 	}
 
 	bool IsMaterialThinGlass()
 	{
-		return ShadingModelID == SHADINGMODELID_THIN_TRANSLUCENT && HasRefraction();
+		return PathTracingShadingModelID == PATH_TRACING_SHADINGMODELID_THIN_GLASS && HasRefraction();
 	}
 
 	bool IsMaterialTransmissive()
 	{
-		return ShadingModelID == SHADINGMODELID_TWOSIDED_FOLIAGE ||
-			   ShadingModelID == SHADINGMODELID_HAIR ||
-			   ShadingModelID == SHADINGMODELID_MEDIUM ||
-			   ShadingModelID == SHADINGMODELID_UNLIT || // Volumetric Directional/Non-Directional
-			   IsMaterialSolidGlass() ||
-			   IsMaterialThinGlass();
+		return
+			PathTracingShadingModelID == PATH_TRACING_SHADINGMODELID_FOLIAGE ||
+			PathTracingShadingModelID == PATH_TRACING_SHADINGMODELID_HAIR ||
+			PathTracingShadingModelID == PATH_TRACING_SHADINGMODELID_MEDIUM ||
+			PathTracingShadingModelID == PATH_TRACING_SHADINGMODELID_VOLUMETRIC ||
+			IsMaterialSolidGlass() ||
+			IsMaterialThinGlass();
 	}
 
 	bool IsSubsurfaceMaterial()
 	{
-		return ShadingModelID == SHADINGMODELID_SUBSTRATE ||
-		       ShadingModelID == SHADINGMODELID_EYE;
+		return
+			PathTracingShadingModelID == PATH_TRACING_SHADINGMODELID_BASIC ||
+			PathTracingShadingModelID == PATH_TRACING_SHADINGMODELID_EYE;
 	}
+
+	// Are the Fuzz* parameters used for the fuzz lobe?
+	bool HasFuzzLobe()
+	{
+		return
+			PathTracingShadingModelID == PATH_TRACING_SHADINGMODELID_BASIC ||
+			PathTracingShadingModelID == PATH_TRACING_SHADINGMODELID_FOLIAGE ||
+			PathTracingShadingModelID == PATH_TRACING_SHADINGMODELID_CLEARCOAT ||
+			PathTracingShadingModelID == PATH_TRACING_SHADINGMODELID_SOLID_GLASS ||
+			PathTracingShadingModelID == PATH_TRACING_SHADINGMODELID_THIN_GLASS ||
+			PathTracingShadingModelID == PATH_TRACING_SHADINGMODELID_EYE;
+	}
+
 #endif // IS_PAYLOAD_ENABLED(RT_PAYLOAD_TYPE_GPULIGHTMASS)
 };
 
@@ -876,11 +896,11 @@ FPackedPathTracingPayload PackPathTracingPayload(FPathTracingPayload Input)
 	FPackedPathTracingPayload Output = (FPackedPathTracingPayload)0;
 	Output.HitT = Input.HitT;
 
-	Output.PackedData[0]  = f32tof16(Input.BSDFOpacity);                                      // 16 bits
-	Output.PackedData[0] |= (Input.ShadingModelID & SHADINGMODELID_MASK) << 16;               // 4  bits
-	Output.PackedData[0] |= (Input.Flags & 0x1FF) << 20;                                      // 9  bits
-	Output.PackedData[0] |= (Input.PrimitiveLightingChannelMask & 0x7) << 29;                 // 3  bits
-	                                                                                          // 32 bits total
+	Output.PackedData[0]  = f32tof16(Input.BSDFOpacity);                                      			// 16 bits
+	Output.PackedData[0] |= (Input.PathTracingShadingModelID & PATH_TRACING_SHADINGMODELID_MASK) << 16;	// 4  bits
+	Output.PackedData[0] |= (Input.Flags & 0x1FF) << 20;                                      			// 9  bits
+	Output.PackedData[0] |= (Input.PrimitiveLightingChannelMask & 0x7) << 29;                 			// 3  bits
+	                                                                                          			// 32 bits total
 #if IS_PAYLOAD_ENABLED(RT_PAYLOAD_TYPE_GPULIGHTMASS)
 	Output.PackedData[1] = PayloadEncodeUnitVector(Input.WorldGeoNormal);
 	Output.PackedData[2] = PayloadEncodeUnitVector(Input.WorldNormal);
@@ -940,7 +960,7 @@ FPathTracingPayload UnpackPathTracingPayload(FPackedPathTracingPayload Input, FR
 	Output.TranslatedWorldPos = Ray.Origin + Output.HitT * Ray.Direction;
 
 	Output.BSDFOpacity                  = Input.UnpackBSDFOpacity();
-	Output.ShadingModelID               = Input.GetShadingModelID();
+	Output.PathTracingShadingModelID    = Input.GetPathTracingShadingModelID();
 	Output.Flags                        = Input.GetFlags();
 	Output.PrimitiveLightingChannelMask = Input.GetPrimitiveLightingChannelMask();
 #if IS_PAYLOAD_ENABLED(RT_PAYLOAD_TYPE_GPULIGHTMASS)
diff --git a/Engine/Shaders/Private/PathTracing/PathTracingCore.ush b/Engine/Shaders/Private/PathTracing/PathTracingCore.ush
--- a/Engine/Shaders/Private/PathTracing/PathTracingCore.ush
+++ b/Engine/Shaders/Private/PathTracing/PathTracingCore.ush
@@ -667,7 +667,7 @@ FPackedPathTracingPayload TraceTransparentRay(inout FPathState PathState, int Bo
 				FPathTracingPayload LightHitPayload = (FPathTracingPayload)0; // clear all fields
 				LightHitPayload.HitT = LightHit.HitT;
 				LightHitPayload.Radiance = LightHit.Radiance * SceneLights[LightId].CameraColorScale;
-				LightHitPayload.ShadingModelID = SHADINGMODELID_UNLIT;
+				LightHitPayload.PathTracingShadingModelID = PATH_TRACING_SHADINGMODELID_LAMBERT;
 				// light hits are fully transparent to match how the light loop works (the light surface does not cast a shadow)
 				LightHitPayload.BSDFOpacity = 1.0;
 				LightHitPayload.TransparencyColor = LightHit.HitT == RAY_DEFAULT_T_MAX ? 1.0 - BackgroundAlpha : 1.0;
@@ -807,7 +807,7 @@ FPackedPathTracingPayload TraceTransparentRay(inout FPathState PathState, int Bo
 					FPathTracingPayload Result = (FPathTracingPayload)0; // clear all fields
 					Result.HitT = CloudPayload.HitT;
 					Result.TranslatedWorldPos = PathState.Ray.Origin + PathState.Ray.Direction * CloudPayload.HitT;
-					Result.ShadingModelID = SHADINGMODELID_MEDIUM;
+					Result.PathTracingShadingModelID = PATH_TRACING_SHADINGMODELID_MEDIUM;
 					Result.BSDFOpacity = 1.0;
 					Result.PrimitiveLightingChannelMask = 7;
 					Result.SetFrontFace();
@@ -1568,7 +1568,7 @@ bool ProcessSubsurfaceRandomWalk(inout FPathTracingPayload Payload, inout float3
 			Payload.WorldSmoothNormal = SignFlip * ProbeResult.WorldSmoothNormal;
 			Payload.WorldGeoNormal    = SignFlip * ProbeResult.WorldGeoNormal;
 			
-			Payload.ShadingModelID = SHADINGMODELID_LAMBERT; // diffuse shading
+			Payload.PathTracingShadingModelID = PATH_TRACING_SHADINGMODELID_LAMBERT; // diffuse shading
 			Payload.BSDFOpacity = 1;
 			Payload.SetBaseColor(1.0);
 			Payload.TransparencyColor = 0;
@@ -1864,7 +1864,7 @@ bool PathTracingKernel(inout FPathState PathState, int Bounce)
 	}
 
 	// Don't treat internal volume scattering as a volume sample (so its not affected by those options in the lights)
-	const bool bIsVolumeSample = Payload.ShadingModelID == SHADINGMODELID_MEDIUM && !any(PathState.GetSigmaT() > 0);
+	const bool bIsVolumeSample = Payload.PathTracingShadingModelID == PATH_TRACING_SHADINGMODELID_MEDIUM && !any(PathState.GetSigmaT() > 0);
 
 	FLightLoopCount LightLoopCount = LightGridLookup(Payload.TranslatedWorldPos);
 	// If this is a volume sample, we have not forced the renderer to ignore volume segment sampling _and_ we aren't in an internal medium which does not use volume segments
diff --git a/Engine/Shaders/Private/PathTracing/PathTracingDebug.usf b/Engine/Shaders/Private/PathTracing/PathTracingDebug.usf
--- a/Engine/Shaders/Private/PathTracing/PathTracingDebug.usf
+++ b/Engine/Shaders/Private/PathTracing/PathTracingDebug.usf
@@ -523,7 +523,7 @@ FPackedPathTracingPayload TraceTransparentRay(inout FPathState PathState, int Bo
 				FPathTracingPayload LightHitPayload = (FPathTracingPayload)0; // clear all fields
 				LightHitPayload.HitT = LightHit.HitT;
 				LightHitPayload.Radiance = LightHit.Radiance * SceneLights[LightId].CameraColorScale;
-				LightHitPayload.ShadingModelID = SHADINGMODELID_UNLIT;
+				LightHitPayload.PathTracingShadingModelID = PATH_TRACING_SHADINGMODELID_LAMBERT;
 				// light hits are fully transparent to match how the light loop works (the light surface does not cast a shadow)
 				LightHitPayload.BSDFOpacity = 1.0;
 				//LightHitPayload.TransparencyColor = LightHit.HitT == RAY_DEFAULT_T_MAX ? 1.0 - BackgroundAlpha : 1.0;
@@ -914,6 +914,26 @@ bool PathTracingDebugCore(inout FPathState PathState, int Bounce)
 
 #else // PATH_TRACER_PRIMARY_RAYS
 
+// for debugging and to visualize
+float3 GetPathTracingShadingModelColor(uint PathTracingShadingModelID)
+{
+	switch (PathTracingShadingModelID)
+	{
+		case PATH_TRACING_SHADINGMODELID_BASIC: 		return float3(0.1f, 1.0f, 0.1f); // Green
+		case PATH_TRACING_SHADINGMODELID_FOLIAGE: 		return float3(0.2f, 0.2f, 0.8f); // Cyan
+		case PATH_TRACING_SHADINGMODELID_CLEARCOAT: 	return float3(0.1f, 0.4f, 0.4f); // Brown
+		case PATH_TRACING_SHADINGMODELID_SOLID_GLASS: 	return float3(1.0f, 0.1f, 0.1f); // Red
+		case PATH_TRACING_SHADINGMODELID_THIN_GLASS: 	return float3(1.0f, 0.8f, 0.3f);
+		case PATH_TRACING_SHADINGMODELID_EYE: 			return float3(0.3f, 1.0f, 1.0f);
+		case PATH_TRACING_SHADINGMODELID_HAIR: 			return float3(0.6f, 0.1f, 0.5f);
+		case PATH_TRACING_SHADINGMODELID_MEDIUM:		return float3(0.6f, 0.4f, 0.1f); // Brown
+		case PATH_TRACING_SHADINGMODELID_VOLUMETRIC:	return float3(0.2f, 0.6f, 0.5f); // Cyan
+		case PATH_TRACING_SHADINGMODELID_LAMBERT: 		return float3(0.1f, 0.1f, 0.2f); // Dark Blue
+		default: return float3(1.0f, 1.0f, 1.0f); // White
+	}
+}
+
+
 FPathTracingResult PathTracingDebugCore(FRayDesc Ray, uint2 LaunchIndex)
 {
 	FPathTracingResult Result = (FPathTracingResult) 0;
@@ -1064,7 +1084,7 @@ FPathTracingResult PathTracingDebugCore(FRayDesc Ray, uint2 LaunchIndex)
 				case PATH_TRACER_DEBUG_VIZ_ANISOTROPY:					Color = AnisoDebugColor(HitPayload.Anisotropy); break;
 				case PATH_TRACER_DEBUG_VIZ_TRANSPARENCY_COLOR:			Color = HitPayload.TransparencyColor; break;
 				case PATH_TRACER_DEBUG_VIZ_IOR:							Color = HitPayload.Ior; break;
-				case PATH_TRACER_DEBUG_VIZ_SHADING_MODEL:				Color = GetShadingModelColor(HitPayload.ShadingModelID); break;
+				case PATH_TRACER_DEBUG_VIZ_SHADING_MODEL:				Color = GetPathTracingShadingModelColor(HitPayload.PathTracingShadingModelID); break;
 				case PATH_TRACER_DEBUG_VIZ_LIGHTING_CHANNEL_MASK:		Color = GetLightingChannelMaskDebugColor(HitPayload.PrimitiveLightingChannelMask); break;
 				case PATH_TRACER_DEBUG_VIZ_WORLD_POSITION:
 				{
@@ -1126,16 +1146,16 @@ FPathTracingResult PathTracingDebugCore(FRayDesc Ray, uint2 LaunchIndex)
 					break;
 				}
 				case PATH_TRACER_DEBUG_VIZ_BSDF_OPACITY:    Color = HitPayload.BSDFOpacity; break;
-				case PATH_TRACER_DEBUG_VIZ_FUZZ_COLOR:     	Color = (HitPayload.ShadingModelID == SHADINGMODELID_SUBSTRATE) ? HitPayload.FuzzColor     : 0.0; break;
-				case PATH_TRACER_DEBUG_VIZ_FUZZ_ROUGHNESS:	Color = (HitPayload.ShadingModelID == SHADINGMODELID_SUBSTRATE) ? HitPayload.FuzzRoughness : 0.0; break;
-				case PATH_TRACER_DEBUG_VIZ_FUZZ_AMOUNT:		Color = (HitPayload.ShadingModelID == SHADINGMODELID_SUBSTRATE) ? HitPayload.FuzzAmount    : 0.0; break;
+				case PATH_TRACER_DEBUG_VIZ_FUZZ_COLOR:     	Color = HitPayload.HasFuzzLobe() ? HitPayload.FuzzColor     : 0.0; break;
+				case PATH_TRACER_DEBUG_VIZ_FUZZ_ROUGHNESS:	Color = HitPayload.HasFuzzLobe() ? HitPayload.FuzzRoughness : 0.0; break;
+				case PATH_TRACER_DEBUG_VIZ_FUZZ_AMOUNT:		Color = HitPayload.HasFuzzLobe() ? HitPayload.FuzzAmount    : 0.0; break;
 				case PATH_TRACER_DEBUG_VIZ_SUBSTRATE_WEIGHT_V: 			Color = HitPayload.WeightV; 			break;
 				case PATH_TRACER_DEBUG_VIZ_SUBSTRATE_TRANSMITTANCE_N: 	Color = HitPayload.TransmittanceN; 		break;
 				case PATH_TRACER_DEBUG_VIZ_SUBSTRATE_COVERAGE_ABOVE_N: 	Color = HitPayload.CoverageAboveAlongN; break;
 				case PATH_TRACER_DEBUG_VIZ_SUBSTRATE_F90:				Color = HitPayload.SpecularEdgeColor;   break;
-				case PATH_TRACER_DEBUG_VIZ_EYE_IRIS_MASK:      Color = (HitPayload.ShadingModelID == SHADINGMODELID_EYE) ?             HitPayload.GetEyeIrisMask()      : 0.0; break;
-				case PATH_TRACER_DEBUG_VIZ_EYE_IRIS_NORMAL:    Color = (HitPayload.ShadingModelID == SHADINGMODELID_EYE) ? 0.5 + 0.5 * HitPayload.GetEyeIrisNormal()    : 0.0; break;
-				case PATH_TRACER_DEBUG_VIZ_EYE_CAUSTIC_NORMAL: Color = (HitPayload.ShadingModelID == SHADINGMODELID_EYE) ? 0.5 + 0.5 * HitPayload.GetEyeCausticNormal() : 0.0; break;
+				case PATH_TRACER_DEBUG_VIZ_EYE_IRIS_MASK:      Color = (HitPayload.PathTracingShadingModelID == PATH_TRACING_SHADINGMODELID_EYE) ?             HitPayload.GetEyeIrisMask()      : 0.0; break;
+				case PATH_TRACER_DEBUG_VIZ_EYE_IRIS_NORMAL:    Color = (HitPayload.PathTracingShadingModelID == PATH_TRACING_SHADINGMODELID_EYE) ? 0.5 + 0.5 * HitPayload.GetEyeIrisNormal()    : 0.0; break;
+				case PATH_TRACER_DEBUG_VIZ_EYE_CAUSTIC_NORMAL: Color = (HitPayload.PathTracingShadingModelID == PATH_TRACING_SHADINGMODELID_EYE) ? 0.5 + 0.5 * HitPayload.GetEyeCausticNormal() : 0.0; break;
 				case PATH_TRACER_DEBUG_VIZ_BACKFACING_NORMAL:
 				{
 					Color = FakeShading;
@@ -1156,7 +1176,7 @@ FPathTracingResult PathTracingDebugCore(FRayDesc Ray, uint2 LaunchIndex)
 				}
 			}
 
-			Result.Radiance.rgb += PathThroughput * Color;
+			Result.Radiance.rgb += PathThroughput * Color * HitPayload.BSDFOpacity;
 			Result.Radiance.a += Luminance(PathThroughput) * ExtraAlpha;
 		}
 
diff --git a/Engine/Shaders/Private/PathTracing/PathTracingDefaultHitShader.usf b/Engine/Shaders/Private/PathTracing/PathTracingDefaultHitShader.usf
--- a/Engine/Shaders/Private/PathTracing/PathTracingDefaultHitShader.usf
+++ b/Engine/Shaders/Private/PathTracing/PathTracingDefaultHitShader.usf
@@ -11,7 +11,7 @@ RAY_TRACING_ENTRY_CLOSEST_HIT(PathTracingDefaultOpaqueCHS,
 {
 	// mark hit as solid black
 	FPathTracingPayload Payload = (FPathTracingPayload)0;
-	Payload.ShadingModelID = SHADINGMODELID_UNLIT;
+	Payload.PathTracingShadingModelID = PATH_TRACING_SHADINGMODELID_LAMBERT;
 	Payload.BSDFOpacity = 1.0;
 	PackedPayload = PackPathTracingPayload(Payload);
 	PackedPayload.HitT = RayTCurrent();
@@ -23,7 +23,7 @@ RAY_TRACING_ENTRY_CLOSEST_HIT(PathTracingDefaultHiddenCHS,
 {
 	// mark hit as fully transparent
 	FPathTracingPayload Payload = (FPathTracingPayload)0;
-	Payload.ShadingModelID = SHADINGMODELID_UNLIT;
+	Payload.PathTracingShadingModelID = PATH_TRACING_SHADINGMODELID_LAMBERT;
 	Payload.TransparencyColor = 1.0;
 	PackedPayload = PackPathTracingPayload(Payload);
 	PackedPayload.HitT = RayTCurrent();
diff --git a/Engine/Shaders/Private/PathTracing/PathTracingGPULightmassMaterialHitShader.usf b/Engine/Shaders/Private/PathTracing/PathTracingGPULightmassMaterialHitShader.usf
--- a/Engine/Shaders/Private/PathTracing/PathTracingGPULightmassMaterialHitShader.usf
+++ b/Engine/Shaders/Private/PathTracing/PathTracingGPULightmassMaterialHitShader.usf
@@ -120,7 +120,7 @@ RAY_TRACING_ENTRY_CLOSEST_HIT(GPULightmassMaterialCHS,
 	Payload.BSDFOpacity = Coverage;
 	Payload.BaseColor = BaseColorPostCoverage;
 	Payload.SubsurfaceColor = 0;							// SUBSTRATE_TODO: Figure out how to support old two-sided foliage approximation
-	Payload.ShadingModelID = any(BaseColorPostCoverage > 0.0) ? SHADINGMODELID_DEFAULT_LIT : SHADINGMODELID_UNLIT;
+	Payload.PathTracingShadingModelID = PATH_TRACING_SHADINGMODELID_BASIC;
 	Payload.TransparencyColor = TransparencyColor;
 
 	const uint PrimitiveFlags = GetPrimitiveData(MaterialParameters.PrimitiveId).Flags;
@@ -179,7 +179,7 @@ RAY_TRACING_ENTRY_CLOSEST_HIT(GPULightmassMaterialCHS,
 	/**
 	 * Set common material attributes for both full and simplified materials
 	 **/
-	Payload.ShadingModelID = GetMaterialShadingModel(PixelMaterialInputs);
+	uint ShadingModelID = GetMaterialShadingModel(PixelMaterialInputs);
 
 #if MATERIALBLENDING_ALPHACOMPOSITE
 	Payload.BSDFOpacity = 1.0;
@@ -189,7 +189,7 @@ RAY_TRACING_ENTRY_CLOSEST_HIT(GPULightmassMaterialCHS,
 	Payload.TransparencyColor = 1.0 - GetMaterialOpacity(PixelMaterialInputs);
 	Payload.SetHoldout();
 	HLSL_STATIC_ASSERT(MATERIAL_SHADINGMODEL_UNLIT == 1, "Alpha holdout blend mode requires unlit shading model");
-	Payload.ShadingModelID = SHADINGMODELID_UNLIT;
+	ShadingModelID = SHADINGMODELID_UNLIT;
 #elif MATERIALBLENDING_TRANSLUCENT
 	Payload.BSDFOpacity = GetMaterialOpacity(PixelMaterialInputs);
 	Payload.TransparencyColor = 1.0 - Payload.BSDFOpacity;
@@ -201,7 +201,7 @@ RAY_TRACING_ENTRY_CLOSEST_HIT(GPULightmassMaterialCHS,
 	Payload.TransparencyColor = GetMaterialEmissive(PixelMaterialInputs);
 	HLSL_STATIC_ASSERT(MATERIAL_SHADINGMODEL_UNLIT == 1 ||
 	                   MATERIAL_SHADINGMODEL_DEFAULT_LIT, "Modulate blend mode requires unlit shading model");
-	Payload.ShadingModelID = SHADINGMODELID_UNLIT;
+	ShadingModelID = SHADINGMODELID_UNLIT;
 #elif MATERIALBLENDING_MASKED && MATERIAL_DITHER_OPACITY_MASK
 	// dithering emulates real transparency, so switch to translucent
 	// NOTE: the raster path technically takes into account the opacity mask clip value, so the effective transparency should be:
@@ -280,14 +280,15 @@ RAY_TRACING_ENTRY_CLOSEST_HIT(GPULightmassMaterialCHS,
 
 	// Anything unsupported will be forced to default lit
 #if MATERIAL_SHADINGMODEL_TWOSIDED_FOLIAGE
-	if (Payload.ShadingModelID == SHADINGMODELID_TWOSIDED_FOLIAGE)
+	if (ShadingModelID == SHADINGMODELID_TWOSIDED_FOLIAGE)
 	{
+		Payload.PathTracingShadingModelID = PATH_TRACING_SHADINGMODELID_FOLIAGE;
 		Payload.SubsurfaceColor = GetMaterialSubsurfaceData(PixelMaterialInputs).rgb;
 	}
 	else
 #endif
 #if MATERIAL_SHADINGMODEL_THIN_TRANSLUCENT
-		if (Payload.ShadingModelID == SHADINGMODELID_THIN_TRANSLUCENT)
+		if (ShadingModelID == SHADINGMODELID_THIN_TRANSLUCENT)
 		{
 			const float3 Transmission = GetThinTranslucentMaterialOutput0(MaterialParameters);
 			const float3 V = WorldRayDirection();
@@ -295,21 +296,23 @@ RAY_TRACING_ENTRY_CLOSEST_HIT(GPULightmassMaterialCHS,
 			const float VoN = abs(dot(V, N));
 			// simplified logic with no bending at the interface and no fresnel
 			Payload.TransparencyColor *= pow(Transmission, 1.0 / VoN);
+			Payload.PathTracingShadingModelID = PATH_TRACING_SHADINGMODELID_BASIC;
 		}
 		else
 #endif
 #if MATERIAL_SHADINGMODEL_UNLIT
-			if (Payload.ShadingModelID == SHADINGMODELID_UNLIT)
+			if (ShadingModelID == SHADINGMODELID_UNLIT)
 			{
+				Payload.PathTracingShadingModelID = PATH_TRACING_SHADINGMODELID_LAMBERT;
 			}
 			else
 #endif
 			{
 				// Redirect translucent blending mode to thin translucent
 #if MATERIALBLENDING_TRANSLUCENT
-				Payload.ShadingModelID = SHADINGMODELID_THIN_TRANSLUCENT;
+				Payload.PathTracingShadingModelID = PATH_TRACING_SHADINGMODELID_THIN_GLASS;
 #else
-				Payload.ShadingModelID = SHADINGMODELID_DEFAULT_LIT;
+				Payload.PathTracingShadingModelID = PATH_TRACING_SHADINGMODELID_BASIC;
 #endif
 			}
 	Payload.Radiance = GetMaterialEmissive(PixelMaterialInputs);
diff --git a/Engine/Shaders/Private/PathTracing/PathTracingMaterialHitShader.usf b/Engine/Shaders/Private/PathTracing/PathTracingMaterialHitShader.usf
--- a/Engine/Shaders/Private/PathTracing/PathTracingMaterialHitShader.usf
+++ b/Engine/Shaders/Private/PathTracing/PathTracingMaterialHitShader.usf
@@ -277,9 +277,9 @@ RAY_TRACING_ENTRY_CLOSEST_HIT(PathTracingMaterialCHS,
 	 **/
 
 #if PATH_TRACER_USE_VOLUMETRIC_SHADING_MODEL
-	Payload.ShadingModelID = SHADINGMODELID_UNLIT;
+	uint ShadingModelID = SHADINGMODELID_UNLIT;
 #else
-	Payload.ShadingModelID = GetMaterialShadingModel(PixelMaterialInputs);
+	uint ShadingModelID = GetMaterialShadingModel(PixelMaterialInputs);
 #endif
 
 #if MATERIALBLENDING_ALPHACOMPOSITE
@@ -290,7 +290,7 @@ RAY_TRACING_ENTRY_CLOSEST_HIT(PathTracingMaterialCHS,
 	Payload.TransparencyColor = 1.0 - GetMaterialOpacity(PixelMaterialInputs);
 	Payload.SetHoldout();
 	HLSL_STATIC_ASSERT(MATERIAL_SHADINGMODEL_UNLIT == 1, "Alpha holdout blend mode requires unlit shading model");
-	Payload.ShadingModelID = SHADINGMODELID_UNLIT;
+	ShadingModelID = SHADINGMODELID_UNLIT;
 #elif MATERIALBLENDING_TRANSLUCENT
 	Payload.BSDFOpacity = GetMaterialOpacity(PixelMaterialInputs);
 	Payload.TransparencyColor = 1.0 - Payload.BSDFOpacity;
@@ -303,7 +303,7 @@ RAY_TRACING_ENTRY_CLOSEST_HIT(PathTracingMaterialCHS,
 	// NOTE: in the case of substrate materials we could have DEFAULT_LIT here
 	HLSL_STATIC_ASSERT(MATERIAL_SHADINGMODEL_UNLIT == 1 ||
 	                   MATERIAL_SHADINGMODEL_DEFAULT_LIT == 1, "Modulate blend mode requires unlit shading model");
-	Payload.ShadingModelID = SHADINGMODELID_UNLIT;
+	ShadingModelID = SHADINGMODELID_UNLIT;
 #elif MATERIALBLENDING_MASKED && MATERIAL_DITHER_OPACITY_MASK
 	// dithering emulates real transparency, so switch to translucent
 	// NOTE: the raster path technically takes into account the opacity mask clip value, so the effective transparency should be:
@@ -358,7 +358,7 @@ RAY_TRACING_ENTRY_CLOSEST_HIT(PathTracingMaterialCHS,
 
 #if SUBSTRATE_ENABLED && !USE_LEGACY_CODEPATH_FOR_SUBSTRATE
 	// unless proven otherwise, we are a basic unlit material
-	Payload.ShadingModelID = SHADINGMODELID_UNLIT;
+	Payload.PathTracingShadingModelID = PATH_TRACING_SHADINGMODELID_LAMBERT;
 
 	float BaseBSDFOpacity = 1; //Payload.BSDFOpacity;
 
@@ -545,7 +545,7 @@ RAY_TRACING_ENTRY_CLOSEST_HIT(PathTracingMaterialCHS,
 							Payload.WeightV               = BSDF.LuminanceWeightV;
 							Payload.TransmittanceN        = BSDF.TransmittanceAboveAlongN;
 							Payload.CoverageAboveAlongN   = BSDF.CoverageAboveAlongN;
-							Payload.ShadingModelID    	  = SHADINGMODELID_SUBSTRATE;
+							Payload.PathTracingShadingModelID = PATH_TRACING_SHADINGMODELID_BASIC;
 							Payload.Ior                   = 0.0;
 							Payload.DiffuseColor          = DiffuseColor;
 							Payload.SpecularColor         = F0;
@@ -553,7 +553,7 @@ RAY_TRACING_ENTRY_CLOSEST_HIT(PathTracingMaterialCHS,
 							Payload.RoughnessData         = RoughnessData;
 							if (bUseClearCoat)
 							{
-								Payload.ShadingModelID = SHADINGMODELID_CLEAR_COAT;
+								Payload.PathTracingShadingModelID = PATH_TRACING_SHADINGMODELID_CLEARCOAT;
 							#if CLEAR_COAT_BOTTOM_NORMAL
 								Payload.SetBottomCoatNormal(SubstrateGetSimpleCoatBottomNormal(UnpackHaziness(SLAB_HAZINESS(BSDF)).BottomNormalOct, N));
 							#endif
@@ -588,7 +588,7 @@ RAY_TRACING_ENTRY_CLOSEST_HIT(PathTracingMaterialCHS,
 							else if (BSDF_GETSSSTYPE(BSDF) == SSS_TYPE_TWO_SIDED_WRAP)
 							{
 								// Legacy support for SHADINGMODELID_TWOSIDED_FOLIAGE
-								Payload.ShadingModelID = SHADINGMODELID_TWOSIDED_FOLIAGE;
+								Payload.PathTracingShadingModelID = PATH_TRACING_SHADINGMODELID_FOLIAGE;
 								Payload.SetFoliageTransmissionColor(LegacySubsurfaceColor);
 							}
 							else
@@ -644,9 +644,9 @@ RAY_TRACING_ENTRY_CLOSEST_HIT(PathTracingMaterialCHS,
 							Payload.TransmittanceN        = BSDF.TransmittanceAboveAlongN;
 							Payload.CoverageAboveAlongN   = BSDF.CoverageAboveAlongN;
 							#if MATERIAL_ISTHINSURFACE
-							Payload.ShadingModelID = SHADINGMODELID_THIN_TRANSLUCENT;
+							Payload.PathTracingShadingModelID = PATH_TRACING_SHADINGMODELID_THIN_GLASS;
 							#else
-							Payload.ShadingModelID = SHADINGMODELID_SOLID_GLASS;
+							Payload.PathTracingShadingModelID = PATH_TRACING_SHADINGMODELID_SOLID_GLASS;
 							// Albedo inversion for ambient medium
 							// NOTE: we could technically also include IOR and Phase function - but this basic inversion does most of the work
 							float3 Color = DiffuseColor;
@@ -724,7 +724,7 @@ RAY_TRACING_ENTRY_CLOSEST_HIT(PathTracingMaterialCHS,
 					Payload.SetEyeCausticNormal(CausticNormal);
 					Payload.SetEyeIrisMask(IrisMask);
 					Payload.SetEyeIrisNormal(normalize(EYE_IRISNORMAL(BSDF)));
-					Payload.ShadingModelID = SHADINGMODELID_EYE;
+					Payload.PathTracingShadingModelID = PATH_TRACING_SHADINGMODELID_EYE;
 				#if SUBSTRATE_PATH_TRACING_SPECULAR_PROFILE
 					Payload.SpecularProfileId = 0;
 				#endif
@@ -749,7 +749,7 @@ RAY_TRACING_ENTRY_CLOSEST_HIT(PathTracingMaterialCHS,
 					#else
 					Payload.SetHairPrimitiveUV(0.5);
 					#endif
-					Payload.ShadingModelID = SHADINGMODELID_HAIR;
+					Payload.PathTracingShadingModelID = PATH_TRACNIG_SHADINGMODELID_HAIR;
 					#if SUBSTRATE_PATH_TRACING_SPECULAR_PROFILE
 					Payload.SpecularProfileId = 0;
 					#endif
@@ -763,7 +763,7 @@ RAY_TRACING_ENTRY_CLOSEST_HIT(PathTracingMaterialCHS,
 					SlabSampler.Accept(1); // dummy, avoids special case below
 					Payload.WorldNormal = SubstratePixelHeader.SharedLocalBases.Normals[BSDF_GETSHAREDLOCALBASISID(BSDF)];
 					Payload.BSDFOpacity = 1;
-					Payload.ShadingModelID = SHADINGMODELID_SOLID_GLASS;
+					Payload.PathTracingShadingModelID = PATH_TRACING_SHADINGMODELID_SOLID_GLASS;
 					float3 BaseColor = SLW_BASECOLOR(BSDF);
 					float Metallic = SLW_METALLIC(BSDF);
 					float Specular = SLW_SPECULAR(BSDF);
@@ -809,7 +809,7 @@ RAY_TRACING_ENTRY_CLOSEST_HIT(PathTracingMaterialCHS,
 					Payload.WeightV               = BSDF.LuminanceWeightV;
 					Payload.TransmittanceN        = BSDF.TransmittanceAboveAlongN;
 					Payload.CoverageAboveAlongN   = BSDF.CoverageAboveAlongN;
-					Payload.ShadingModelID    	  = SHADINGMODELID_SUBSTRATE;
+					Payload.PathTracingShadingModelID = PATH_TRACING_SHADINGMODELID_BASIC;
 					Payload.Ior                   = 0.0f;
 					Payload.DiffuseColor          = float3(1,0,1); // TOON_TODO
 					Payload.SpecularColor         = 0.0f;
@@ -891,7 +891,7 @@ RAY_TRACING_ENTRY_CLOSEST_HIT(PathTracingMaterialCHS,
 #endif
 
 #if PATH_TRACER_USE_VOLUMETRIC_SHADING_MODEL
-	if (Payload.ShadingModelID == SHADINGMODELID_UNLIT)
+	if (ShadingModelID == SHADINGMODELID_UNLIT)
 	{
 		// encode the directionality using Anisotropy: remap [0,2] range to [-1,1]
 	#if PATH_TRACER_USE_VOLUMETRIC_SHADING_MODEL_DIRECTIONAL
@@ -903,21 +903,23 @@ RAY_TRACING_ENTRY_CLOSEST_HIT(PathTracingMaterialCHS,
 		Payload.DiffuseColor = BaseColor;
 		Payload.SpecularColor = 0;
 		Payload.RoughnessData = 1;
+		Payload.PathTracingShadingModelID = PATH_TRACING_SHADINGMODELID_VOLUMETRIC;
 	}
 	else
 #elif MATERIAL_SHADINGMODEL_UNLIT
-	if (Payload.ShadingModelID == SHADINGMODELID_UNLIT)
+	if (ShadingModelID == SHADINGMODELID_UNLIT)
 	{
 		// Regular unlit material, make sure the volumetric component is zero'ed out
 		Payload.DiffuseColor = 0;
 		Payload.SpecularColor = 0;
+		Payload.PathTracingShadingModelID = PATH_TRACING_SHADINGMODELID_LAMBERT;
 	}
 	else
 #endif // MATERIAL_SHADINGMODEL_UNLIT
 
 
 #if MATERIAL_SHADINGMODEL_CLEAR_COAT
-	if (Payload.ShadingModelID == SHADINGMODELID_CLEAR_COAT)
+	if (ShadingModelID == SHADINGMODELID_CLEAR_COAT)
 	{
 		const float ClearCoat = GetMaterialCustomData0(PixelMaterialInputs);
 		const float ClearCoatRoughness = GetMaterialCustomData1(PixelMaterialInputs);
@@ -935,12 +937,13 @@ RAY_TRACING_ENTRY_CLOSEST_HIT(PathTracingMaterialCHS,
 	#endif
 		Payload.RoughnessData.y = ClearCoatRoughness; // adjusted later
 		Payload.RoughnessData.z = ClearCoat;
+		Payload.PathTracingShadingModelID = PATH_TRACING_SHADINGMODELID_CLEARCOAT;
 	}
 	else
 #endif // MATERIAL_SHADINGMODEL_CLEAR_COAT
 
 #if MATERIAL_SHADINGMODEL_CLOTH
-	if (Payload.ShadingModelID == SHADINGMODELID_CLOTH)
+	if (ShadingModelID == SHADINGMODELID_CLOTH)
 	{
 		const float3 ClothColor = GetMaterialSubsurfaceData(PixelMaterialInputs).rgb;
 		const float Fuzz = saturate(GetMaterialCustomData0(PixelMaterialInputs));
@@ -950,23 +953,24 @@ RAY_TRACING_ENTRY_CLOSEST_HIT(PathTracingMaterialCHS,
 		Payload.FuzzColor = ClothColor;
 		Payload.FuzzAmount = Fuzz;
 		Payload.FuzzRoughness = 0.5 * Roughness + 0.5; // (adjust range because the legacy model has a very different response)
-		Payload.ShadingModelID = SHADINGMODELID_SUBSTRATE;
+		Payload.PathTracingShadingModelID = PATH_TRACING_SHADINGMODELID_BASIC;
 	}
 	else
 #endif // MATERIAL_SHADINGMODEL_CLOTH
 
 #if MATERIAL_SHADINGMODEL_TWOSIDED_FOLIAGE
-	if (Payload.ShadingModelID == SHADINGMODELID_TWOSIDED_FOLIAGE)
+	if (ShadingModelID == SHADINGMODELID_TWOSIDED_FOLIAGE)
 	{
 		const float3 SubsurfaceColor = GetMaterialSubsurfaceData(PixelMaterialInputs).rgb;
 		Payload.SetFoliageTransmissionColor(SubsurfaceColor);
+		Payload.PathTracingShadingModelID = PATH_TRACING_SHADINGMODELID_FOLIAGE;
 	}
 	else
 #endif // MATERIAL_SHADINGMODEL_TWOSIDED_FOLIAGE
 
 #if MATERIAL_SHADINGMODEL_SUBSURFACE || MATERIAL_SHADINGMODEL_PREINTEGRATED_SKIN
-	if (Payload.ShadingModelID == SHADINGMODELID_SUBSURFACE ||
-		Payload.ShadingModelID == SHADINGMODELID_PREINTEGRATED_SKIN)
+	if (ShadingModelID == SHADINGMODELID_SUBSURFACE ||
+		ShadingModelID == SHADINGMODELID_PREINTEGRATED_SKIN)
 	{
 		float3 SubsurfaceColor = GetMaterialSubsurfaceData(PixelMaterialInputs).rgb;
 
@@ -988,7 +992,7 @@ RAY_TRACING_ENTRY_CLOSEST_HIT(PathTracingMaterialCHS,
 
 		AdjustPathTracingTwoSidedColorBalance(Payload.DiffuseColor, SubsurfaceColor);
 		// default substrate slab includes SSS
-		Payload.ShadingModelID = SHADINGMODELID_SUBSTRATE;
+		Payload.PathTracingShadingModelID = PATH_TRACING_SHADINGMODELID_BASIC;
 		Payload.DiffuseColor += SubsurfaceColor; // no distinction between diffuse portion and SSS portion (could try using stochastic blending here instead?)
 		Payload.MeanFreePath = SSSRadius;
 		Payload.PhaseG = PhaseG;
@@ -997,7 +1001,7 @@ RAY_TRACING_ENTRY_CLOSEST_HIT(PathTracingMaterialCHS,
 #endif // MATERIAL_SHADINGMODEL_SUBSURFACE || MATERIAL_SHADINGMODEL_PREINTEGRATED_SKIN
 
 #if MATERIAL_SHADINGMODEL_SUBSURFACE_PROFILE
-	if (Payload.ShadingModelID == SHADINGMODELID_SUBSURFACE_PROFILE)
+	if (ShadingModelID == SHADINGMODELID_SUBSURFACE_PROFILE)
 	{
 		const uint ProfileId = ExtractSubsurfaceProfileInt(GetMaterialSubsurfaceDataRaw(PixelMaterialInputs).a);
 
@@ -1036,7 +1040,7 @@ RAY_TRACING_ENTRY_CLOSEST_HIT(PathTracingMaterialCHS,
 	#endif
 
 		// substrate slab includes SSS and dual roughness already
-		Payload.ShadingModelID = SHADINGMODELID_SUBSTRATE;
+		Payload.PathTracingShadingModelID = PATH_TRACING_SHADINGMODELID_BASIC;
 		Payload.MeanFreePath = SSSRadius;
 		Payload.PhaseG = PhaseG;
 		Payload.RoughnessData = float3(Lobe0Roughness, Lobe1Roughness, LobeMix);
@@ -1045,7 +1049,7 @@ RAY_TRACING_ENTRY_CLOSEST_HIT(PathTracingMaterialCHS,
 #endif // MATERIAL_SHADINGMODEL_SUBSURFACE_PROFILE
 
 #if MATERIAL_SHADINGMODEL_SINGLELAYERWATER
-	if (Payload.ShadingModelID == SHADINGMODELID_SINGLELAYERWATER)
+	if (ShadingModelID == SHADINGMODELID_SINGLELAYERWATER)
 	{
 		// For single layer water we always want to treat it as front facing when evaluating parameters 
 		// as hitting back face is treated as being underwater and we get completely different parameters.
@@ -1059,7 +1063,7 @@ RAY_TRACING_ENTRY_CLOSEST_HIT(PathTracingMaterialCHS,
 		const float3 WaterExtinction = ScatteringCoeff + AbsorptionCoeff;
 		const float3 WaterAlbedo = select(WaterExtinction > 0.0, ScatteringCoeff / WaterExtinction, float3(0.0, 0.0, 0.0));
 
-		Payload.ShadingModelID = SHADINGMODELID_SOLID_GLASS; // replace shading model to avoid extra cases in raygen
+		Payload.PathTracingShadingModelID = PATH_TRACING_SHADINGMODELID_SOLID_GLASS; // replace shading model to avoid extra cases in raygen
 		Payload.MeanFreePath = rcp(WaterExtinction);
 	    Payload.SetGlassAlbedo(WaterAlbedo);
 		Payload.SetGlassPhase(PhaseG);
@@ -1072,7 +1076,7 @@ RAY_TRACING_ENTRY_CLOSEST_HIT(PathTracingMaterialCHS,
 #endif // MATERIAL_SHADINGMODEL_SINGLELAYERWATER
 
 #if MATERIAL_SHADINGMODEL_THIN_TRANSLUCENT
-	if (Payload.ShadingModelID == SHADINGMODELID_THIN_TRANSLUCENT)
+	if (ShadingModelID == SHADINGMODELID_THIN_TRANSLUCENT)
 	{
 		const float3 TransmittanceColor = GetThinTranslucentMaterialOutput0(MaterialParameters);
 
@@ -1083,12 +1087,13 @@ RAY_TRACING_ENTRY_CLOSEST_HIT(PathTracingMaterialCHS,
 			Payload.Ior = Ior;
 			Payload.SetTransmittanceColor(TransmittanceColor);
 			Payload.DiffuseColor = 0;
+			Payload.PathTracingShadingModelID = PATH_TRACING_SHADINGMODELID_THIN_GLASS;
 		}
 		else
 		{
 			// Pick solid lobe -- toggle back to default_lit response
 			// Ior stays 0.0 since this mode does not include any refraction
-			Payload.ShadingModelID = SHADINGMODELID_SUBSTRATE;
+			Payload.PathTracingShadingModelID = PATH_TRACING_SHADINGMODELID_BASIC;
 		}
 		
 		float SurfaceCoverage = GetThinTranslucentMaterialOutput1(MaterialParameters);
@@ -1120,7 +1125,7 @@ RAY_TRACING_ENTRY_CLOSEST_HIT(PathTracingMaterialCHS,
 #endif // MATERIAL_SHADINGMODEL_THIN_TRANSLUCENT
 
 #if MATERIAL_SHADINGMODEL_HAIR
-	if (Payload.ShadingModelID == SHADINGMODELID_HAIR)
+	if (ShadingModelID == SHADINGMODELID_HAIR)
 	{
 		Payload.SetBaseColor(BaseColor);
 		Payload.SetHairLongitudinalRoughness(AdjustMaterialRoughness(Roughness, PackedPayload.GetPathRoughness()));
@@ -1132,12 +1137,13 @@ RAY_TRACING_ENTRY_CLOSEST_HIT(PathTracingMaterialCHS,
 #else
 		Payload.SetHairPrimitiveUV(0.5);
 #endif
+		Payload.PathTracingShadingModelID = PATH_TRACING_SHADINGMODELID_HAIR;
 	}
 	else
 #endif // MATERIAL_SHADINGMODEL_HAIR
 
 #if MATERIAL_SHADINGMODEL_EYE
-	if (Payload.ShadingModelID == SHADINGMODELID_EYE)
+	if (ShadingModelID == SHADINGMODELID_EYE)
 	{
 		// This is all based on logic from ShadingModelsMaterial.ush
 		const float IrisMask     = saturate(GetMaterialCustomData0(PixelMaterialInputs));
@@ -1199,14 +1205,15 @@ RAY_TRACING_ENTRY_CLOSEST_HIT(PathTracingMaterialCHS,
 
 		Payload.MeanFreePath = SSSRadius;
 		Payload.PhaseG = PhaseG;
+		Payload.PathTracingShadingModelID = PATH_TRACING_SHADINGMODELID_EYE;
 	}
 	else
 #endif // MATERIAL_SHADINGMODEL_EYE
 
 #if MATERIAL_SHADINGMODEL_DEFAULT_LIT
-	if (Payload.ShadingModelID == SHADINGMODELID_DEFAULT_LIT)
+	if (ShadingModelID == SHADINGMODELID_DEFAULT_LIT)
 	{
-		Payload.ShadingModelID = SHADINGMODELID_SUBSTRATE;
+		Payload.PathTracingShadingModelID = PATH_TRACING_SHADINGMODELID_BASIC;
 
 	#if PATH_TRACER_REFRACTION_ENABLED
 		// only allow refraction for default lit materials since we need space for the absorption amount
@@ -1223,7 +1230,7 @@ RAY_TRACING_ENTRY_CLOSEST_HIT(PathTracingMaterialCHS,
 				// glass
 				// we are using actual refraction, so disable transparency
 				Payload.Ior = Ior;
-				Payload.ShadingModelID = SHADINGMODELID_SOLID_GLASS;
+				Payload.PathTracingShadingModelID = PATH_TRACING_SHADINGMODELID_SOLID_GLASS;
 
 			#if HAVE_GetAbsorptionMediumMaterialOutput0
 				float3 Extinction = PathTracingGlassTransmittanceToExtinction(GetAbsorptionMediumMaterialOutput0(MaterialParameters));
@@ -1293,10 +1300,9 @@ RAY_TRACING_ENTRY_CLOSEST_HIT(PathTracingMaterialCHS,
 			}
 
 			Payload.SetHoldout();
-			// zero-out the material contribution (including base-color and specular since UNLIT is used for volumetric lighting modes - see above)
-			Payload.ShadingModelID = SHADINGMODELID_UNLIT;
+			// zero-out the material contribution so it renders black
+			Payload.PathTracingShadingModelID = PATH_TRACING_SHADINGMODELID_LAMBERT;
 			Payload.Radiance = 0;
-			Payload.Anisotropy = 0;
 			Payload.SetBaseColor(0.0);
 		}
 	}
diff --git a/Engine/Shaders/Private/PathTracing/Volume/PathTracingAtmosphere.ush b/Engine/Shaders/Private/PathTracing/Volume/PathTracingAtmosphere.ush
--- a/Engine/Shaders/Private/PathTracing/Volume/PathTracingAtmosphere.ush
+++ b/Engine/Shaders/Private/PathTracing/Volume/PathTracingAtmosphere.ush
@@ -129,14 +129,14 @@ FPackedPathTracingPayload AtmosphereGetBlockerHit(float3 Origin, float3 Directio
 	Payload.SetBaseColor(Atmosphere.GroundAlbedo.rgb);
 	Payload.BSDFOpacity = 1.0;
 	Payload.WorldNormal = Payload.WorldGeoNormal = Payload.WorldSmoothNormal = WorldN;
-	Payload.ShadingModelID = SHADINGMODELID_LAMBERT; // defaults to diffuse
+	Payload.PathTracingShadingModelID = PATH_TRACING_SHADINGMODELID_LAMBERT; // defaults to diffuse
 	Payload.HitT = HitT;
 	Payload.PrimitiveLightingChannelMask = 7; // visible to all lights
 	Payload.SetFrontFace();
 	if (bIsCameraRay && (VolumeFlags & PATH_TRACER_VOLUME_HOLDOUT_ATMOSPHERE) != 0)
 	{
 		Payload.SetHoldout();
-		Payload.ShadingModelID = SHADINGMODELID_UNLIT;
+		Payload.SetBaseColor(0);
 	}
 	return PackPathTracingPayload(Payload);
 }
diff --git a/Engine/Shaders/Private/PathTracing/Volume/PathTracingVolumeCommon.ush b/Engine/Shaders/Private/PathTracing/Volume/PathTracingVolumeCommon.ush
--- a/Engine/Shaders/Private/PathTracing/Volume/PathTracingVolumeCommon.ush
+++ b/Engine/Shaders/Private/PathTracing/Volume/PathTracingVolumeCommon.ush
@@ -234,7 +234,7 @@ FPathTracingPayload CreateMediumHitPayload(float HitT, float3 TranslatedWorldPos
 	FPathTracingPayload Result = (FPathTracingPayload)0; // clear all fields
 	Result.HitT = HitT;
 	Result.TranslatedWorldPos = TranslatedWorldPos;
-	Result.ShadingModelID = SHADINGMODELID_MEDIUM;
+	Result.PathTracingShadingModelID = PATH_TRACING_SHADINGMODELID_MEDIUM;
 	Result.BSDFOpacity = 1.0;
 	Result.PrimitiveLightingChannelMask = 7;
 	Result.SetFrontFace();
