Fix preexposure being applied twice on diffuse lighting when Lumen reflections are disabled.

#Synced-CL 51018413

diff --git a/Engine/Shaders/Private/ReflectionEnvironmentPixelShader.usf b/Engine/Shaders/Private/ReflectionEnvironmentPixelShader.usf
--- a/Engine/Shaders/Private/ReflectionEnvironmentPixelShader.usf
+++ b/Engine/Shaders/Private/ReflectionEnvironmentPixelShader.usf
@@ -546,15 +546,15 @@ void ReflectionEnvironmentSkyLighting(
 				DiffuseLighting,
 				SpecularLighting);
 
-			// Diffuse lighting needs pre exposition as GetEnvDiffuseLighting() is not preexposed.
-			DiffuseLighting *= View.PreExposure; 
-
 		#if REFLECTION_SOURCE_TYPE == REFLECTION_SOURCE_LUMEN_STANDALONE
 			{
 				// Pre exposure the sky light specular evaluated from within SubstrateEnvLightingCommon (CompositeReflectionCapturesAndSkylightTWS) when rendering lumen HWRT reflections while diffuse global illumination is "none".
 				// This matches with the fact that LumenReflectionTexture is also already pre-exposed, see condition below using "REFLECTION_SOURCE_TYPE == REFLECTION_SOURCE_LUMEN_STANDALONE".
 				SpecularLighting *= View.PreExposure; 
 
+				// Diffuse lighting needs pre exposition as GetEnvDiffuseLighting() is not preexposed
+				DiffuseLighting *= View.PreExposure; 
+
 				const FLumenReflectionOnlyOutput LumenReflections = GetLumenReflectionOnly(SubstrateEnvLight, SubstrateBSDFContext.BSDF, BSDFThroughput, Texture2DArraySample(LumenReflectionTexture, LumenReflectionTextureSampler, float3(BufferUV, ClosureIndex)));
 
 				SpecularLighting *= min(LumenReflections.LightingLobe0.a, LumenReflections.LightingLobe1.a);
