Lumen Hit-Lighting - Fix Substrate-GBuffer subsurface shading model artifacts, cause by an incorrect subsurface color in ray hit payload

#Synced-CL 51011389

diff --git a/Engine/Shaders/Private/RayTracing/RayTracingMaterialHitShaders.usf b/Engine/Shaders/Private/RayTracing/RayTracingMaterialHitShaders.usf
--- a/Engine/Shaders/Private/RayTracing/RayTracingMaterialHitShaders.usf
+++ b/Engine/Shaders/Private/RayTracing/RayTracingMaterialHitShaders.usf
@@ -509,7 +509,7 @@ void GetMaterialPayload(
 	}
 #endif
 
-#if !SUBSTRATE_ENABLED	// Not need to do that since the SubstrateExport already handle the CustomData all at once.
+#if !SUBSTRATE_ENABLED || SUBSTRATE_GBUFFER_FORMAT == 0	// Not need to do that since the SubstrateExport already handle the CustomData all at once.
 #if MATERIAL_SHADINGMODEL_CLEAR_COAT
 	#if CLEAR_COAT_BOTTOM_NORMAL
 	if (Payload.ShadingModelID == SHADINGMODELID_CLEAR_COAT)
