Substrate - fix for shader compilation missed from 51052878.

#rnx
#Synced-CL 51050956

diff --git a/Engine/Shaders/Private/Substrate/SubstrateExport.ush b/Engine/Shaders/Private/Substrate/SubstrateExport.ush
--- a/Engine/Shaders/Private/Substrate/SubstrateExport.ush
+++ b/Engine/Shaders/Private/Substrate/SubstrateExport.ush
@@ -1282,7 +1282,7 @@ FExportResult SubstrateMaterialExportOut(
 		#undef CurrentBSDF
 	}
 
-	if (!LegacyShadingModelSupportsAnisotropy(Out.ShadingModel))
+	if (!LegacyShadingModelSupportsAnisotropy(Out.ShadingModelID))
 	{
 		// Legacy Subsurface profile did not support anisotropy so enforcing that to 0 here for backward compatibility.
 		Out.Anisotropy = 0;
