Fix spelling error
#rnx

diff --git a/Engine/Shaders/Private/PathTracing/PathTracingMaterialHitShader.usf b/Engine/Shaders/Private/PathTracing/PathTracingMaterialHitShader.usf
--- a/Engine/Shaders/Private/PathTracing/PathTracingMaterialHitShader.usf
+++ b/Engine/Shaders/Private/PathTracing/PathTracingMaterialHitShader.usf
@@ -749,7 +749,7 @@ RAY_TRACING_ENTRY_CLOSEST_HIT(PathTracingMaterialCHS,
 					#else
 					Payload.SetHairPrimitiveUV(0.5);
 					#endif
-					Payload.PathTracingShadingModelID = PATH_TRACNIG_SHADINGMODELID_HAIR;
+					Payload.PathTracingShadingModelID = PATH_TRACING_SHADINGMODELID_HAIR;
 					#if SUBSTRATE_PATH_TRACING_SPECULAR_PROFILE
 					Payload.SpecularProfileId = 0;
 					#endif
