Fixed: Incorrect custom primitive data access on mobile platforms that use UBO GPUScene implementation

#Synced-CL 51046133

diff --git a/Engine/Shaders/Private/SceneData.ush b/Engine/Shaders/Private/SceneData.ush
--- a/Engine/Shaders/Private/SceneData.ush
+++ b/Engine/Shaders/Private/SceneData.ush
@@ -1251,16 +1251,6 @@ FInstanceSceneData GetInstanceSceneDataInternal(uint InstanceId, bool bCheckVali
 	return InstanceData;
 }
 
-FInstanceSceneData GetInstanceSceneData(uint InstanceId)
-{
-	return GetInstanceSceneDataInternal(InstanceId, true);
-}
-
-FInstanceSceneData GetInstanceSceneDataUnchecked(uint InstanceId)
-{
-	return GetInstanceSceneDataInternal(InstanceId, false);
-}
-
 // Instance culling flags
 #define INSTANCE_CULLING_FLAG_EVALUATE_WPO (1u << 0u)
 
@@ -1438,3 +1428,21 @@ FLightSceneData GetLightSceneData(uint LightId)
 
 	return Output;
 }
+
+FInstanceSceneData GetInstanceSceneData(uint InstanceId)
+{
+#if SHADER_USES_PRIMITIVE_UBO
+	return GetSceneDataIntermediates(0, InstanceId).InstanceData;
+#else
+	return GetInstanceSceneDataInternal(InstanceId, true);
+#endif
+}
+
+FInstanceSceneData GetInstanceSceneDataUnchecked(uint InstanceId)
+{
+#if SHADER_USES_PRIMITIVE_UBO
+	return GetSceneDataIntermediates(0, InstanceId).InstanceData;
+#else
+	return GetInstanceSceneDataInternal(InstanceId, false);
+#endif
+}
