Fixed: Incorrect custom primitive data access on mobile platforms that use UBO GPUScene implementation

#Synced-CL 51046133

diff --git a/Engine/Shaders/Private/SceneData.ush b/Engine/Shaders/Private/SceneData.ush
--- a/Engine/Shaders/Private/SceneData.ush
+++ b/Engine/Shaders/Private/SceneData.ush
@@ -1251,15 +1251,17 @@
 	return InstanceData;
 }
 
-FInstanceSceneData GetInstanceSceneData(uint InstanceId)
+struct FSceneDataIntermediates
 {
-	return GetInstanceSceneDataInternal(InstanceId, true);
-}
+	uint PrimitiveId;
+	uint InstanceId;
+	uint ViewIndex;
+	uint CullingFlags;
+	// Index from which we load the instance info, needed for the 
+	uint InstanceIdLoadIndex;
+	FInstanceSceneData InstanceData;
+	FPrimitiveSceneData Primitive;
+};
 
-FInstanceSceneData GetInstanceSceneDataUnchecked(uint InstanceId)
-{
-	return GetInstanceSceneDataInternal(InstanceId, false);
-}
-
 // Instance culling flags
 #define INSTANCE_CULLING_FLAG_EVALUATE_WPO (1u << 0u)
@@ -1286,17 +1288,54 @@
 	CullingFlags = (PackedId >> INSTANCE_ID_NUM_BITS) & INSTANCE_CULLING_FLAGS_MASK;
 }
 
-struct FSceneDataIntermediates
+
+#if VF_USE_PRIMITIVE_SCENE_DATA && SHADER_USES_PRIMITIVE_UBO
+FSceneDataIntermediates GetSceneDataIntermediatesUBO(uint DrawInstanceId)
 {
-	uint PrimitiveId;
-	uint InstanceId;
-	uint ViewIndex;
-	uint CullingFlags;
-	// Index from which we load the instance info, needed for the 
-	uint InstanceIdLoadIndex;
-	FInstanceSceneData InstanceData;
-	FPrimitiveSceneData Primitive;
-};
+	FSceneDataIntermediates Intermediates = (FSceneDataIntermediates)0;
+#if USE_INSTANCE_CULLING
+	Intermediates.Primitive = GetPrimitiveDataFromUniformBuffer();
+	Intermediates.InstanceData = LoadInstanceDataUBO(DrawInstanceId);
+	ComputeInstanceDerivedData(Intermediates.InstanceData, Intermediates.InstanceData.LocalToWorld.PostTranslation, Intermediates.InstanceData.LocalToWorld.M);
+#else
+	Intermediates.Primitive = LoadPrimitiveDataUBO(DrawInstanceId);
+	// Populate instance data from primitive data
+	Intermediates.InstanceData.LocalToWorld = Intermediates.Primitive.LocalToWorld;
+	Intermediates.InstanceData.PrevLocalToWorld = Intermediates.Primitive.PreviousLocalToWorld;
+	Intermediates.InstanceData.WorldToLocal = Intermediates.Primitive.WorldToLocal;
+	Intermediates.InstanceData.NonUniformScale = Intermediates.Primitive.NonUniformScale;
+	Intermediates.InstanceData.InvNonUniformScale = Intermediates.Primitive.InvNonUniformScale;
+	Intermediates.InstanceData.DeterminantSign = GetPrimitive_DeterminantSign_FromFlags(Intermediates.InstanceData.Flags);
+#endif
+	Intermediates.InstanceData.LocalBoundsCenter = Intermediates.Primitive.InstanceLocalBoundsCenter;
+	Intermediates.InstanceData.LocalBoundsExtent = Intermediates.Primitive.InstanceLocalBoundsExtent;
+	Intermediates.InstanceData.ValidInstance = true;
+	Intermediates.InstanceData.PrimitiveId = DrawInstanceId;
+	Intermediates.CullingFlags = INSTANCE_CULLING_FLAGS_DEFAULT;
+	Intermediates.ViewIndex = 0U;
+	Intermediates.PrimitiveId = DrawInstanceId;
+	Intermediates.InstanceId = 0U;
+	return Intermediates;
+}
+#endif
 
+FInstanceSceneData GetInstanceSceneData(uint InstanceId)
+{
+#if SHADER_USES_PRIMITIVE_UBO
+	return GetSceneDataIntermediatesUBO(InstanceId).InstanceData;
+#else
+	return GetInstanceSceneDataInternal(InstanceId, true);
+#endif
+}
+
+FInstanceSceneData GetInstanceSceneDataUnchecked(uint InstanceId)
+{
+#if SHADER_USES_PRIMITIVE_UBO
+	return GetSceneDataIntermediatesUBO(InstanceId).InstanceData;
+#else
+	return GetInstanceSceneDataInternal(InstanceId, false);
+#endif
+}
+
 /**
  * Load scene data once given the inputs require.
@@ -1308,28 +1347,7 @@
 {
 	FSceneDataIntermediates Intermediates = (FSceneDataIntermediates)0;
 #if SHADER_USES_PRIMITIVE_UBO
-	#if USE_INSTANCE_CULLING
-		Intermediates.Primitive = GetPrimitiveDataFromUniformBuffer();
-		Intermediates.InstanceData = LoadInstanceDataUBO(DrawInstanceId);
-		ComputeInstanceDerivedData(Intermediates.InstanceData, Intermediates.InstanceData.LocalToWorld.PostTranslation, Intermediates.InstanceData.LocalToWorld.M);
-	#else
-		Intermediates.Primitive = LoadPrimitiveDataUBO(DrawInstanceId);
-		// Populate instance data from primitive data
-		Intermediates.InstanceData.LocalToWorld = Intermediates.Primitive.LocalToWorld;
-		Intermediates.InstanceData.PrevLocalToWorld = Intermediates.Primitive.PreviousLocalToWorld;
-		Intermediates.InstanceData.WorldToLocal = Intermediates.Primitive.WorldToLocal;
-		Intermediates.InstanceData.NonUniformScale = Intermediates.Primitive.NonUniformScale;
-		Intermediates.InstanceData.InvNonUniformScale = Intermediates.Primitive.InvNonUniformScale;
-		Intermediates.InstanceData.DeterminantSign = GetPrimitive_DeterminantSign_FromFlags(Intermediates.InstanceData.Flags);
-	#endif
-	Intermediates.InstanceData.LocalBoundsCenter = Intermediates.Primitive.InstanceLocalBoundsCenter;
-	Intermediates.InstanceData.LocalBoundsExtent = Intermediates.Primitive.InstanceLocalBoundsExtent;
-	Intermediates.InstanceData.ValidInstance = true;
-	Intermediates.InstanceData.PrimitiveId = DrawInstanceId;
-	Intermediates.CullingFlags = INSTANCE_CULLING_FLAGS_DEFAULT;
-	Intermediates.ViewIndex = 0U;
-	Intermediates.PrimitiveId = DrawInstanceId;
-	Intermediates.InstanceId = 0U;
+	Intermediates = GetSceneDataIntermediatesUBO(DrawInstanceId);
 #else
 	Intermediates.InstanceIdLoadIndex = InstanceIdOffset + DrawInstanceId;
 	// GPUCULL_TODO: workaround for the fact that DrawDynamicMeshPassPrivate et al. don't work with GPU-Scene instancing
