Ensure depth samples aren't truncated to half precision. Sample texture directly instead of the MaterialFloat* Texture*Sample*() functions as these cast to MaterialFloat* which is defined as half when half precision is enabled.

#Synced-CL 51052979

diff --git a/Engine/Shaders/Private/CopyShadowMaps.usf b/Engine/Shaders/Private/CopyShadowMaps.usf
--- a/Engine/Shaders/Private/CopyShadowMaps.usf
+++ b/Engine/Shaders/Private/CopyShadowMaps.usf
@@ -16,7 +16,7 @@ void Copy2DDepthPS(
 	)
 {
 	OutColor = 0;
-	OutDepth = Texture2DSampleLevel(ShadowDepthTexture, ShadowDepthSampler, Input.UV, 0).x;
+	OutDepth = ShadowDepthTexture.SampleLevel(ShadowDepthSampler, Input.UV, 0).x;
 }
 
 struct FCubeCopyGSToPS
diff --git a/Engine/Shaders/Private/ScreenPass.usf b/Engine/Shaders/Private/ScreenPass.usf
--- a/Engine/Shaders/Private/ScreenPass.usf
+++ b/Engine/Shaders/Private/ScreenPass.usf
@@ -29,8 +29,8 @@ float4 CopyRectPS(
 {
 	float2 UV = UVAndScreenPos.xy;
 #if MOBILE_MULTI_VIEW
-	return Texture2DArraySample(InputTexture, InputSampler, float3(UV, GetEyeIndex(StereoInput)));
+	return InputTexture.SampleLevel(InputSampler, float3(UV, GetEyeIndex(StereoInput)), 0);
 #else
-	return Texture2DSample(InputTexture, InputSampler, UV);
+	return InputTexture.SampleLevel(InputSampler, UV, 0);
 #endif
 }
\ No newline at end of file
