Fixed incorrect ifdef in SceneData.ush that causes the legacy non-GPU scene path to fail to compile.

#rnx

#Synced-CL 51095457

diff --git a/Engine/Shaders/Private/SceneData.ush b/Engine/Shaders/Private/SceneData.ush
--- a/Engine/Shaders/Private/SceneData.ush
+++ b/Engine/Shaders/Private/SceneData.ush
@@ -1321,7 +1321,7 @@ FSceneDataIntermediates GetSceneDataIntermediatesUBO(uint DrawInstanceId)
 
 FInstanceSceneData GetInstanceSceneData(uint InstanceId)
 {
-#if SHADER_USES_PRIMITIVE_UBO
+#if VF_USE_PRIMITIVE_SCENE_DATA && SHADER_USES_PRIMITIVE_UBO
 	return GetSceneDataIntermediatesUBO(InstanceId).InstanceData;
 #else
 	return GetInstanceSceneDataInternal(InstanceId, true);
@@ -1330,7 +1330,7 @@ FInstanceSceneData GetInstanceSceneData(uint InstanceId)
 
 FInstanceSceneData GetInstanceSceneDataUnchecked(uint InstanceId)
 {
-#if SHADER_USES_PRIMITIVE_UBO
+#if VF_USE_PRIMITIVE_SCENE_DATA && SHADER_USES_PRIMITIVE_UBO
 	return GetSceneDataIntermediatesUBO(InstanceId).InstanceData;
 #else
 	return GetInstanceSceneDataInternal(InstanceId, false);
