Add Substrate IOR to F0 variant for conductor materials.

#virtualized
#Synced-CL 51090839

diff --git a/Engine/Shaders/Private/ShadingCommon.ush b/Engine/Shaders/Private/ShadingCommon.ush
--- a/Engine/Shaders/Private/ShadingCommon.ush
+++ b/Engine/Shaders/Private/ShadingCommon.ush
@@ -139,6 +139,19 @@ float DielectricIorToF0(float Ior)
 	return F0;
 }
 
+// Note: n and k are usually measured in spectral domain
+// There is no good/correct translation to linear sRGB
+// One could use linear sRGB measurement curve's peaks
+// * R primary ~ 610-630 nm
+// * G primary ~ 530-550 nm
+// * B primary ~ 450-470 nm
+float3 ConductorIorToF0(float3 n, float3 k)
+{
+	const float3 k2 = k*k;
+	const float3 F0 = (Pow2(n-1) + k2) / (Pow2(n+1) + k2);
+	return F0;
+}
+
 // Anything with Specular less than 2% is physically impossible and is instead considered to be shadowing.
 // See F_Schlick implementation.
 float  GetF0MicroOcclusionThreshold() { return 0.02f; }
