Remove saturate, inadventently clamping HDR values when querying volumetric shadow maps from translucent objects.

#Synced-CL 51070163

diff --git a/Engine/Shaders/Private/BasePassPixelShader.usf b/Engine/Shaders/Private/BasePassPixelShader.usf
--- a/Engine/Shaders/Private/BasePassPixelShader.usf
+++ b/Engine/Shaders/Private/BasePassPixelShader.usf
@@ -1557,7 +1557,7 @@ void FPixelShaderInOut_MainPS(
 	float2 NDC = MaterialParameters.ScreenPosition.xy / MaterialParameters.ScreenPosition.w;
 	float2 ScreenUV = NDC * ResolvedView.ScreenPositionScaleBias.xy + ResolvedView.ScreenPositionScaleBias.wz;
 	float3 TranslatedWorldPosition = SvPositionToResolvedTranslatedWorld(In.SvPosition);
-	float4 HeterogeneousVolumeResult = saturate(AVSM_SampleCameraRadianceAndTransmittance4(ScreenUV, TranslatedWorldPosition, ResolvedView.TranslatedWorldCameraOrigin));
+	float4 HeterogeneousVolumeResult = AVSM_SampleCameraRadianceAndTransmittance4(ScreenUV, TranslatedWorldPosition, ResolvedView.TranslatedWorldCameraOrigin);
 	Fogging.rgb = HeterogeneousVolumeResult.rgb + Fogging.rgb * HeterogeneousVolumeResult.a;
 	Fogging.a *= HeterogeneousVolumeResult.a;
 #endif // ADAPTIVE_VOLUMETRIC_SHADOW_MAP
