Fixed incorrect brace causing a lot of shader compile warnings when cooking EngineTest

#rnx
#Synced-CL 51151659

diff --git a/Engine/Shaders/Private/DeferredShadingCommon.ush b/Engine/Shaders/Private/DeferredShadingCommon.ush
--- a/Engine/Shaders/Private/DeferredShadingCommon.ush
+++ b/Engine/Shaders/Private/DeferredShadingCommon.ush
@@ -878,7 +878,7 @@ FGBufferData MobileFetchAndDecodeGBuffer(in float2 UV, in float2 PixelPos, uint
 	GBuffer.CustomDepth = ConvertFromDeviceZ(MobileSceneTextures.CustomDepthTextureArray.Load(int4(PixelPos.xy, ViewId, 0)).r);
 	GBuffer.CustomStencil = MobileSceneTextures.CustomStencilTextureArray.Load(int4(PixelPos.xy, ViewId, 0)) STENCIL_COMPONENT_SWIZZLE;
 #else
-	GBuffer.CustomDepth = ConvertFromDeviceZ(MobileSceneTextures.CustomDepthTexture.Load(int3(PixelPos.xy, 0))).r;
+	GBuffer.CustomDepth = ConvertFromDeviceZ(MobileSceneTextures.CustomDepthTexture.Load(int3(PixelPos.xy, 0)).r);
 	GBuffer.CustomStencil = MobileSceneTextures.CustomStencilTexture.Load(int3(PixelPos.xy, 0)) STENCIL_COMPONENT_SWIZZLE;
 #endif
 
