Renamed DataBuffer to avoid apparent reserved word conflict

#Synced-CL 51151659

diff --git a/Engine/Shaders/Private/PhysicsFieldEval.ush b/Engine/Shaders/Private/PhysicsFieldEval.ush
--- a/Engine/Shaders/Private/PhysicsFieldEval.ush
+++ b/Engine/Shaders/Private/PhysicsFieldEval.ush
@@ -15,7 +15,7 @@
 struct FFieldEvalContext
 {
 	/** Data buffer that will represent the result of each node evaluation that could be consumed by other nodes */
-    float DataBuffer[MAX_DATA_SIZE];
+    float NodeEvalResultBuffer[MAX_DATA_SIZE];
 	
 	/** Data current index in that data buffer */
     int DataIndex;
@@ -35,19 +35,19 @@ struct FFieldEvalContext
 
 float GetFieldValue(FFieldEvalContext EvalContext, int DataIndex)
 {
-    return (DataIndex < MAX_DATA_SIZE) ? EvalContext.DataBuffer[DataIndex] : 0;
+    return (DataIndex < MAX_DATA_SIZE) ? EvalContext.NodeEvalResultBuffer[DataIndex] : 0;
 }
 
 float3 GetFieldVector(FFieldEvalContext EvalContext, int DataIndex)
 {
-	return ((DataIndex+2) < MAX_DATA_SIZE) ? float3(EvalContext.DataBuffer[DataIndex], EvalContext.DataBuffer[DataIndex+1], EvalContext.DataBuffer[DataIndex+2]) : float3(0,0,0);
+	return ((DataIndex+2) < MAX_DATA_SIZE) ? float3(EvalContext.NodeEvalResultBuffer[DataIndex], EvalContext.NodeEvalResultBuffer[DataIndex+1], EvalContext.NodeEvalResultBuffer[DataIndex+2]) : float3(0,0,0);
 }
 
 void SetFieldValue(inout FFieldEvalContext EvalContext, int DataIndex, float FieldValue)
 {
 	if(DataIndex < MAX_DATA_SIZE)
 	{
-		EvalContext.DataBuffer[DataIndex] = FieldValue;
+		EvalContext.NodeEvalResultBuffer[DataIndex] = FieldValue;
 	}
 }
 
@@ -56,7 +56,7 @@ void InitEvalContext(inout FFieldEvalContext EvalContext, in float TimeSeconds)
 	[unroll]
 	for (int DataIndex = 0; DataIndex < MAX_DATA_SIZE; ++DataIndex)
 	{
-		EvalContext.DataBuffer[DataIndex] = 0.0;
+		EvalContext.NodeEvalResultBuffer[DataIndex] = 0.0;
 	}
     EvalContext.DataIndex = 0;
 	EvalContext.TimeSeconds = TimeSeconds;
