Changed landscape grass weight shader from per-RGBA channel grass weight layout to per-grass weight channel layout. This allows the shader compiler to de-duplicate the code paths for the 4 channels. This saves on shader code size and also makes it sligthly more efficient (3%) on some targets

#Synced-CL 50901530

diff --git a/Engine/Shaders/Private/LandscapeGrassWeight.usf b/Engine/Shaders/Private/LandscapeGrassWeight.usf
--- a/Engine/Shaders/Private/LandscapeGrassWeight.usf
+++ b/Engine/Shaders/Private/LandscapeGrassWeight.usf
@@ -53,182 +53,173 @@ void VSMain(
 
 #elif PIXELSHADER
 
-int4 PerChannelGrassWeightIndices; // One index per channel in the OutColor. 0 indicates the channel exports grass weight 0, 1 for grass weight 1, etc. 
+uint HeightGrassWeightChannelMask; // first 4 bits for HeightGrassWeightIndex0 (i.e. indicates whether a channel is active for the low 8 height bits), the next 4 bits for HeightGrassWeightIndex1 (same for the high 8 height bits), the rest is unused
+uint4 PerGrassWeightChannelMask; // == 128 bits : 4 bits per grass weight x MaxGrassTypes (32 currently). Add more if MaxGrassTypes is increased
 
-// 1 special index to indicate we don't want to export anything
-#define GrassWeightInvalidIndex -3
-// 2 special indices to indicate we want to export heights
-#define HeightGrassWeightIndex0 -2
-#define HeightGrassWeightIndex1 -1
+// FourBitsMask is a mask ([0, 3] range) to store whether a channel (RGBA) is active 
+float4 FourBitsMaskToChannelMask(uint FourBitsMask)
+{
+	return float4(
+		(FourBitsMask >> 0) & 1u,
+		(FourBitsMask >> 1) & 1u,
+		(FourBitsMask >> 2) & 1u,
+		(FourBitsMask >> 3) & 1u);
+}
 
-float GetGrassWeightForChannel(uint ChannelIndex, float2 PackedHeight, FMaterialPixelParameters MaterialParameters)
+// GrassWeightIndex is betweenin the [0, MaxGrassTypes[ range. This allows to store 4 bits per grass weight in a single uint4 (PerGrassWeightChannelMask) : 8 per uint
+// This function retrieves the 4 bits mask (range [0, 15], hence the & 0xf masking) for a given grass weight index
+uint GrassWeightIndexToFourBitsMask(uint GrassWeightIndex)
 {
-	float OutputColor = 0;
+	return (PerGrassWeightChannelMask[GrassWeightIndex / 8] >> ((GrassWeightIndex % 8) * 4)) & 0xf;
+}
 
-#define GRASSWEIGHTCASE(N) case N: \
-	{ \
-	OutputColor = GetGrassWeight##N(MaterialParameters).x; \
-	break; \
-	}
+// Pixel Shader
+void PSMain( 
+	FLandscapeGrassWeightInterpolantsVSToPS Interpolants
+	OPTIONAL_IsFrontFace,
+	out HALF4_TYPE OutColor	: SV_Target0
+	)
+{
+	OutColor = 0;
+
+	ResolvedView = ResolveView();
+
+	// Build the SvPosition without RenderOffset so that we correctly reproject WorldPosition etc.
+	float4 SvPositionPreOffset = Interpolants.PositionPreOffset / Interpolants.PositionPreOffset.w;
+	SvPositionPreOffset.xy = (SvPositionPreOffset.xy * float2(0.5f, -0.5f) + 0.5f) * ResolvedView.ViewSizeAndInvSize.xy;
+
+	FMaterialPixelParameters MaterialParameters = GetMaterialPixelParameters(Interpolants.FactoryInterpolants, SvPositionPreOffset);
+	FPixelMaterialInputs PixelMaterialInputs;
+	CalcMaterialParameters(MaterialParameters, PixelMaterialInputs, SvPositionPreOffset, bIsFrontFace);
 
+	OutColor = 0.0;
 
-	// Do not remove : this branch directive is actually necessary to avoid FXC screwing up the generated shader
-	BRANCH
-	switch(PerChannelGrassWeightIndices[ChannelIndex])
+	if (HeightGrassWeightChannelMask != 0)
 	{
-		case GrassWeightInvalidIndex: 
-			break;
-		case HeightGrassWeightIndex0: 
-		{
-			OutputColor = PackedHeight.x; 
-			break;
-		}
-		case HeightGrassWeightIndex1: 
-		{
-			OutputColor = PackedHeight.y; 
-			break;
-		}
+		// Calculate and pack landscape height.
+#if VF_USE_PRIMITIVE_SCENE_DATA
+		float Height = clamp(round(Interpolants.LocalPosZ / TERRAIN_ZSCALE) + 32768.0, 0, 65535);
+#else // !VF_USE_PRIMITIVE_SCENE_DATA
+		float3 LocalPos = WSMultiplyDemote(GetWorldPosition(MaterialParameters), GetWorldToInstance(MaterialParameters));
+		float Height = clamp(round(LocalPos.z / TERRAIN_ZSCALE) + 32768.0, 0, 65535);
+#endif // VF_USE_PRIMITIVE_SCENE_DATA
+		float2 PackedHeight = float2(floor(Height / 256), fmod(Height, 256)) / 255;
+
+		OutColor += FourBitsMaskToChannelMask((HeightGrassWeightChannelMask >> 0) & 0xf) * PackedHeight.x;
+		OutColor += FourBitsMaskToChannelMask((HeightGrassWeightChannelMask >> 4) & 0xf) * PackedHeight.y;
+	}
+
+	if (any(PerGrassWeightChannelMask) != 0)
+	{	
+#define COMPUTEGRASSWEIGHT(N) \
+{ \
+	uint FourBitsMask = GrassWeightIndexToFourBitsMask(N); \
+	BRANCH \
+	if (FourBitsMask != 0) \
+	{ \
+		OutColor += FourBitsMaskToChannelMask(FourBitsMask) * GetGrassWeight##N(MaterialParameters).x; \
+	} \
+}
+
 #if HAVE_GetGrassWeight0
-		GRASSWEIGHTCASE(0);
+		COMPUTEGRASSWEIGHT(0);
 #endif // HAVE_GetGrassWeight0
 #if HAVE_GetGrassWeight1
-		GRASSWEIGHTCASE(1);
+		COMPUTEGRASSWEIGHT(1);
 #endif // HAVE_GetGrassWeight1
 #if HAVE_GetGrassWeight2
-		GRASSWEIGHTCASE(2);
+		COMPUTEGRASSWEIGHT(2);
 #endif // HAVE_GetGrassWeight2
 #if HAVE_GetGrassWeight3
-		GRASSWEIGHTCASE(3);
+		COMPUTEGRASSWEIGHT(3);
 #endif // HAVE_GetGrassWeight3
 #if HAVE_GetGrassWeight4
-		GRASSWEIGHTCASE(4);
+		COMPUTEGRASSWEIGHT(4);
 #endif // HAVE_GetGrassWeight4
 #if HAVE_GetGrassWeight5
-		GRASSWEIGHTCASE(5);
+		COMPUTEGRASSWEIGHT(5);
 #endif // HAVE_GetGrassWeight5
 #if HAVE_GetGrassWeight6
-		GRASSWEIGHTCASE(6);
+		COMPUTEGRASSWEIGHT(6);
 #endif // HAVE_GetGrassWeight6
 #if HAVE_GetGrassWeight7
-		GRASSWEIGHTCASE(7);
+		COMPUTEGRASSWEIGHT(7);
 #endif // HAVE_GetGrassWeight7
 #if HAVE_GetGrassWeight8
-		GRASSWEIGHTCASE(8);
+		COMPUTEGRASSWEIGHT(8);
 #endif // HAVE_GetGrassWeight8
 #if HAVE_GetGrassWeight9
-		GRASSWEIGHTCASE(9);
+		COMPUTEGRASSWEIGHT(9);
 #endif // HAVE_GetGrassWeight9
 
 #if HAVE_GetGrassWeight10
-		GRASSWEIGHTCASE(10);
+		COMPUTEGRASSWEIGHT(10);
 #endif // HAVE_GetGrassWeight10
 #if HAVE_GetGrassWeight11
-		GRASSWEIGHTCASE(11);
+		COMPUTEGRASSWEIGHT(11);
 #endif // HAVE_GetGrassWeight11
 #if HAVE_GetGrassWeight12
-		GRASSWEIGHTCASE(12);
+		COMPUTEGRASSWEIGHT(12);
 #endif // HAVE_GetGrassWeight12
 #if HAVE_GetGrassWeight13
-		GRASSWEIGHTCASE(13);
+		COMPUTEGRASSWEIGHT(13);
 #endif // HAVE_GetGrassWeight13
 #if HAVE_GetGrassWeight14
-		GRASSWEIGHTCASE(14);
+		COMPUTEGRASSWEIGHT(14);
 #endif // HAVE_GetGrassWeight14
 #if HAVE_GetGrassWeight15
-		GRASSWEIGHTCASE(15);
+		COMPUTEGRASSWEIGHT(15);
 #endif // HAVE_GetGrassWeight15
 #if HAVE_GetGrassWeight16
-		GRASSWEIGHTCASE(16);
+		COMPUTEGRASSWEIGHT(16);
 #endif // HAVE_GetGrassWeight16
 #if HAVE_GetGrassWeight17
-		GRASSWEIGHTCASE(17);
+		COMPUTEGRASSWEIGHT(17);
 #endif // HAVE_GetGrassWeight17
 #if HAVE_GetGrassWeight18
-		GRASSWEIGHTCASE(18);
+		COMPUTEGRASSWEIGHT(18);
 #endif // HAVE_GetGrassWeight18
 #if HAVE_GetGrassWeight19
-		GRASSWEIGHTCASE(19);
+		COMPUTEGRASSWEIGHT(19);
 #endif // HAVE_GetGrassWeight19
 
 #if HAVE_GetGrassWeight20
-		GRASSWEIGHTCASE(20);
+		COMPUTEGRASSWEIGHT(20);
 #endif // HAVE_GetGrassWeight20
 #if HAVE_GetGrassWeight21
-		GRASSWEIGHTCASE(21);
+		COMPUTEGRASSWEIGHT(21);
 #endif // HAVE_GetGrassWeight21
 #if HAVE_GetGrassWeight22
-		GRASSWEIGHTCASE(22);
+		COMPUTEGRASSWEIGHT(22);
 #endif // HAVE_GetGrassWeight22
 #if HAVE_GetGrassWeight23
-		GRASSWEIGHTCASE(23);
+		COMPUTEGRASSWEIGHT(23);
 #endif // HAVE_GetGrassWeight23
 #if HAVE_GetGrassWeight24
-		GRASSWEIGHTCASE(24);
+		COMPUTEGRASSWEIGHT(24);
 #endif // HAVE_GetGrassWeight24
 #if HAVE_GetGrassWeight25
-		GRASSWEIGHTCASE(25);
+		COMPUTEGRASSWEIGHT(25);
 #endif // HAVE_GetGrassWeight25
 #if HAVE_GetGrassWeight26
-		GRASSWEIGHTCASE(26);
+		COMPUTEGRASSWEIGHT(26);
 #endif // HAVE_GetGrassWeight26
 #if HAVE_GetGrassWeight27
-		GRASSWEIGHTCASE(27);
+		COMPUTEGRASSWEIGHT(27);
 #endif // HAVE_GetGrassWeight27
 #if HAVE_GetGrassWeight28
-		GRASSWEIGHTCASE(28);
+		COMPUTEGRASSWEIGHT(28);
 #endif // HAVE_GetGrassWeight28
 #if HAVE_GetGrassWeight29
-		GRASSWEIGHTCASE(29);
+		COMPUTEGRASSWEIGHT(29);
 #endif // HAVE_GetGrassWeight29
 
 #if HAVE_GetGrassWeight30
-		GRASSWEIGHTCASE(30);
+		COMPUTEGRASSWEIGHT(30);
 #endif // HAVE_GetGrassWeight30
 #if HAVE_GetGrassWeight31
-		GRASSWEIGHTCASE(31);
+		COMPUTEGRASSWEIGHT(31);
 #endif // HAVE_GetGrassWeight31
-
-		default:
-			break;
 	}
-	
-	return OutputColor;
-}
-
-// Pixel Shader
-void PSMain( 
-	FLandscapeGrassWeightInterpolantsVSToPS Interpolants
-	OPTIONAL_IsFrontFace,
-	out HALF4_TYPE OutColor	: SV_Target0
-	)
-{
-	OutColor = 0;
-
-	ResolvedView = ResolveView();
-
-	// Build the SvPosition without RenderOffset so that we correctly reproject WorldPosition etc.
-	float4 SvPositionPreOffset = Interpolants.PositionPreOffset / Interpolants.PositionPreOffset.w;
-	SvPositionPreOffset.xy = (SvPositionPreOffset.xy * float2(0.5f, -0.5f) + 0.5f) * ResolvedView.ViewSizeAndInvSize.xy;
-
-	FMaterialPixelParameters MaterialParameters = GetMaterialPixelParameters(Interpolants.FactoryInterpolants, SvPositionPreOffset);
-	FPixelMaterialInputs PixelMaterialInputs;
-	CalcMaterialParameters(MaterialParameters, PixelMaterialInputs, SvPositionPreOffset, bIsFrontFace);
-
-	// Calculate and pack landscape height.
-#if VF_USE_PRIMITIVE_SCENE_DATA
-	float Height = clamp(round(Interpolants.LocalPosZ / TERRAIN_ZSCALE) + 32768.0, 0, 65535);
-#else // !VF_USE_PRIMITIVE_SCENE_DATA
-	float3 LocalPos = WSMultiplyDemote(GetWorldPosition(MaterialParameters), GetWorldToInstance(MaterialParameters));
-	float Height = clamp(round(LocalPos.z / TERRAIN_ZSCALE) + 32768.0, 0, 65535);
-#endif // VF_USE_PRIMITIVE_SCENE_DATA
-	float2 PackedHeight = float2(floor(Height / 256), fmod(Height, 256)) / 255;
-
-	OutColor = 0.0;
-
-	OutColor.r = GetGrassWeightForChannel(0, PackedHeight, MaterialParameters);
-	OutColor.g = GetGrassWeightForChannel(1, PackedHeight, MaterialParameters);
-	OutColor.b = GetGrassWeightForChannel(2, PackedHeight, MaterialParameters);
-	OutColor.a = GetGrassWeightForChannel(3, PackedHeight, MaterialParameters);
-
 }
 #endif
\ No newline at end of file
diff --git a/Engine/Source/Runtime/Landscape/Private/LandscapeGrassWeightExporter.cpp b/Engine/Source/Runtime/Landscape/Private/LandscapeGrassWeightExporter.cpp
--- a/Engine/Source/Runtime/Landscape/Private/LandscapeGrassWeightExporter.cpp
+++ b/Engine/Source/Runtime/Landscape/Private/LandscapeGrassWeightExporter.cpp
@@ -84,8 +84,9 @@ END_SHADER_PARAMETER_STRUCT()
 class FLandscapeGrassWeightShaderElementData : public FMeshMaterialShaderElementData
 {
 public:
-	FIntVector4 PerChannelGrassWeightIndices;
-	FVector2f RenderOffset;
+	uint32 HeightGrassWeightChannelMask = 0;
+	FUint32Vector4 PerGrassWeightChannelMask = FUint32Vector4(0, 0, 0, 0);
+	FVector2f RenderOffset = FVector2f::ZeroVector;
 };
 
 
@@ -180,7 +181,8 @@ IMPLEMENT_MATERIAL_SHADER_TYPE(, FLandscapeGrassWeightVS, TEXT("/Engine/Private/
 class FLandscapeGrassWeightPS : public FMeshMaterialShader
 {
 	DECLARE_SHADER_TYPE(FLandscapeGrassWeightPS, MeshMaterial);
-	LAYOUT_FIELD(FShaderParameter, PerChannelGrassWeightIndicesParameter);
+	LAYOUT_FIELD(FShaderParameter, HeightGrassWeightChannelMaskParameter);
+	LAYOUT_FIELD(FShaderParameter, PerGrassWeightChannelMaskParameter);
 public:
 
 	static bool ShouldCompilePermutation(const FMeshMaterialShaderPermutationParameters& Parameters)
@@ -191,7 +193,8 @@ public:
 	FLandscapeGrassWeightPS(const ShaderMetaType::CompiledShaderInitializerType& Initializer)
 		: FMeshMaterialShader(Initializer)
 	{
-		PerChannelGrassWeightIndicesParameter.Bind(Initializer.ParameterMap, TEXT("PerChannelGrassWeightIndices"));
+		HeightGrassWeightChannelMaskParameter.Bind(Initializer.ParameterMap, TEXT("HeightGrassWeightChannelMask"));
+		PerGrassWeightChannelMaskParameter.Bind(Initializer.ParameterMap, TEXT("PerGrassWeightChannelMask"));
 		PassUniformBuffer.Bind(Initializer.ParameterMap, FSceneTextureUniformParameters::FTypeInfo::GetStructMetadata()->GetShaderVariableName());
 	}
 
@@ -209,7 +212,8 @@ public:
 	{
 		FMeshMaterialShader::GetShaderBindings(Scene, FeatureLevel, PrimitiveSceneProxy, MaterialRenderProxy, Material, ShaderElementData, ShaderBindings);
 
-		ShaderBindings.Add(PerChannelGrassWeightIndicesParameter, ShaderElementData.PerChannelGrassWeightIndices);
+		ShaderBindings.Add(HeightGrassWeightChannelMaskParameter, ShaderElementData.HeightGrassWeightChannelMask);
+		ShaderBindings.Add(PerGrassWeightChannelMaskParameter, ShaderElementData.PerGrassWeightChannelMask);
 	}
 };
 
@@ -348,6 +352,8 @@ bool FLandscapeGrassWeightMeshProcessor::Process(
 	int32 FirstHeightMipsPassIndex,
 	const TArray<int32>& HeightMips)
 {
+	using namespace UE::Landscape::Grass::Private;
+
 	const FVertexFactory* VertexFactory = MeshBatch.VertexFactory;
 
 	FMaterialShaderTypes ShaderTypes;
@@ -378,7 +384,30 @@ bool FLandscapeGrassWeightMeshProcessor::Process(
 	const int32 NumPasses = PerPassPerChannelGrassWeightIndices.Num();
 	for (int32 PassIndex = 0; PassIndex < NumPasses; ++PassIndex)
 	{
-		ShaderElementData.PerChannelGrassWeightIndices = PerPassPerChannelGrassWeightIndices[PassIndex];
+		ShaderElementData.HeightGrassWeightChannelMask = 0;
+		ShaderElementData.PerGrassWeightChannelMask = FUint32Vector4(0, 0, 0, 0);
+		// Convert the grass weight indices into a format that's more appropriate to the shader compiler : 
+		for (uint32 ChannelIndex = 0; ChannelIndex < 4; ++ChannelIndex)
+		{
+			int32 GrassWeightIndex = PerPassPerChannelGrassWeightIndices[PassIndex][ChannelIndex];
+			// 4 bits per output (one per RGBA channel)
+			if (GrassWeightIndex == HeightGrassWeightIndex0)
+			{
+				// The first 4 bits of HeightGrassWeightChannelMask are for HeightGrassWeightIndex0
+				ShaderElementData.HeightGrassWeightChannelMask |= (1u << (ChannelIndex + 0));
+			}
+			else if (GrassWeightIndex == HeightGrassWeightIndex1)
+			{
+				// The next 4 bits of HeightGrassWeightChannelMask are for HeightGrassWeightIndex1
+				ShaderElementData.HeightGrassWeightChannelMask |= (1u << (ChannelIndex + 4));
+			}
+			// Grass weights : max 32 (MaxGrassTypes), 4 bits per grass weight (one per RGBA channel) = 128 bits, which fit in a uint4 (i.e. 8 grass weights per uint)
+			else if (GrassWeightIndex != GrassWeightInvalidIndex)
+			{
+				ShaderElementData.PerGrassWeightChannelMask[GrassWeightIndex / 8] |= (1u << (ChannelIndex + (GrassWeightIndex % 8) * 4u));
+			}
+		}
+
 		ShaderElementData.RenderOffset = FVector2f(ViewOffset) + FVector2f(PassOffsetX * PassIndex, 0);	// LWC_TODO: Precision loss
 
 		uint64 Mask = (PassIndex >= FirstHeightMipsPassIndex) ? HeightMips[PassIndex - FirstHeightMipsPassIndex] : BatchElementMask;
