Fix for shader not compiling when plugging SelectedForwardDirLight into WPO or while TwoSided is enabled.

#tests PC editor dx12 with a material using SelectedForwardDirLight node with two sided enabled or WPO plugged in.
#rnx
#Synced-CL 51166676

diff --git a/Engine/Shaders/Private/HitProxyPixelShader.usf b/Engine/Shaders/Private/HitProxyPixelShader.usf
--- a/Engine/Shaders/Private/HitProxyPixelShader.usf
+++ b/Engine/Shaders/Private/HitProxyPixelShader.usf
@@ -10,6 +10,7 @@
 #include "Common.ush"
 #include "/Engine/Generated/Material.ush"
 #include "/Engine/Generated/VertexFactory.ush"
+#include "/Engine/Private/BasePassCommonFunctions.ush"
 
 // Whether HitProxy IDs are provided per instance or not
 #define PER_INSTANCE_HITPROXY_ID	(USE_INSTANCING || USE_PER_VERTEX_HITPROXY_ID || USE_VERTEXFACTORY_HITPROXY_ID)
diff --git a/Engine/Shaders/Private/LightMapDensityShader.usf b/Engine/Shaders/Private/LightMapDensityShader.usf
--- a/Engine/Shaders/Private/LightMapDensityShader.usf
+++ b/Engine/Shaders/Private/LightMapDensityShader.usf
@@ -9,6 +9,7 @@
 #include "Common.ush"
 #include "/Engine/Generated/Material.ush"
 #include "/Engine/Generated/VertexFactory.ush"
+#include "/Engine/Private/BasePassCommonFunctions.ush"
 
 struct FLightMapDensityVSToPS
 {
diff --git a/Engine/Shaders/Private/VelocityShader.usf b/Engine/Shaders/Private/VelocityShader.usf
--- a/Engine/Shaders/Private/VelocityShader.usf
+++ b/Engine/Shaders/Private/VelocityShader.usf
@@ -8,6 +8,7 @@
 #include "/Engine/Generated/Material.ush"
 #include "/Engine/Generated/VertexFactory.ush"
 #include "VelocityCommon.ush"
+#include "/Engine/Private/BasePassCommonFunctions.ush"
 
 #ifndef VELOCITY_CLIPPED_DEPTH_PASS
 #define VELOCITY_CLIPPED_DEPTH_PASS 0
