Fixed lighting channels on Nanite meshes

#rnx

#Synced-CL 51176602

diff --git a/Engine/Shaders/Private/DownsampleDepthPixelShader.usf b/Engine/Shaders/Private/DownsampleDepthPixelShader.usf
--- a/Engine/Shaders/Private/DownsampleDepthPixelShader.usf
+++ b/Engine/Shaders/Private/DownsampleDepthPixelShader.usf
@@ -155,7 +155,7 @@ void CopyStencilToLightingChannelsPS(
 
 	uint CombinedMask = LightingChannels;
 	CombinedMask |= bCastShadow ? (1u << LIGHTING_CHANNELS_TEXTURE_CAST_SHADOWS_BIT) : 0;
-	CombinedMask |= bHasRaytracingRepresentation ? (1u << STENCIL_RAY_TRACING_REPRESENTATION_BIT_ID) : 0;
+	CombinedMask |= bHasRaytracingRepresentation ? (1u << LIGHTING_CHANNELS_TEXTURE_RAY_TRACING_REPRESENTATION_BIT) : 0;
 
 	OutValue = CombinedMask;
 }
