Rewrite Slate UI pixel shader such that we call GetMaterialColor only once. We were duplicating that code 5 times in the worst case, and this could result in a shader binary that would be 5x larger that it should be.

#rnx

#Synced-CL 50966172

diff --git a/Engine/Shaders/Private/SlateElementPixelShader.usf b/Engine/Shaders/Private/SlateElementPixelShader.usf
--- a/Engine/Shaders/Private/SlateElementPixelShader.usf
+++ b/Engine/Shaders/Private/SlateElementPixelShader.usf
@@ -48,6 +48,26 @@ uint GetSlateShaderType()
 Texture2D ElementTexture;
 SamplerState ElementTextureSampler;
 
+struct FontSignedDistanceSample
+{
+	MaterialFloat DistanceSample;
+	MaterialFloat SmoothDistanceSample;
+	float ScreenSpaceSpread;
+};
+
+struct FMaterialColorParameters
+{
+	half4 MaterialColor;
+	
+	float4 SdfData;
+	float2 MeshUVs;
+	
+	FontSignedDistanceSample SdfSample;
+	half4 FontSampledColor;
+
+	bool bNeedMaterialColor;
+};
+
 /**
  * Applys a contrast about a midpoint squeezing the colors towards dark or light depending on
  * the midpoint.  Additionally - it increases the opaqueness of colors to try and seperate them from
@@ -276,14 +296,14 @@ half4 GetTextureColor( VertexToPixelInterpolants InVertex, float2 UV )
 #endif
 
 /** Gets a color from a texture and supplied vertex color.  Adjusting for tint masking if necessary */
-half4 GetColor( VertexToPixelInterpolants InVertex, float2 UV )
+half4 GetColor( VertexToPixelInterpolants InVertex, FMaterialColorParameters MaterialColorParameters)
 {
 	half4 FinalColor;
 	
 #if USE_MATERIALS
-	half4 BaseColor = GetMaterialColor( InVertex, UV );
+	half4 BaseColor = MaterialColorParameters.MaterialColor;
 #else
-	half4 BaseColor = GetTextureColor( InVertex, UV );
+	half4 BaseColor = GetTextureColor( InVertex, MaterialColorParameters.MeshUVs );
 #endif
 
 	const uint ShaderType = GetSlateShaderType();
@@ -304,13 +324,6 @@ half4 GetColor( VertexToPixelInterpolants InVertex, float2 UV )
 	return FinalColor;
 }
 
-struct FontSignedDistanceSample
-{
-	MaterialFloat DistanceSample;
-	MaterialFloat SmoothDistanceSample;
-	float ScreenSpaceSpread;
-};
-
 FontSignedDistanceSample SampleFontSignedDistance(float2 TextureCoordinate)
 {
 	FontSignedDistanceSample OutSample;
@@ -350,13 +363,13 @@ float4 GetFontSignedDistanceData(FontSignedDistanceSample InSample)
 }
 #endif
 
-half4 GetSdfTextElementColor(VertexToPixelInterpolants InVertex)
+half4 GetSdfTextElementColor(VertexToPixelInterpolants InVertex, FMaterialColorParameters MaterialColorParameters)
 {
-	const FontSignedDistanceSample SdfSample = SampleFontSignedDistance(SLATE_UV0(InVertex));
+	const FontSignedDistanceSample SdfSample = MaterialColorParameters.SdfSample;
 	half4 OutColor = InVertex.Color;
 #if USE_MATERIALS
-	const float4 SdfData = GetFontSignedDistanceData(SdfSample);
-	OutColor *= GetMaterialColor(InVertex, float2(0, 0), SdfData);
+	const float4 SdfData = MaterialColorParameters.SdfData;
+	OutColor *= MaterialColorParameters.MaterialColor;
 	if (ShaderParams2.y < 0.5) // ShaderParams2.y is bMaterialIsStencil
 	{
 	#if MATERIALBLENDING_ADDITIVE || MATERIALBLENDING_ALPHACOMPOSITE
@@ -371,27 +384,34 @@ half4 GetSdfTextElementColor(VertexToPixelInterpolants InVertex)
 	return OutColor;
 }
 
-half4 GetFontElementColor( VertexToPixelInterpolants InVertex )
+MaterialFloat4 GetFontElementSampledColor(VertexToPixelInterpolants InVertex)
 {
 	float2 AtlasTextureCoordinates = SLATE_UV0( InVertex );
 	MaterialFloat4 SampledColor = MaterialFloat4(1,1,1,1);
 
 	if (GetSlateShaderType() == ST_GrayscaleFont)
 	{
-		#if USE_MATERIALS && (MATERIALBLENDING_ADDITIVE || MATERIALBLENDING_ALPHACOMPOSITE)
-			SampledColor *= Texture2DSample_A8(ElementTexture, ElementTextureSampler, AtlasTextureCoordinates);
-		#else
-			SampledColor.a = Texture2DSample_A8(ElementTexture, ElementTextureSampler, AtlasTextureCoordinates);
-		#endif
+#if USE_MATERIALS && (MATERIALBLENDING_ADDITIVE || MATERIALBLENDING_ALPHACOMPOSITE)
+		SampledColor *= Texture2DSample_A8(ElementTexture, ElementTextureSampler, AtlasTextureCoordinates);
+#else
+		SampledColor.a = Texture2DSample_A8(ElementTexture, ElementTextureSampler, AtlasTextureCoordinates);
+#endif
 	}
 	else
 	{
 		SampledColor = Texture2DSample(ElementTexture, ElementTextureSampler, AtlasTextureCoordinates);
 	}
 
+	return SampledColor;
+}
+
+half4 GetFontElementColor( VertexToPixelInterpolants InVertex, FMaterialColorParameters MaterialColorParameters )
+{
+	MaterialFloat4 SampledColor = MaterialColorParameters.FontSampledColor;
+
 	half4 OutColor = InVertex.Color;
 #if USE_MATERIALS
-	OutColor *= GetMaterialColor( InVertex, float2(0,0), float4(0, 0, 0, SampledColor.a) );
+	OutColor *= MaterialColorParameters.MaterialColor;
 	if (ShaderParams2.y < 0.5) // ShaderParams2.y is bMaterialIsStencil
 	{
 		OutColor *= half4(SampledColor);
@@ -403,12 +423,12 @@ half4 GetFontElementColor( VertexToPixelInterpolants InVertex )
 	return OutColor;
 }
 
-half4 GetDefaultElementColor( VertexToPixelInterpolants InVertex )
+half4 GetDefaultElementColor( VertexToPixelInterpolants InVertex, FMaterialColorParameters MaterialColorParameters )
 {
-	return GetColor( InVertex, SLATE_UV0( InVertex ) * SLATE_UV1( InVertex ) );
+	return GetColor( InVertex, MaterialColorParameters );
 }
 
-half4 GetBorderElementColor( VertexToPixelInterpolants InVertex )
+float2 GetBorderElementColorUV(VertexToPixelInterpolants InVertex)
 {
 	float4 TextureCoordinates = InVertex.TextureCoordinates0;
 	float2 NewUV;
@@ -439,13 +459,17 @@ half4 GetBorderElementColor( VertexToPixelInterpolants InVertex )
 		NewUV = lerp(MinUV,MaxUV,NewUV);	
 	}
 
-	return GetColor( InVertex, NewUV );
+	return NewUV;
 }
 
+half4 GetBorderElementColor( VertexToPixelInterpolants InVertex, FMaterialColorParameters MaterialColorParameters )
+{
+	return GetColor( InVertex, MaterialColorParameters );
+}
 
-half4 GetRoundedBoxElementColor( VertexToPixelInterpolants InVertex )
+half4 GetRoundedBoxElementColor( VertexToPixelInterpolants InVertex, FMaterialColorParameters MaterialColorParameters )
 {
-	return GetRoundedBoxElementColorInternal(ShaderParams.zw, InVertex.MaterialTexCoords.xy, ShaderParams.y, ShaderParams2, GetColor(InVertex, SLATE_UV0(InVertex) * SLATE_UV1(InVertex)),  InVertex.SecondaryColor);
+	return GetRoundedBoxElementColorInternal(ShaderParams.zw, InVertex.MaterialTexCoords.xy, ShaderParams.y, ShaderParams2, GetColor(InVertex, MaterialColorParameters),  InVertex.SecondaryColor);
 }
 
 /**
@@ -504,30 +528,81 @@ half4 GetLineSegmentElementColor( VertexToPixelInterpolants InVertex )
 	return Color;
 }
 
-HALF4_TYPE Main( VertexToPixelInterpolants VIn ) : SV_Target0
+FMaterialColorParameters GetMaterialColorParameters(VertexToPixelInterpolants InVertex, const uint ShaderType)
+{
+	FMaterialColorParameters Output = (FMaterialColorParameters)0;
+	Output.bNeedMaterialColor = true;
+
+	if (ShaderType == ST_Default || ShaderType == ST_Custom)
+	{
+		Output.MeshUVs = SLATE_UV0(InVertex) * SLATE_UV1(InVertex);
+	}
+	else if (ShaderType == ST_RoundedBox)
+	{
+		Output.MeshUVs = SLATE_UV0(InVertex) * SLATE_UV1(InVertex);
+	}
+	else if (ShaderType == ST_Border)
+	{
+		Output.MeshUVs = GetBorderElementColorUV(InVertex);
+	}
+	else if (ShaderType == ST_GrayscaleFont || ShaderType == ST_ColorFont)
+	{
+		Output.MeshUVs = float2(0,0);
+		Output.FontSampledColor = GetFontElementSampledColor(InVertex);
+		Output.SdfData = float4(0, 0, 0, Output.FontSampledColor.a);
+	}
+	else if (ShaderType == ST_SdfFont || ShaderType == ST_MsdfFont)
+	{
+		Output.MeshUVs = float2(0, 0);
+		Output.SdfSample = SampleFontSignedDistance(SLATE_UV0(InVertex));
+#if USE_MATERIALS
+		Output.SdfData = GetFontSignedDistanceData(Output.SdfSample);
+#endif
+	}
+	else
+	{
+		Output.bNeedMaterialColor = false;
+		// Not using Material Color
+		Output.MeshUVs = float2(0, 0);
+	}
+
+	return Output;
+}
+
+HALF4_TYPE Main(VertexToPixelInterpolants VIn) : SV_Target0
 {
 	half4 OutColor;
 	const uint ShaderType = GetSlateShaderType();
 
+	FMaterialColorParameters MaterialColorParameters = GetMaterialColorParameters(VIn, ShaderType);
+
+	BRANCH
+	if (MaterialColorParameters.bNeedMaterialColor)
+	{
+#if USE_MATERIALS
+		MaterialColorParameters.MaterialColor = GetMaterialColor(VIn, MaterialColorParameters.MeshUVs, MaterialColorParameters.SdfData);
+#endif
+	}
+
 	if (ShaderType == ST_Default || ShaderType == ST_Custom)
 	{
-		OutColor = GetDefaultElementColor(VIn);
+		OutColor = GetDefaultElementColor(VIn, MaterialColorParameters);
 	}
 	else if (ShaderType == ST_RoundedBox)
 	{
-		OutColor = GetRoundedBoxElementColor( VIn );
+		OutColor = GetRoundedBoxElementColor( VIn, MaterialColorParameters );
 	}
 	else if (ShaderType == ST_Border)
 	{
-		OutColor = GetBorderElementColor(VIn);
+		OutColor = GetBorderElementColor(VIn, MaterialColorParameters);
 	}
 	else if (ShaderType == ST_GrayscaleFont || ShaderType == ST_ColorFont)
 	{
-		OutColor = GetFontElementColor(VIn);
+		OutColor = GetFontElementColor(VIn, MaterialColorParameters);
 	}
 	else if (ShaderType == ST_SdfFont || ShaderType == ST_MsdfFont)
 	{
-		OutColor = GetSdfTextElementColor(VIn);
+		OutColor = GetSdfTextElementColor(VIn, MaterialColorParameters);
 	}
 	else
 	{
