Fix haloing artifact in MegaLights with Substrate Format=1 due to neighbor pixels reusing history having belonging to materials with different topology.

#Synced-CL 51204827

diff --git a/Engine/Shaders/Private/StochasticLighting/StochasticLightingTileClassification.usf b/Engine/Shaders/Private/StochasticLighting/StochasticLightingTileClassification.usf
--- a/Engine/Shaders/Private/StochasticLighting/StochasticLightingTileClassification.usf
+++ b/Engine/Shaders/Private/StochasticLighting/StochasticLightingTileClassification.usf
@@ -23,6 +23,8 @@ MAX_OCCUPANCY
 #define STOCHASTIC_SAMPLE_OFFSET_2x1 1
 #define STOCHASTIC_SAMPLE_OFFSET_2x2 2
 
+#define USE_CLOSURECOUNT_HISTORY_REJECTION (SUBSTRATE_ENABLED && SUBSTRATE_GBUFFER_FORMAT == 1 && MATERIAL_SOURCE == MATERIAL_SOURCE_GBUFFER)
+
 Texture2D DepthHistoryTexture;
 Texture2D<float4> NormalAndShadingInfoHistory;
 Texture2D<uint> ClosureHistory;
@@ -907,6 +909,13 @@ void StochasticLightingTileClassificationMarkCS(
 					// History might have invalid lighting pixels. Only use history lighting if it has valid data.
 					half4 FramesAccumulatedHistory4 = MegaLightsNumFramesAccumulatedHistory.GatherRed(GlobalPointClampedSampler, HistoryGatherUV).wzxy;
 					MegaLightsHistoryWeights = select(FramesAccumulatedHistory4 > 0.0, MegaLightsHistoryWeights, 0.0);
+
+					// Use Closure count as discriminent to avoid reusing history data of pixels with different material topology
+					// Concretely, this removes haloing artifact at materials' boundary
+					#if USE_CLOSURECOUNT_HISTORY_REJECTION
+					const uint4 PackedClosureCount = ClosureHistory.GatherRed(GlobalPointClampedSampler, HistoryGatherUV).wzxy;
+					MegaLightsHistoryWeights *= select(Material.ClosureCount == PackedClosureCount, half(1.0), half(0.0));
+					#endif
 				#endif
 				}
 
