Fix incorrect RWLumenTileBitmask writing during the overflow pass when all the pixels of the tile are invalid.

This was causing the tile (0,0) to be overriden with incorrect value, creating a square artifact on the top left corner of the screen.

#Synced-CL 51204827

diff --git a/Engine/Shaders/Private/StochasticLighting/StochasticLightingTileClassification.usf b/Engine/Shaders/Private/StochasticLighting/StochasticLightingTileClassification.usf
--- a/Engine/Shaders/Private/StochasticLighting/StochasticLightingTileClassification.usf
+++ b/Engine/Shaders/Private/StochasticLighting/StochasticLightingTileClassification.usf
@@ -633,6 +633,7 @@ void StochasticLightingTileClassificationMarkCS(
 	if (LinearThreadIndex == 0)
 	{
 		#if TILE_CLASSIFY_LUMEN
+		if (bIsAnyValid)
 		{
 			RWLumenTileBitmask[int3(Coord.SvPosition / TILE_SIZE, Coord.ClosureIndex)] = ReadSharedTileBitmask(0, TILE_BITMASK_INDEX_LUMEN);
 		}
