RHI: Experimental: Create specialization constant version of FLightGridInjectionCS (reduce shader count from 128 to 2).  Enabled through CVar, disabled by default.

#rnx

#Synced-CL 51094475

#Synced-CL 51094475

diff --git a/Engine/Config/Android/AndroidEngine.ini b/Engine/Config/Android/AndroidEngine.ini
--- a/Engine/Config/Android/AndroidEngine.ini
+++ b/Engine/Config/Android/AndroidEngine.ini
@@ -99,6 +99,8 @@ r.Vulkan.EnableTransientResourceAllocator=0
 r.Vulkan.AllowSparseResidency=0
 ; Use renames instead of copy+barrier on mobile devices
 r.Vulkan.AllowUniformUpload=0
+; Experiment to try specialization constants on a single shader (disabled by default)
+r.LightGridInjection.UseSpecializationConstants=0
 
 ; enable stripping debug strings from streamable manager in shipping to save memory
 s.StreamableStripDebugName=true
diff --git a/Engine/Shaders/Private/LightGridInjection.usf b/Engine/Shaders/Private/LightGridInjection.usf
--- a/Engine/Shaders/Private/LightGridInjection.usf
+++ b/Engine/Shaders/Private/LightGridInjection.usf
@@ -12,9 +12,9 @@
 #include "WaveOpUtil.ush"
 #include "GPUMessaging.ush"
 
-#if USE_HZB_CULL
-	#include "Nanite/NaniteHZBCull.ush"
-#endif
+// used by USE_HZB_CULL
+#include "Nanite/NaniteHZBCull.ush"
+
 
 #include "/Engine/Shared/LightGridDefinitions.h"
 #include "MegaLights/MegaLightsEstimateLight.ush"
@@ -59,28 +59,32 @@ StructuredBuffer<float4> LightViewSpaceRectPlanes;
 StructuredBuffer<float4> LightViewSpaceWeightEstimate;
 
 StructuredBuffer<int> IndirectionIndices;
+UE_DECL(APPLY_INDIRECTION);
+
+// All params used with USE_PARENT_LIGHT_GRID
+UE_DECL(USE_PARENT_LIGHT_GRID);
+StructuredBuffer<uint> ParentNumCulledLightsGrid;
+#if LIGHT_GRID_USES_16BIT_BUFFERS
+	Buffer<uint>						ParentCulledLightDataGrid16Bit;
+	#define ParentCulledLightDataGrid	ParentCulledLightDataGrid16Bit
+#else
+	StructuredBuffer<uint>				ParentCulledLightDataGrid32Bit;
+	#define ParentCulledLightDataGrid	ParentCulledLightDataGrid32Bit
+#endif
+uint3 ParentGridSize;
+uint NumParentGridCells;
+uint ParentGridSizeFactor;
 
-#if USE_PARENT_LIGHT_GRID
-	StructuredBuffer<uint> ParentNumCulledLightsGrid;
-#	if LIGHT_GRID_USES_16BIT_BUFFERS
-		Buffer<uint>						ParentCulledLightDataGrid16Bit;
-#		define ParentCulledLightDataGrid	ParentCulledLightDataGrid16Bit
-#	else
-		StructuredBuffer<uint>				ParentCulledLightDataGrid32Bit;
-#		define ParentCulledLightDataGrid	ParentCulledLightDataGrid32Bit
-#	endif
-	uint3 ParentGridSize;
-	uint NumParentGridCells;
-	uint ParentGridSizeFactor;
-#endif // USE_PARENT_LIGHT_GRID
 
 #define NUM_PLANES_PER_RECT_LIGHT 4
 
 #define REFINE_SPOTLIGHT_BOUNDS 1
 
-#ifndef REFINE_RECTLIGHT_BOUNDS
-#define REFINE_RECTLIGHT_BOUNDS 1
-#endif
+UE_DECL(REFINE_RECTLIGHT_BOUNDS);
+
+UE_DECL(USE_LINKED_CULL_LIST);
+
+UE_DECL(USE_HZB_CULL);
 
 #ifdef LightGridInjectionCS
 
@@ -108,7 +112,7 @@ float ComputeCellNearViewDepthFromZSlice(uint ZSlice, uint ZSliceScale)
 	return SliceDepth;
 }
 
-#if USE_HZB_CULL
+// used by USE_HZB_CULL
 FScreenRect ComputeCellCullRect(uint3 GridCoordinate, float MinTileZ, float MaxTileZ, float MarginPixels)
 {
 	const float2 TileSize = (1u << LightGridPixelSizeShift);
@@ -137,7 +141,6 @@ FScreenRect ComputeCellCullRect(uint3 GridCoordinate, float MinTileZ, float MaxT
 
 	return GetScreenRect(int4(0, 0, HZBViewSize), CullRectMin, CullRectMax, 4);
 }
-#endif
 
 struct FCellBounds
 {
@@ -296,7 +299,7 @@ bool OverlapsLight(uint LocalLightIndex, FCellBounds CellBounds, out bool bIsRec
 			}
 			#endif // REFINE_SPOTLIGHT_BOUNDS
 
-			#if REFINE_RECTLIGHT_BOUNDS
+			if(UE_EVAL(REFINE_RECTLIGHT_BOUNDS))
 			{
 				if (bIsRectLight)
 				{
@@ -317,12 +320,12 @@ bool OverlapsLight(uint LocalLightIndex, FCellBounds CellBounds, out bool bIsRec
 					}
 				}
 			}
-			#endif // REFINE_RECTLIGHT_BOUNDS
 		}
 
 		return bPassEstimatedWeightTest && bPassSpotLightTest && bPassRectLightTest;
 	}
 
+
 	bIsRectLight = false;
 	bIsTexturedLight = false;
 
@@ -372,28 +375,34 @@ struct FArrayView
 uint GetPrimitiveIndex(FArrayView ArrayView, uint Index)
 {
 	checkSlow(Index < ArrayView.Num);
-#if USE_PARENT_LIGHT_GRID
-	return ParentCulledLightDataGrid[ArrayView.Offset + Index];
-#else
-	return ArrayView.Offset + Index;
-#endif
+	if (UE_EVAL(USE_PARENT_LIGHT_GRID))
+	{
+		return ParentCulledLightDataGrid[ArrayView.Offset + Index];
+	}
+	else
+	{
+		return ArrayView.Offset + Index;
+	}
 }
 
 FArrayView InitLocalLightArrayView(uint ParentGridIndex)
 {
 	FArrayView ArrayView;
 
-#if USE_PARENT_LIGHT_GRID
-	uint Unused_NumVisibleMegaLights;
-	UnpackCulledLightsGridHeader0(ParentNumCulledLightsGrid[ParentGridIndex * 2 + 0], ArrayView.Num, Unused_NumVisibleMegaLights);
+	if (UE_EVAL(USE_PARENT_LIGHT_GRID))
+	{
+		uint Unused_NumVisibleMegaLights;
+		UnpackCulledLightsGridHeader0(ParentNumCulledLightsGrid[ParentGridIndex * 2 + 0], ArrayView.Num, Unused_NumVisibleMegaLights);
 
-	bool bHasRectLights_Unused;
-	bool bHasTexturedLights_Unused;
-	UnpackCulledLightsGridHeader1(ParentNumCulledLightsGrid[ParentGridIndex * 2 + 1], ArrayView.Offset, bHasRectLights_Unused, bHasTexturedLights_Unused);
-#else
-	ArrayView.Num = NumLocalLights;
-	ArrayView.Offset = 0;
-#endif
+		bool bHasRectLights_Unused;
+		bool bHasTexturedLights_Unused;
+		UnpackCulledLightsGridHeader1(ParentNumCulledLightsGrid[ParentGridIndex * 2 + 1], ArrayView.Offset, bHasRectLights_Unused, bHasTexturedLights_Unused);
+	}
+	else
+	{
+		ArrayView.Num = NumLocalLights;
+		ArrayView.Offset = 0;
+	}
 
 	return ArrayView;
 }
@@ -402,13 +411,16 @@ FArrayView InitReflectionCaptureArrayView(uint ParentGridIndex)
 {
 	FArrayView ArrayView;
 
-#if USE_PARENT_LIGHT_GRID
-	ArrayView.Num = ParentNumCulledLightsGrid[(NumParentGridCells + ParentGridIndex) * NUM_CULLED_LIGHTS_GRID_STRIDE + 0];
-	ArrayView.Offset = ParentNumCulledLightsGrid[(NumParentGridCells + ParentGridIndex) * NUM_CULLED_LIGHTS_GRID_STRIDE + 1];
-#else
-	ArrayView.Num = NumReflectionCaptures;
-	ArrayView.Offset = 0;
-#endif
+	if (UE_EVAL(USE_PARENT_LIGHT_GRID))
+	{
+		ArrayView.Num = ParentNumCulledLightsGrid[(NumParentGridCells + ParentGridIndex) * NUM_CULLED_LIGHTS_GRID_STRIDE + 0];
+		ArrayView.Offset = ParentNumCulledLightsGrid[(NumParentGridCells + ParentGridIndex) * NUM_CULLED_LIGHTS_GRID_STRIDE + 1];
+	}
+	else
+	{
+		ArrayView.Num = NumReflectionCaptures;
+		ArrayView.Offset = 0;
+	}
 
 	return ArrayView;
 }
@@ -438,59 +450,68 @@ void CullLights_SingleThread(uint GridIndex, FCellBounds CellBounds, uint Parent
 			bHasRectLights |= bIsRectLight;
 			bHasTexturedLights |= bIsTexturedLight;
 
-#if USE_LINKED_CULL_LIST
-			uint NextLink;
-			WaveInterlockedAddScalar_(RWCulledLightLinkAllocator[0], 1U, NextLink);
-
-			if (NextLink < NumAvailableLinks)
+			if (UE_EVAL(USE_LINKED_CULL_LIST))
 			{
-				RWCulledLightLinks[NextLink * LIGHT_LINK_STRIDE + 0] = LocalLightIndex;
-				RWCulledLightLinks[NextLink * LIGHT_LINK_STRIDE + 1] = LinkOffset;
-
-				LinkOffset = NextLink;
+				uint NextLink;
+				WaveInterlockedAddScalar_(RWCulledLightLinkAllocator[0], 1U, NextLink);
 
-				++NumVisibleLights;
-				if (LocalLightIndex >= MegaLightsSupportedStartIndex)
+				if (NextLink < NumAvailableLinks)
 				{
-					++NumVisibleMegaLights;
+					RWCulledLightLinks[NextLink * LIGHT_LINK_STRIDE + 0] = LocalLightIndex;
+					RWCulledLightLinks[NextLink * LIGHT_LINK_STRIDE + 1] = LinkOffset;
+
+					LinkOffset = NextLink;
+
+					++NumVisibleLights;
+					if (LocalLightIndex >= MegaLightsSupportedStartIndex)
+					{
+						++NumVisibleMegaLights;
+					}
 				}
 			}
-#else // !USE_LINKED_CULL_LIST
-			if (NumVisibleLights < MaxCulledLightsPerCell)
+			else
 			{
-				uint PrimitiveIndex = LocalLightIndex;
-#if APPLY_INDIRECTION
-				PrimitiveIndex = IndirectionIndices[PrimitiveIndex];
-#endif
-				RWCulledLightDataGrid[GridIndex * MaxCulledLightsPerCell + NumVisibleLights] = PrimitiveIndex;
-
-				++NumVisibleLights;
-				if (LocalLightIndex >= MegaLightsSupportedStartIndex)
+				if (NumVisibleLights < MaxCulledLightsPerCell)
 				{
-					++NumVisibleMegaLights;
+					uint PrimitiveIndex = LocalLightIndex;
+
+					if (UE_EVAL(APPLY_INDIRECTION))
+					{
+						PrimitiveIndex = IndirectionIndices[PrimitiveIndex];
+					}
+
+					RWCulledLightDataGrid[GridIndex * MaxCulledLightsPerCell + NumVisibleLights] = PrimitiveIndex;
+
+					++NumVisibleLights;
+					if (LocalLightIndex >= MegaLightsSupportedStartIndex)
+					{
+						++NumVisibleMegaLights;
+					}
 				}
 			}
-#endif // !USE_LINKED_CULL_LIST
 		}
 	}
 
-#if USE_LINKED_CULL_LIST
 	uint CulledLightDataStart;
-	#if FEATURE_LEVEL == FEATURE_LEVEL_ES3_1
-	// Adreno compiler fails to reg-alloc on 6xx for this shader,
-	// and doesn't support UGPR spilling, so just fails to compile.
-	InterlockedAdd(RWCulledLightDataAllocator[0], NumVisibleLights, CulledLightDataStart);
-	#else
-	WaveInterlockedAdd_(RWCulledLightDataAllocator[0], NumVisibleLights, CulledLightDataStart);
-	#endif
-
-	CulledLightDataStart += ViewCulledDataOffset;
-
-	// Mega Lights are order dependent
-	OutputReversedLinkedList(LinkOffset, NumVisibleLights, CulledLightDataStart, APPLY_INDIRECTION);
-#else
-	uint CulledLightDataStart = GridIndex * MaxCulledLightsPerCell;
-#endif
+	if (UE_EVAL(USE_LINKED_CULL_LIST))
+	{
+		#if FEATURE_LEVEL == FEATURE_LEVEL_ES3_1
+			// Adreno compiler fails to reg-alloc on 6xx for this shader,
+			// and doesn't support UGPR spilling, so just fails to compile.
+			InterlockedAdd(RWCulledLightDataAllocator[0], NumVisibleLights, CulledLightDataStart);
+		#else
+			WaveInterlockedAdd_(RWCulledLightDataAllocator[0], NumVisibleLights, CulledLightDataStart);
+		#endif
+
+		CulledLightDataStart += ViewCulledDataOffset;
+
+		// Mega Lights are order dependent
+		OutputReversedLinkedList(LinkOffset, NumVisibleLights, CulledLightDataStart, UE_EVAL(APPLY_INDIRECTION));
+	}
+	else
+	{
+		CulledLightDataStart = GridIndex * MaxCulledLightsPerCell;
+	}
 
 	RWNumCulledLightsGrid[GridIndex * NUM_CULLED_LIGHTS_GRID_STRIDE + 0] = PackCulledLightsGridHeader0(NumVisibleLights, NumVisibleMegaLights);
 	RWNumCulledLightsGrid[GridIndex * NUM_CULLED_LIGHTS_GRID_STRIDE + 1] = PackCulledLightsGridHeader1(CulledLightDataStart, bHasRectLights, bHasTexturedLights);
@@ -512,41 +533,47 @@ void CullReflectionCaptures_SingleThread(uint GridIndex, FCellBounds CellBounds,
 
 		if (bOverlapsReflectionCapture)
 		{
-#if USE_LINKED_CULL_LIST
-			uint NextLink;
-			WaveInterlockedAddScalar_(RWCulledLightLinkAllocator[0], 1U, NextLink);
-
-			if (NextLink < NumAvailableLinks)
+			if (UE_EVAL(USE_LINKED_CULL_LIST))
 			{
-				RWCulledLightLinks[NextLink * LIGHT_LINK_STRIDE + 0] = ReflectionCaptureIndex;
-				RWCulledLightLinks[NextLink * LIGHT_LINK_STRIDE + 1] = LinkOffset;
+				uint NextLink;
+				WaveInterlockedAddScalar_(RWCulledLightLinkAllocator[0], 1U, NextLink);
 
-				LinkOffset = NextLink;
+				if (NextLink < NumAvailableLinks)
+				{
+					RWCulledLightLinks[NextLink * LIGHT_LINK_STRIDE + 0] = ReflectionCaptureIndex;
+					RWCulledLightLinks[NextLink * LIGHT_LINK_STRIDE + 1] = LinkOffset;
 
-				++NumVisibleCaptures;
+					LinkOffset = NextLink;
+
+					++NumVisibleCaptures;
+				}
 			}
-#else // !USE_LINKED_CULL_LIST
-			if (NumVisibleCaptures < MaxCulledLightsPerCell)
+			else
 			{
-				RWCulledLightDataGrid[(NumGridCells + GridIndex) * MaxCulledLightsPerCell + NumVisibleCaptures] = ReflectionCaptureIndex;
+				if (NumVisibleCaptures < MaxCulledLightsPerCell)
+				{
+					RWCulledLightDataGrid[(NumGridCells + GridIndex) * MaxCulledLightsPerCell + NumVisibleCaptures] = ReflectionCaptureIndex;
 
-				++NumVisibleCaptures;
+					++NumVisibleCaptures;
+				}
 			}
-#endif // !USE_LINKED_CULL_LIST
 		}
 	}
 
-#if USE_LINKED_CULL_LIST
 	uint CulledReflectionCaptureDataStart;
-	WaveInterlockedAdd_(RWCulledLightDataAllocator[0], NumVisibleCaptures, CulledReflectionCaptureDataStart);
+	if (UE_EVAL(USE_LINKED_CULL_LIST))
+	{
+		WaveInterlockedAdd_(RWCulledLightDataAllocator[0], NumVisibleCaptures, CulledReflectionCaptureDataStart);
 
-	CulledReflectionCaptureDataStart += ViewCulledDataOffset;
+		CulledReflectionCaptureDataStart += ViewCulledDataOffset;
 
-	// Reflection captures are order dependent
-	OutputReversedLinkedList(LinkOffset, NumVisibleCaptures, CulledReflectionCaptureDataStart, false);
-#else
-	uint CulledReflectionCaptureDataStart = (NumGridCells + GridIndex) * MaxCulledLightsPerCell;
-#endif
+		// Reflection captures are order dependent
+		OutputReversedLinkedList(LinkOffset, NumVisibleCaptures, CulledReflectionCaptureDataStart, false);
+	}
+	else
+	{
+		CulledReflectionCaptureDataStart = (NumGridCells + GridIndex) * MaxCulledLightsPerCell;
+	}
 
 	RWNumCulledLightsGrid[(NumGridCells + GridIndex) * NUM_CULLED_LIGHTS_GRID_STRIDE + 0] = NumVisibleCaptures;
 	RWNumCulledLightsGrid[(NumGridCells + GridIndex) * NUM_CULLED_LIGHTS_GRID_STRIDE + 1] = CulledReflectionCaptureDataStart;
@@ -639,9 +666,10 @@ void OutputCell(FCellWriter CellWriter, uint NumPrimitives, uint CulledLightData
 		const uint IndexInCellH8 = (PackedData1 >> 24) & 0xFF;
 		const uint IndexInCell = (IndexInCellH8 << 8) | IndexInCellL8;
 
-#if APPLY_INDIRECTION
-		PrimitiveIndex = IndirectionIndices[PrimitiveIndex];
-#endif
+		if (UE_EVAL(APPLY_INDIRECTION))
+		{
+			PrimitiveIndex = IndirectionIndices[PrimitiveIndex];
+		}
 
 		RWCulledLightDataGrid[CulledLightDataStart + IndexInCell] = PrimitiveIndex;
 		CellWriter.LinkOffset = LinkOffset;
@@ -657,6 +685,8 @@ groupshared uint CellNumVisibleReflectionCaptures;
 
 groupshared uint CulledLightDataStart;
 
+UE_DECL(USE_THREAD_GROUP_PER_CELL);
+
 void CullLights_ThreadGroup(uint ThreadIndex, uint GridIndex, FCellBounds CellBounds, uint ParentGridIndex)
 {
 	if (ThreadIndex == 0)
@@ -737,7 +767,7 @@ void CullLights_ThreadGroup(uint ThreadIndex, uint GridIndex, FCellBounds CellBo
 
 	GroupMemoryBarrierWithGroupSync();
 
-	OutputCell(CellWriter, CellNumVisibleLights, CulledLightDataStart, ThreadIndex, APPLY_INDIRECTION);
+	OutputCell(CellWriter, CellNumVisibleLights, CulledLightDataStart, ThreadIndex, UE_EVAL(APPLY_INDIRECTION));
 
 	if (ThreadIndex == 0)
 	{
@@ -823,12 +853,8 @@ void LightGridInjectionCS(
 	uint3 GroupId : SV_GroupID,
 	uint GroupIndex : SV_GroupIndex)
 {
-#if USE_THREAD_GROUP_PER_CELL
-	const uint3 GridCoordinate = GroupId;
 	const uint ThreadIndex = GroupIndex;
-#else
-	const uint3 GridCoordinate = DispatchThreadId;
-#endif
+	const uint3 GridCoordinate = (UE_EVAL(USE_THREAD_GROUP_PER_CELL)) ? GroupId : DispatchThreadId;
 
 	if (all(GridCoordinate < (uint3)CulledGridSize))
 	{
@@ -841,61 +867,67 @@ void LightGridInjectionCS(
 		float MinTileZ = ComputeCellNearViewDepthFromZSlice(GridCoordinate.z + 0, LightGridZSliceScale);
 		float MaxTileZ = ComputeCellNearViewDepthFromZSlice(GridCoordinate.z + 1, LightGridZSliceScale);
 
-#if USE_HZB_CULL
-		float CullMinTileZ = MinTileZ;
-		float CullMaxTileZ = MaxTileZ;
-		if(LightGridCullMarginZ > 0)
+		if (UE_EVAL(USE_HZB_CULL))
 		{
-			// convert Min/MaxTileZ to volumetric fog grid coordinate
-			int MinCullSlice = (int)floor(ComputeZSliceFromDepth(LightGridCullMarginZParams, MinTileZ));
-			int MaxCullSlice = (int)ceil(ComputeZSliceFromDepth(LightGridCullMarginZParams, MaxTileZ));
-			// add margins
-			MinCullSlice = max(0, MinCullSlice - (int)LightGridCullMarginZ);
-			MaxCullSlice = min(LightGridCullMaxZ, MaxCullSlice + LightGridCullMarginZ);
-			// convert back to Min/MaxTileZ
-			CullMinTileZ = min(CullMinTileZ, ComputeDepthFromZSlice(LightGridCullMarginZParams, MinCullSlice));
-			CullMaxTileZ = max(CullMaxTileZ, ComputeDepthFromZSlice(LightGridCullMarginZParams, MaxCullSlice));
-		}
-
-		FScreenRect Rect = ComputeCellCullRect(GridCoordinate, CullMinTileZ, CullMaxTileZ, LightGridCullMarginXY);
-
- 		if (!IsVisibleHZB(Rect, true /*bSample4x4*/))
- 		{
-#if USE_THREAD_GROUP_PER_CELL
-			if (ThreadIndex == 0)
-#endif
+			float CullMinTileZ = MinTileZ;
+			float CullMaxTileZ = MaxTileZ;
+			if(LightGridCullMarginZ > 0)
 			{
- 				RWNumCulledLightsGrid[GridIndex * NUM_CULLED_LIGHTS_GRID_STRIDE + 0] = PackCulledLightsGridHeader0(0, 0);
- 				RWNumCulledLightsGrid[GridIndex * NUM_CULLED_LIGHTS_GRID_STRIDE + 1] = PackCulledLightsGridHeader1(0, false, false);
+				// convert Min/MaxTileZ to volumetric fog grid coordinate
+				int MinCullSlice = (int)floor(ComputeZSliceFromDepth(LightGridCullMarginZParams, MinTileZ));
+				int MaxCullSlice = (int)ceil(ComputeZSliceFromDepth(LightGridCullMarginZParams, MaxTileZ));
+				// add margins
+				MinCullSlice = max(0, MinCullSlice - (int)LightGridCullMarginZ);
+				MaxCullSlice = min(LightGridCullMaxZ, MaxCullSlice + LightGridCullMarginZ);
+				// convert back to Min/MaxTileZ
+				CullMinTileZ = min(CullMinTileZ, ComputeDepthFromZSlice(LightGridCullMarginZParams, MinCullSlice));
+				CullMaxTileZ = max(CullMaxTileZ, ComputeDepthFromZSlice(LightGridCullMarginZParams, MaxCullSlice));
 			}
- 			return;
- 		}
+
+			FScreenRect Rect = ComputeCellCullRect(GridCoordinate, CullMinTileZ, CullMaxTileZ, LightGridCullMarginXY);
+
+ 			if (!IsVisibleHZB(Rect, true /*bSample4x4*/))
+ 			{
+				if (!UE_EVAL(USE_THREAD_GROUP_PER_CELL) || ThreadIndex == 0)
+				{
+ 					RWNumCulledLightsGrid[GridIndex * NUM_CULLED_LIGHTS_GRID_STRIDE + 0] = PackCulledLightsGridHeader0(0, 0);
+ 					RWNumCulledLightsGrid[GridIndex * NUM_CULLED_LIGHTS_GRID_STRIDE + 1] = PackCulledLightsGridHeader1(0, false, false);
+				}
+ 				return;
+ 			}
 		
-		// tighten cell bounds based on farthest depth in cell
-		{
-			const float FarthestDeviceZ = GetMinDepthFromHZB(Rect, true /*bSample4x4*/);
-			const float FarthestSceneDepth = ConvertFromDeviceZ(FarthestDeviceZ);
+			// tighten cell bounds based on farthest depth in cell
+			{
+				const float FarthestDeviceZ = GetMinDepthFromHZB(Rect, true /*bSample4x4*/);
+				const float FarthestSceneDepth = ConvertFromDeviceZ(FarthestDeviceZ);
 
-			MaxTileZ = min(MaxTileZ, FarthestSceneDepth);
+				MaxTileZ = min(MaxTileZ, FarthestSceneDepth);
+			}
 		}
-#endif // USE_HZB_CULL
 		
 		FCellBounds CellBounds = ComputeCellBounds(GridCoordinate, MinTileZ, MaxTileZ);
 
-#if USE_PARENT_LIGHT_GRID
-		const uint3 ParentGridCoordinate = GridCoordinate / ParentGridSizeFactor;
-		const uint ParentGridIndex = (ParentGridCoordinate.z * ParentGridSize.y + ParentGridCoordinate.y) * ParentGridSize.x + ParentGridCoordinate.x;
-#else
-		const uint ParentGridIndex = 0xFFFFFFFF;
-#endif
+		uint ParentGridIndex;
+		if (UE_EVAL(USE_PARENT_LIGHT_GRID))
+		{
+			const uint3 ParentGridCoordinate = GridCoordinate / ParentGridSizeFactor;
+			ParentGridIndex = (ParentGridCoordinate.z * ParentGridSize.y + ParentGridCoordinate.y) * ParentGridSize.x + ParentGridCoordinate.x;
+		}
+		else
+		{
+			ParentGridIndex = 0xFFFFFFFF;
+		}
 
-#if USE_THREAD_GROUP_PER_CELL
-		CullLights_ThreadGroup(ThreadIndex, GridIndex, CellBounds, ParentGridIndex);
-		CullReflectionCaptures_ThreadGroup(ThreadIndex, GridIndex, CellBounds, ParentGridIndex);
-#else
-		CullLights_SingleThread(GridIndex, CellBounds, ParentGridIndex);
-		CullReflectionCaptures_SingleThread(GridIndex, CellBounds, ParentGridIndex);
-#endif
+		if (UE_EVAL(USE_THREAD_GROUP_PER_CELL))
+		{
+			CullLights_ThreadGroup(ThreadIndex, GridIndex, CellBounds, ParentGridIndex);
+			CullReflectionCaptures_ThreadGroup(ThreadIndex, GridIndex, CellBounds, ParentGridIndex);
+		}
+		else
+		{
+			CullLights_SingleThread(GridIndex, CellBounds, ParentGridIndex);
+			CullReflectionCaptures_SingleThread(GridIndex, CellBounds, ParentGridIndex);
+		}
 
 #else // !CULL_LIGHTS
 		LOOP
diff --git a/Engine/Source/Runtime/Renderer/Private/LightGridInjection.cpp b/Engine/Source/Runtime/Renderer/Private/LightGridInjection.cpp
--- a/Engine/Source/Runtime/Renderer/Private/LightGridInjection.cpp
+++ b/Engine/Source/Runtime/Renderer/Private/LightGridInjection.cpp
@@ -184,6 +184,20 @@ FAutoConsoleVariableRef CVarLightCullingMaxDistanceOverride(
 	ECVF_RenderThreadSafe
 );
 
+// Experimental version of FLightGridInjectionCS that uses specialization constants
+static TAutoConsoleVariable<int32> CVarLightGridInjectionUseSpecializationConstants(
+	TEXT("r.LightGridInjection.UseSpecializationConstants"),
+	0,
+	TEXT("Experimental: Selects whether LightGridInjection should use specialization constant. (default: 0)"),
+	ECVF_ReadOnly);
+
+static inline bool UseSpecializationConstants(EShaderPlatform ShaderPlatform)
+{
+	return FDataDrivenShaderPlatformInfo::GetSupportsSpecializationConstants(ShaderPlatform) &&
+		CVarLightGridInjectionUseSpecializationConstants.GetValueOnAnyThread() == 1;
+}
+
+
 DECLARE_GPU_STAT(LightGrid);
 
 #if !UE_BUILD_SHIPPING
@@ -450,6 +464,19 @@ public:
 		return true;
 	}
 
+	static EShaderPermutationPrecacheRequest ShouldPrecachePermutation(const FGlobalShaderPermutationParameters& Parameters)
+	{
+		// Do not precache if not in use
+		if (UseSpecializationConstants(Parameters.Platform))
+		{
+			return EShaderPermutationPrecacheRequest::NotPrecached;
+		}
+		else
+		{
+			return FGlobalShader::ShouldPrecachePermutation(Parameters);
+		}
+	}
+
 	static FIntVector GetGroupSize(FPermutationDomain PermutationVector)
 	{
 		if (PermutationVector.Get<FUseThreadGroupSize32>())
@@ -482,6 +509,75 @@ public:
 IMPLEMENT_GLOBAL_SHADER(FLightGridInjectionCS, "/Engine/Private/LightGridInjection.usf", "LightGridInjectionCS", SF_Compute);
 
 
+
+// Experimental version of FLightGridInjectionCS that uses specialization constants
+class FLightGridInjectionSpecConstCS : public FGlobalShader
+{
+	DECLARE_GLOBAL_SHADER(FLightGridInjectionSpecConstCS);
+	SHADER_USE_PARAMETER_STRUCT(FLightGridInjectionSpecConstCS, FGlobalShader)
+public:
+	class FUseLinkedList : SHADER_SPECIALIZATION_BOOL("USE_LINKED_CULL_LIST");
+	class FRefineRectLightBounds : SHADER_SPECIALIZATION_BOOL("REFINE_RECTLIGHT_BOUNDS");
+	class FUseHZBCull : SHADER_SPECIALIZATION_BOOL("USE_HZB_CULL");
+	class FUseParentLightGrid : SHADER_SPECIALIZATION_BOOL("USE_PARENT_LIGHT_GRID");
+	class FUseThreadGroupPerCell : SHADER_SPECIALIZATION_BOOL("USE_THREAD_GROUP_PER_CELL");
+	class FUseThreadGroupSize32 : SHADER_PERMUTATION_BOOL("USE_THREAD_GROUP_SIZE_32");
+	class FApplyIndirection : SHADER_SPECIALIZATION_BOOL("APPLY_INDIRECTION");
+	using FPermutationDomain = TShaderPermutationDomain<FUseLinkedList, FRefineRectLightBounds, FUseHZBCull, FUseParentLightGrid, FUseThreadGroupPerCell, FUseThreadGroupSize32, FApplyIndirection>;
+
+	using FParameters = FLightGridInjectionCS::FParameters;
+
+	static bool ShouldCompilePermutation(const FGlobalShaderPermutationParameters& Parameters)
+	{
+		return FDataDrivenShaderPlatformInfo::GetSupportsSpecializationConstants(Parameters.Platform);
+	}
+
+	static EShaderPermutationPrecacheRequest ShouldPrecachePermutation(const FGlobalShaderPermutationParameters& Parameters)
+	{
+		// Do not precache if not in use
+		if (UseSpecializationConstants(Parameters.Platform))
+		{
+			return FGlobalShader::ShouldPrecachePermutation(Parameters);
+		}
+		else
+		{
+			return EShaderPermutationPrecacheRequest::NotPrecached;
+		}
+	}
+
+	static FIntVector GetGroupSize(FPermutationDomain PermutationVector)
+	{
+		if (PermutationVector.Get<FUseThreadGroupSize32>())
+		{
+			return FIntVector(4, 4, 2);
+		}
+		else
+		{
+			return FIntVector(4, 4, 4);
+		}
+	}
+
+	static void ModifyCompilationEnvironment(const FGlobalShaderPermutationParameters& Parameters, FShaderCompilerEnvironment& OutEnvironment)
+	{
+		FGlobalShader::ModifyCompilationEnvironment(Parameters, OutEnvironment);
+		FForwardLightingParameters::ModifyCompilationEnvironment(Parameters.Platform, OutEnvironment);
+		OutEnvironment.SetDefine(TEXT("LIGHT_LINK_STRIDE"), LightLinkStride);
+
+		FPermutationDomain PermutationVector(Parameters.PermutationId);
+
+		const FIntVector GroupSize = GetGroupSize(PermutationVector);
+
+		OutEnvironment.SetDefine(TEXT("THREADGROUP_SIZE"), GroupSize.X * GroupSize.Y * GroupSize.Z);
+		OutEnvironment.SetDefine(TEXT("THREADGROUP_SIZE_X"), GroupSize.X);
+		OutEnvironment.SetDefine(TEXT("THREADGROUP_SIZE_Y"), GroupSize.Y);
+		OutEnvironment.SetDefine(TEXT("THREADGROUP_SIZE_Z"), GroupSize.Z);
+	}
+};
+
+IMPLEMENT_GLOBAL_SHADER(FLightGridInjectionSpecConstCS, "/Engine/Private/LightGridInjection.usf", "LightGridInjectionCS", SF_Compute);
+
+
+
 /**
  */
 FORCEINLINE float GetTanRadAngleOrZero(float coneAngle)
@@ -874,36 +970,61 @@ FLightGrid LightGridInjection(
 		PassParameters->HZBParameters = GetDummyHZBParameters(GraphBuilder);
 	}
 
-	FLightGridInjectionCS::FPermutationDomain PermutationVector;
-	PermutationVector.Set<FLightGridInjectionCS::FUseLinkedList>(GLightLinkedListCulling != 0);
-	PermutationVector.Set<FLightGridInjectionCS::FRefineRectLightBounds>(bRefineRectLightBounds);
-	PermutationVector.Set<FLightGridInjectionCS::FUseHZBCull>(GLightGridHZBCull != 0 && bIsHZBValid);
-	PermutationVector.Set<FLightGridInjectionCS::FUseParentLightGrid>(ParentNumCulledLightsGridSRV != nullptr && ParentCulledLightDataGridSRV != nullptr);
-	PermutationVector.Set<FLightGridInjectionCS::FUseThreadGroupPerCell>(bThreadGroupPerCell);
-	PermutationVector.Set<FLightGridInjectionCS::FUseThreadGroupSize32>(bThreadGroupSize32);
-	PermutationVector.Set<FLightGridInjectionCS::FApplyIndirection>(IndirectionIndicesSRV != nullptr);
-	auto ComputeShader = View.ShaderMap->GetShader<FLightGridInjectionCS>(PermutationVector);
-
-	FIntVector NumGroups;
-	if (bThreadGroupPerCell)
+	if (UseSpecializationConstants(View.GetShaderPlatform()))
 	{
-		NumGroups = GridSize;
+		FLightGridInjectionSpecConstCS::FPermutationDomain PermutationVector;
+		PermutationVector.Set<FLightGridInjectionSpecConstCS::FUseLinkedList>(GLightLinkedListCulling != 0);
+		PermutationVector.Set<FLightGridInjectionSpecConstCS::FRefineRectLightBounds>(bRefineRectLightBounds);
+		PermutationVector.Set<FLightGridInjectionSpecConstCS::FUseHZBCull>(GLightGridHZBCull != 0 && bIsHZBValid);
+		PermutationVector.Set<FLightGridInjectionSpecConstCS::FUseParentLightGrid>(ParentNumCulledLightsGridSRV != nullptr && ParentCulledLightDataGridSRV != nullptr);
+		PermutationVector.Set<FLightGridInjectionSpecConstCS::FUseThreadGroupPerCell>(bThreadGroupPerCell);
+		PermutationVector.Set<FLightGridInjectionSpecConstCS::FUseThreadGroupSize32>(bThreadGroupSize32);
+		PermutationVector.Set<FLightGridInjectionSpecConstCS::FApplyIndirection>(IndirectionIndicesSRV != nullptr);
+		auto ComputeShader = View.ShaderMap->GetShader<FLightGridInjectionSpecConstCS>(PermutationVector);
+
+		const FIntVector NumGroups = bThreadGroupPerCell ? GridSize : FComputeShaderUtils::GetGroupCount(GridSize, FLightGridInjectionSpecConstCS::GetGroupSize(PermutationVector));
+
+		FComputeShaderUtils::AddPass(
+			GraphBuilder,
+			RDG_EVENT_NAME("LightGridInject SpecConsts %s %s %s %s %s",
+				PermutationVector.Get<FLightGridInjectionSpecConstCS::FUseLinkedList>() ? TEXT("LinkedList") : TEXT("NoLinkedList"),
+				PermutationVector.Get<FLightGridInjectionSpecConstCS::FRefineRectLightBounds>() ? TEXT("RefineRectLightBounds") : TEXT(""),
+				PermutationVector.Get<FLightGridInjectionSpecConstCS::FUseHZBCull>() ? TEXT("UseHZBCull") : TEXT(""),
+				PermutationVector.Get<FLightGridInjectionSpecConstCS::FUseParentLightGrid>() ? TEXT("UseParentLightGrid") : TEXT(""),
+				PermutationVector.Get<FLightGridInjectionSpecConstCS::FUseThreadGroupPerCell>() ? TEXT("ThreadGroup") : TEXT("SingleThread")),
+			RDGPassFlags,
+			ComputeShader,
+			PassParameters,
+			NumGroups);
 	}
 	else
 	{
-		NumGroups = FComputeShaderUtils::GetGroupCount(GridSize, FLightGridInjectionCS::GetGroupSize(PermutationVector));
+		FLightGridInjectionCS::FPermutationDomain PermutationVector;
+		PermutationVector.Set<FLightGridInjectionCS::FUseLinkedList>(GLightLinkedListCulling != 0);
+		PermutationVector.Set<FLightGridInjectionCS::FRefineRectLightBounds>(bRefineRectLightBounds);
+		PermutationVector.Set<FLightGridInjectionCS::FUseHZBCull>(GLightGridHZBCull != 0 && bIsHZBValid);
+		PermutationVector.Set<FLightGridInjectionCS::FUseParentLightGrid>(ParentNumCulledLightsGridSRV != nullptr && ParentCulledLightDataGridSRV != nullptr);
+		PermutationVector.Set<FLightGridInjectionCS::FUseThreadGroupPerCell>(bThreadGroupPerCell);
+		PermutationVector.Set<FLightGridInjectionCS::FUseThreadGroupSize32>(bThreadGroupSize32);
+		PermutationVector.Set<FLightGridInjectionCS::FApplyIndirection>(IndirectionIndicesSRV != nullptr);
+		auto ComputeShader = View.ShaderMap->GetShader<FLightGridInjectionCS>(PermutationVector);
+
+		const FIntVector NumGroups = bThreadGroupPerCell ? GridSize : FComputeShaderUtils::GetGroupCount(GridSize, FLightGridInjectionCS::GetGroupSize(PermutationVector));
+
+		FComputeShaderUtils::AddPass(
+			GraphBuilder,
+			RDG_EVENT_NAME("LightGridInject %s %s %s %s %s",
+				PermutationVector.Get<FLightGridInjectionCS::FUseLinkedList>() ? TEXT("LinkedList") : TEXT("NoLinkedList"),
+				PermutationVector.Get<FLightGridInjectionCS::FRefineRectLightBounds>() ? TEXT("RefineRectLightBounds") : TEXT(""),
+				PermutationVector.Get<FLightGridInjectionCS::FUseHZBCull>() ? TEXT("UseHZBCull") : TEXT(""),
+				PermutationVector.Get<FLightGridInjectionCS::FUseParentLightGrid>() ? TEXT("UseParentLightGrid") : TEXT(""),
+				PermutationVector.Get<FLightGridInjectionCS::FUseThreadGroupPerCell>() ? TEXT("ThreadGroup") : TEXT("SingleThread")),
+			RDGPassFlags,
+			ComputeShader,
+			PassParameters,
+			NumGroups);
 	}
 
-	FComputeShaderUtils::AddPass(
-		GraphBuilder,
-		RDG_EVENT_NAME("LightGridInject %s %s",
-			PermutationVector.Get<FLightGridInjectionCS::FUseLinkedList>() ? TEXT("LinkedList") : TEXT("NoLinkedList"),
-			PermutationVector.Get<FLightGridInjectionCS::FUseThreadGroupPerCell>() ? TEXT("ThreadGroup") : TEXT("SingleThread")),
-    	RDGPassFlags,
-		ComputeShader,
-		PassParameters,
-		NumGroups);
-
 	FLightGrid Output;
 	Output.CulledLightDataGridSRV = CulledLightDataGridSRV;
 	Output.NumCulledLightsGridSRV = NumCulledLightsGridSRV;
