[PostProcessing] Fix crash in VisualizeLocalExposure when instanced stereo is enabled due to mismatch in SV_POSITION interpolator register index.

FVisualizeLocalExposurePS is combined with FScreenPassVS, which produces extra interpolators in FStereoVSOutput if INSTANCED_STEREO is set, causing the index of SV_POSITION to shift.

#Synced-CL 51099840

diff --git a/Engine/Shaders/Private/PostProcessVisualizeLocalExposure.usf b/Engine/Shaders/Private/PostProcessVisualizeLocalExposure.usf
--- a/Engine/Shaders/Private/PostProcessVisualizeLocalExposure.usf
+++ b/Engine/Shaders/Private/PostProcessVisualizeLocalExposure.usf
@@ -46,7 +46,7 @@ float DiagonalStripesPattern(float2 ScreenPos, float Thickness)
     return step(0.5f, V);
 }
 
-float4 MainPS(noperspective float4 UVAndScreenPos : TEXCOORD0, float4 SvPosition : SV_POSITION) : SV_Target0
+float4 MainPS(noperspective float4 UVAndScreenPos : TEXCOORD0, FStereoVSOutput Stereo, float4 SvPosition : SV_POSITION) : SV_Target0
 {
 	float2 UV = UVAndScreenPos.xy;
 	float2 ViewportUV = (int2(SvPosition.xy) - Output_ViewportMin) * Output_ViewportSizeInverse;
