Add common functions for blue noise sampling to reuse sampling logic accross View and BlueNoise resources.

#Synced-CL 51333038

diff --git a/Engine/Shaders/Private/BlueNoise.ush b/Engine/Shaders/Private/BlueNoise.ush
--- a/Engine/Shaders/Private/BlueNoise.ush
+++ b/Engine/Shaders/Private/BlueNoise.ush
@@ -2,6 +2,8 @@
 
 #pragma once
 
+#include "BlueNoiseCommon.ush"
+
 /**
  * Spatiotemporal Blue Noise lookup texture with scalar output.  
  * Returns a uniformly distributed random variable in [0, 1] that is blue (no low frequencies) in both xy (screen space) and in z (frame index), which is designed to work with temporal accumulation.
@@ -10,9 +12,7 @@
  */
 float BlueNoiseScalar(uint2 ScreenCoord, uint FrameIndex)
 {
-	uint3 WrappedCoordinate = uint3(ScreenCoord, FrameIndex) & BlueNoise.ModuloMasks;
-	uint3 TextureCoordinate = uint3(WrappedCoordinate.x, WrappedCoordinate.z * BlueNoise.Dimensions.y + WrappedCoordinate.y, 0);
-	return BlueNoise.ScalarTexture.Load(TextureCoordinate, 0).x;
+	return GetBlueNoiseScalar(BlueNoise.ScalarTexture, BlueNoise.Dimensions, BlueNoise.ModuloMasks, ScreenCoord, FrameIndex);
 }
 
 /**
@@ -22,7 +22,5 @@ float BlueNoiseScalar(uint2 ScreenCoord, uint FrameIndex)
  */
 float2 BlueNoiseVec2(uint2 ScreenCoord, uint FrameIndex)
 {
-	uint3 WrappedCoordinate = uint3(ScreenCoord, FrameIndex) & BlueNoise.ModuloMasks;
-	uint3 TextureCoordinate = uint3(WrappedCoordinate.x, WrappedCoordinate.z * BlueNoise.Dimensions.y + WrappedCoordinate.y, 0);
-	return BlueNoise.Vec2Texture.Load(TextureCoordinate, 0).xy;
+	return GetBlueNoiseVec2(BlueNoise.Vec2Texture, BlueNoise.Dimensions, BlueNoise.ModuloMasks, ScreenCoord, FrameIndex);
 }
\ No newline at end of file
diff --git a/Engine/Shaders/Private/BlueNoiseCommon.ush b/Engine/Shaders/Private/BlueNoiseCommon.ush
new file mode 100644
--- /dev/null
+++ b/Engine/Shaders/Private/BlueNoiseCommon.ush
@@ -0,0 +1,26 @@
+// Copyright Epic Games, Inc. All Rights Reserved.
+
+#pragma once
+
+float GetBlueNoiseScalar(Texture2D<float> InTexture, int3 InDimensions, int3 InModuloMasks, uint2 InScreenCoord, uint InFrameIndex)
+{
+	uint3 WrappedCoordinate = uint3(InScreenCoord, InFrameIndex) & InModuloMasks;
+	uint3 TextureCoordinate = uint3(WrappedCoordinate.x, WrappedCoordinate.z * InDimensions.y + WrappedCoordinate.y, 0);
+	return InTexture.Load(TextureCoordinate, 0).x;
+}
+
+float2 GetBlueNoiseVec2(Texture2D<float2> InTexture, int3 InDimensions, int3 InModuloMasks, uint2 InScreenCoord, uint InFrameIndex)
+{
+	uint3 WrappedCoordinate = uint3(InScreenCoord, InFrameIndex) & InModuloMasks;
+	uint3 TextureCoordinate = uint3(WrappedCoordinate.x, WrappedCoordinate.z * BlueNoise.Dimensions.y + WrappedCoordinate.y, 0);
+	return InTexture.Load(TextureCoordinate, 0).xy;
+}
+
+float GetBlueNoiseDitheredCoverage(Texture2D<float> InTexture, int3 InDimensions, int3 InModuloMasks, uint2 InScreenCoord, uint InFrameIndex, float InCoverage)
+{
+	// New dithering based on Blue noise
+	const float Noise2 = GetBlueNoiseScalar(InTexture, InDimensions, InModuloMasks, InScreenCoord, InFrameIndex);
+	const float LegacyAdjustement = 0.83f; // This is required to get the same look and empty areas as the legacy one.
+	float Dither = LegacyAdjustement * Noise2;
+	return InCoverage + Dither - 0.5f;
+}
\ No newline at end of file
diff --git a/Engine/Shaders/Private/MaterialTemplate.ush b/Engine/Shaders/Private/MaterialTemplate.ush
--- a/Engine/Shaders/Private/MaterialTemplate.ush
+++ b/Engine/Shaders/Private/MaterialTemplate.ush
@@ -42,6 +42,7 @@
 #include "/Engine/Private/ColorSpace.ush"
 #include "/Engine/Private/VirtualTextureMaterial.ush"
 #include "/Engine/Private/MaterialCache/MaterialCacheMaterialCommon.ush"
+#include "/Engine/Private/BlueNoiseCommon.ush"
 
 #if MATERIAL_SKY_ATMOSPHERE && PROJECT_SUPPORT_SKY_ATMOSPHERE
 #include "/Engine/Private/Quantization.ush"
@@ -2042,11 +2043,14 @@ float3 RotateAboutAxis(float4 NormalizedRotationAxisAndAngle, FWSVector3 Positio
 	return R - UAxis;
 }
 
-float ViewScalarBlueNoise(uint2 ScreenCoord, uint FrameIndex)
+float ViewScalarBlueNoise(uint2 InScreenCoord, uint InFrameIndex)
 {
-	uint3 WrappedCoordinate = uint3(ScreenCoord, FrameIndex) & View.BlueNoiseModuloMasks;
-	uint3 TextureCoordinate = uint3(WrappedCoordinate.x, WrappedCoordinate.z * View.BlueNoiseDimensions.y + WrappedCoordinate.y, 0);
-	return View.BlueNoiseScalarTexture.Load(TextureCoordinate, 0).x;
+	return GetBlueNoiseScalar(View.BlueNoiseScalarTexture, View.BlueNoiseDimensions, View.BlueNoiseModuloMasks, InScreenCoord, InFrameIndex);
+}
+
+float GetDitheredCoverage(uint2 InScreenCoord, uint InFrameIndex, float InCoverage)
+{
+	return GetBlueNoiseDitheredCoverage(View.BlueNoiseScalarTexture, View.BlueNoiseDimensions, View.BlueNoiseModuloMasks, InScreenCoord, InFrameIndex, InCoverage);
 }
 
 // Material Expression function
@@ -4535,12 +4539,7 @@ float DitheredOpacityMaskToOpacity(float Mask)
 float GetDitheredMaterialMask(FPixelMaterialInputs PixelMaterialInputs, uint2 SvPosition, uint FrameIndex)
 {
 	#if MATERIAL_DITHER_OPACITY_MASK
-		// New dithering based on Blue noise
-		float Noise2  = ViewScalarBlueNoise(SvPosition.xy, FrameIndex);
-		const float LegacyAdjustement = 0.83f; // This is required to get the same look and empty areas as the legacy one.
-		float Dither = LegacyAdjustement * Noise2;
-
-		return GetMaterialMask(PixelMaterialInputs) + Dither - 0.5f;
+		return GetDitheredCoverage(SvPosition.xy, FrameIndex, GetMaterialMask(PixelMaterialInputs));
 	#else
 		return GetMaterialMask(PixelMaterialInputs);
 	#endif
diff --git a/Engine/Source/Runtime/Renderer/Private/BlueNoise.h b/Engine/Source/Runtime/Renderer/Private/BlueNoise.h
--- a/Engine/Source/Runtime/Renderer/Private/BlueNoise.h
+++ b/Engine/Source/Runtime/Renderer/Private/BlueNoise.h
@@ -12,8 +12,8 @@
 BEGIN_SHADER_PARAMETER_STRUCT(FBlueNoiseParameters, RENDERER_API)
 	SHADER_PARAMETER(FIntVector, Dimensions)
 	SHADER_PARAMETER(FIntVector, ModuloMasks)
-	SHADER_PARAMETER_TEXTURE(Texture2D, ScalarTexture)
-	SHADER_PARAMETER_TEXTURE(Texture2D, Vec2Texture)
+	SHADER_PARAMETER_TEXTURE(Texture2D<float>, ScalarTexture)
+	SHADER_PARAMETER_TEXTURE(Texture2D<float2>, Vec2Texture)
 END_SHADER_PARAMETER_STRUCT()
 
 BEGIN_GLOBAL_SHADER_PARAMETER_STRUCT(FBlueNoise, RENDERER_API)
