Fixed incorrect translucent lighting on mobile renderer when MERGED_LOCAL_LIGHTS_MOBILE is set to 2 (Local Lights Buffer Enabled)

MergeLocalLights was collecting all lights within range but wasn't checking lighting channels, so it merged and applied lights that shouldn't affect the asset (assets/lights on different channels).

#rnx

#Synced-CL 51328445

diff --git a/Engine/Shaders/Private/MobileBasePassPixelShader.usf b/Engine/Shaders/Private/MobileBasePassPixelShader.usf
--- a/Engine/Shaders/Private/MobileBasePassPixelShader.usf
+++ b/Engine/Shaders/Private/MobileBasePassPixelShader.usf
@@ -1121,12 +1121,17 @@ void Main(
 	{
 		const FCulledLightsGridHeader CulledLightGridHeader = GetCulledLightsGridHeader(GridIndex);
 
+		uint LightingChannelMask = 0;
+#if !MOBILE_DEFERRED_LIGHTING
+		LightingChannelMask = GetPrimitive_LightingChannelMask(MaterialParameters.PrimitiveId);
+#endif
+
 		float3 MergedLightColor = uint3(0.f, 0.f , 0.f);
 		float3 MergedLightL = float3(0,0,0);
 		float MergedSpecularScale = 0.f;
 		float MergedTotalWeight = 0;
-		MergeLocalLights(CulledLightGridHeader, MaterialParameters.WorldPosition_CamRelative, EyeIndex, false, MergedLightColor, MergedLightL, MergedSpecularScale, MergedTotalWeight, GBuffer.WorldNormal);
-
+		MergeLocalLights(CulledLightGridHeader, MaterialParameters.WorldPosition_CamRelative, EyeIndex, false, MergedLightColor, MergedLightL, MergedSpecularScale, MergedTotalWeight, GBuffer.WorldNormal, LightingChannelMask);
+		
 		if(MergedTotalWeight > 0.f)
 		{
 			MergedLightL = MergedLightL / MergedTotalWeight;
diff --git a/Engine/Shaders/Private/MobileLightingCommon.ush b/Engine/Shaders/Private/MobileLightingCommon.ush
--- a/Engine/Shaders/Private/MobileLightingCommon.ush
+++ b/Engine/Shaders/Private/MobileLightingCommon.ush
@@ -687,6 +687,7 @@ void MergeLocalLights(const FCulledLightsGridHeader InLightGridHeader
 #if ENABLE_NDOTL_INTEGRATION
 	, inout float3 N
 #endif
+	, uint LightingChannelMask
 )
 {
 	MergedLightColor = uint3(0.f, 0.f , 0.f);
@@ -701,6 +702,14 @@ void MergeLocalLights(const FCulledLightsGridHeader InLightGridHeader
 	{
 		const FLocalLightData LocalLight = GetLocalLightDataFromGrid(InLightGridHeader.DataStartIndex + GridLightListIndex, EyeIndex);
 
+#if !MOBILE_DEFERRED_LIGHTING
+		// Check lighting channels for mobile forward
+		if ((UnpackLightingChannelMask(LocalLight) & LightingChannelMask) == 0)
+		{
+			continue;
+		}
+#endif
+
 		// extra-early out since we know light grid is sloppy and all lights in list are radial (have a range)
 		// appears useless
 		float InvLightRadiusSq = UnpackLightInvRadius(LocalLight) * UnpackLightInvRadius(LocalLight);
diff --git a/Engine/Shaders/Private/MobileLocalLightsBuffer.usf b/Engine/Shaders/Private/MobileLocalLightsBuffer.usf
--- a/Engine/Shaders/Private/MobileLocalLightsBuffer.usf
+++ b/Engine/Shaders/Private/MobileLocalLightsBuffer.usf
@@ -113,7 +113,7 @@ void Main(float4 SvPosition : SV_POSITION
 	float3 MergedLightL = float3(0,0,0);
 	float MergedSpecularScale = 0.f;
 	float MergedTotalWeight = 0;
-	MergeLocalLights(CulledLightsGridHeader, TranslatedWorldPosition, EyeIndex, true, MergedLightColor, MergedLightL, MergedSpecularScale, MergedTotalWeight);
+	MergeLocalLights(CulledLightsGridHeader, TranslatedWorldPosition, EyeIndex, true, MergedLightColor, MergedLightL, MergedSpecularScale, MergedTotalWeight, LIGHTING_CHANNEL_MASK);
 
 	if (MergedTotalWeight > 0)
 	{
diff --git a/Engine/Shaders/Private/Substrate/SubstrateMobileForwardLighting.ush b/Engine/Shaders/Private/Substrate/SubstrateMobileForwardLighting.ush
--- a/Engine/Shaders/Private/Substrate/SubstrateMobileForwardLighting.ush
+++ b/Engine/Shaders/Private/Substrate/SubstrateMobileForwardLighting.ush
@@ -331,7 +331,13 @@ float3 SubstrateMobileForwardLighting(
 					float3 MergedLightL = float3(0,0,0);
 					float MergedSpecularScale = 0.f;
 					float MergedTotalWeight = 0;
-					MergeLocalLights(CulledLightsGridHeader, MaterialParameters.WorldPosition_CamRelative, EyeIndex, false/*bSupportLightFunctions*/, MergedLightColor, MergedLightL, MergedSpecularScale, MergedTotalWeight, SubstrateBSDFContext.N);
+
+					uint LightingChannelMask = 0;
+					#if !MOBILE_DEFERRED_LIGHTING
+						LightingChannelMask = GetPrimitive_LightingChannelMask(MaterialParameters.PrimitiveId);
+					#endif
+
+					MergeLocalLights(CulledLightsGridHeader, MaterialParameters.WorldPosition_CamRelative, EyeIndex, false/*bSupportLightFunctions*/, MergedLightColor, MergedLightL, MergedSpecularScale, MergedTotalWeight, SubstrateBSDFContext.N, LightingChannelMask);
 
 					if (MergedTotalWeight > 0.f)
 					{
