Path Tracer: Optimize 3rd dimension of the random sequence

Take advantage of the diagonal structure inherent in the SZ construction to compute the z component with fewer instructions. This new implementation is reordered compared to the original matrix so both z and w components will return different values

Remove older versions of the sampler which are no longer used

Since the pattern is changing anyway, update the sampler usage for DOF to use a 4D sample instead of two 2D samples. With the current sampler, this results in slightly better quality.

#Synced-CL 51391139

diff --git a/Engine/Shaders/Private/PathTracing/PathTracingCore.ush b/Engine/Shaders/Private/PathTracing/PathTracingCore.ush
--- a/Engine/Shaders/Private/PathTracing/PathTracingCore.ush
+++ b/Engine/Shaders/Private/PathTracing/PathTracingCore.ush
@@ -2,7 +2,6 @@
 
 #pragma once
 
-#define RANDSEQ_VERSION 1
 #define RANDSEQ_SFC_TEXTURE 1
 
 #define USE_PATH_TRACING_LIGHT_GRID	1
@@ -1642,9 +1641,10 @@ FPathState CreatePathState(int2 PixelIndex, int2 TextureIndex)
 #endif
 
 	// Initialize ray and payload
+	float4 Jitter = Output.RandSequence.Get4D();
 	float2 ViewportUV;
 	{
-		float2 AAJitter = Output.RandSequence.Get2D();
+		float2 AAJitter = Jitter.xy;
 #if 0
 		// Tent filter
 		float2 P = 2 * AAJitter - 0.99999994f;
@@ -1663,7 +1663,7 @@ FPathState CreatePathState(int2 PixelIndex, int2 TextureIndex)
 	if (CameraLensRadius.y > 0)
 	{
 	    // DOF enabled - apply simple thin-lens model
-		float2 LensSample = Output.RandSequence.Get2D();
+		float2 LensSample = Jitter.zw;
 		float3 ViewX = View.ViewToTranslatedWorld[0].xyz;
 		float3 ViewY = View.ViewToTranslatedWorld[1].xyz;
 		float3 ViewZ = View.ViewToTranslatedWorld[2].xyz;
@@ -1750,7 +1750,7 @@ bool PathTracingKernel(inout FPathState PathState, int Bounce)
 		if (LightPickingCdfSum > 0.0)
 		{
 			// at least one light is overlapping with our ray, so we have a chance to sample it
-			float3 RandSample = PathState.RandSequence.Get3D();
+			float2 RandSample = PathState.RandSequence.Get2D();
 			int LightSampleLightId = 0;
 			float LightPickPdf = 0.0;
 			SelectLight(RandSample.x * LightPickingCdfSum, NumLights, LightPickingCdf, LightSampleLightId, LightPickPdf);
diff --git a/Engine/Shaders/Private/PathTracing/PathTracingDebug.usf b/Engine/Shaders/Private/PathTracing/PathTracingDebug.usf
--- a/Engine/Shaders/Private/PathTracing/PathTracingDebug.usf
+++ b/Engine/Shaders/Private/PathTracing/PathTracingDebug.usf
@@ -1,6 +1,5 @@
 // Copyright Epic Games, Inc. All Rights Reserved.
 
-#define RANDSEQ_VERSION 1
 #define RANDSEQ_SFC_TEXTURE 1
 
 #define LIGHT_LOOP_METHOD 	2 // 0: stochastic (like ref PT) 1: loop over all lights 2: RIS (take N samples but trace a single shadow ray)
diff --git a/Engine/Shaders/Private/PathTracing/Utilities/PathTracingRandomSequence.ush b/Engine/Shaders/Private/PathTracing/Utilities/PathTracingRandomSequence.ush
--- a/Engine/Shaders/Private/PathTracing/Utilities/PathTracingRandomSequence.ush
+++ b/Engine/Shaders/Private/PathTracing/Utilities/PathTracingRandomSequence.ush
@@ -29,9 +29,6 @@
 #define RANDSEQ_UNROLL_SOBOL 1 // Use unrolled Sobol loop such that table can be inline constants.
 #endif
 
-#ifndef RANDSEQ_VERSION
-#define RANDSEQ_VERSION 1 // TODO: Remove code for old version (0) once all testsuite images have been updated
-#endif
 
 #ifndef RANDSEQ_SFC_TEXTURE
 #define RANDSEQ_SFC_TEXTURE 0 // 0 uses a Hilbert curve (ALU heavy) while 1 uses a precomputed texture (must be passed as a shader parameter)
@@ -56,7 +53,6 @@ uint PerfectIntegerHash(uint x)
 uint StrongIntegerHash(uint x)
 {
 	// From [REDACTED_PROFILE_URL]/hash-prospector
-#if RANDSEQ_VERSION == 1
 	// Current best hash in this form: [REDACTED_PROFILE_URL]/hash-prospector/issues/19#issuecomment-1105792898
 	// bias = 0.10734781817103507
 	x ^= x >> 16;
@@ -64,14 +60,6 @@ uint StrongIntegerHash(uint x)
 	x ^= x >> 15;
 	x *= 0xf35a2d97;
 	x ^= x >> 15;
-#else
-	// bias = 0.16540778981744320
-	x ^= x >> 16;
-	x *= 0xa812d533;
-	x ^= x >> 15;
-	x *= 0xb278e4ad;
-	x ^= x >> 17;
-#endif
 	return x;
 }
 
@@ -102,47 +90,12 @@ uint FastOwenScramblingCore(uint Index, uint Seed)
 		Seed ^= Input & Mask;  // flip state  (depending on input)
 	}
 	return Index;
-#elif RANDSEQ_VERSION == 1
+#else
 	// This simple tweak appears to be nearly as good as the reference, while being cheaper to compute
 	Index ^= Index * 0xe0705d72u;
 	Index += Seed;
 	Seed ^= Seed >> 16;
 	Index *= Seed | 1;
-#elif 1
-	// Laine and Karras / Stratified Sampling for Stochastic Transparency / EGSR 2011
-
-	// The operations below will mix bits toward the left, so temporarily reverse the order
-	// NOTE: This operation has been performed outside this call
-	// Index = reversebits(Index);
-	
-	// This follows the basic construction from the paper above. The error-diffusion sampler which
-	// tiles a single point set across the screen, makes it much easier to visually "see" the impact of the scrambling.
-	// When the scrambling is not good enough, the structure of the space-filling curve is left behind.
-	// Care must be taken to ensure that the scrambling is still unbiased on average though. I found some cases
-	// that seemed to produce lower visual error but were actually biased due to not touching certain bits,
-	// leading to correlations across dimensions.
-
-	// After much experimentation with the hash prospector, I discovered the following two constants which appear to
-	// give results as good (or better) than the original 4 xor-mul constants from the Laine/Karras paper.
-	// It isn't entirely clear to me why some constants work better than others. Some hashes with slightly less
-    // bias produced visibly more error or worse looking power spectra.
-	// Estimates below from hash prospector for all hashes of the form: add,xmul=c0,xmul=c1 (with c0 and c1 being
-	// the constants below).
-	// Ran with score_quality=16 for about ~10000 random hashes
-	// Average bias: ~727.02
-	//   Best  bias: ~723.05
-	//   Worst bias: ~735.19
-	Index += Seed; // randomize the index by our seed (pushes bits toward the left)
-	Index ^= Index * 0x9c117646u;
-	Index ^= Index * 0xe0705d72u;
-#else
-	// This is a higher quality owen-scrambling hash but doesn't appear to make much difference
-	// https://psychopath.io/post/2021_01_30_building_a_better_lk_hash
-	Index ^= Index * 0x3d20adeau;
-	Index += Seed;
-	Index *= (Seed >> 16) | 1;
-	Index ^= Index * 0x05526c56u;
-	Index ^= Index * 0x53a22864u;
 #endif
 	return Index;
 }
@@ -185,78 +138,15 @@ uint SFCInverse(uint X, uint Y, int NumBits)
 }
 #endif
 
-#if RANDSEQ_VERSION == 1
-// Matrices handled explicitly in SamplerCore
-#else
-// 32-bit Sobol matrices for dimension 2,3 from:
-// S. Joe and F. Y. Kuo, Constructing Sobol sequences with better two-dimensional projections, SIAM J. Sci. Comput. 30, 2635-2654 (2008)
-//    https://web.maths.unsw.edu.au/~fkuo/sobol/
-// NOTE: we don't bother storing dimension 0 since it is just a bit reversal
-// NOTE2: we don't bother storing dimension 1 either since it has a very simple pattern
-// NOTE3: the matrix elements are reversed to save one reverse in the owen scrambling
-static const uint2 SobolMatrices[32] = {
-	uint2(0x00000001, 0x00000001),
-	uint2(0x00000003, 0x00000003),
-	uint2(0x00000006, 0x00000004),
-	uint2(0x00000009, 0x0000000a),
-	uint2(0x00000017, 0x0000001f),
-	uint2(0x0000003a, 0x0000002e),
-	uint2(0x00000071, 0x00000045),
-	uint2(0x000000a3, 0x000000c9),
-	uint2(0x00000116, 0x0000011b),
-	uint2(0x00000339, 0x000002a4),
-	uint2(0x00000677, 0x0000079a),
-	uint2(0x000009aa, 0x00000b67),
-	uint2(0x00001601, 0x0000101e),
-	uint2(0x00003903, 0x0000302d),
-	uint2(0x00007706, 0x00004041),
-	uint2(0x0000aa09, 0x0000a0c3),
-	uint2(0x00010117, 0x0001f104),
-	uint2(0x0003033a, 0x0002e28a),
-	uint2(0x00060671, 0x000457df),
-	uint2(0x000909a3, 0x000c9bae),
-	uint2(0x00171616, 0x0011a105),
-	uint2(0x003a3939, 0x002a7289),
-	uint2(0x00717777, 0x0079e7db),
-	uint2(0x00a3aaaa, 0x00b6dba4),
-	uint2(0x01170001, 0x0100011a),
-	uint2(0x033a0003, 0x030002a7),
-	uint2(0x06710006, 0x0400079e),
-	uint2(0x09a30009, 0x0a000b6d),
-	uint2(0x16160017, 0x1f001001),
-	uint2(0x3939003a, 0x2e003003),
-	uint2(0x77770071, 0x45004004),
-	uint2(0xaaaa00a3, 0xc900a00a)
-};
-#endif
-
 uint EvolveSobolSeed(inout uint Seed)
 {
-#if RANDSEQ_VERSION == 1
 	// LCG parameters from
 	// "Computationally Easy, Spectrally Good Multipliers for Congruential Pseudorandom Number Generators"
 	//   Guy Steele, Sebastiano Vigna
 	//  Journal of Software: Practice and Experience, Volume 52, Issue 2, Feb 2022
 	return Seed = Seed * 0x915f77f5u + 0x93d765ddu;
-#else
-	// constant from: https://www.pcg-random.org/posts/does-it-beat-the-minimal-standard.html
-	const uint MCG_C = 2739110765;
-	// a slightly weaker hash is ok since this drives FastOwenScrambling which is itself a hash
-	// Note that the Seed evolution is just an integer addition and the hash should optimize away
-	// when a particular dimension is not used
-	uint X = (Seed += MCG_C);
-	// Generated using [REDACTED_PROFILE_URL]/hash-prospector
-	// Estimated Bias ~583
-	X *= 0x92955555u;
-	X ^= X >> 15;
-	return X;
-#endif
 }
 
-#if RANDSEQ == RANDSEQ_PIX_LATTICE
-Texture2D<uint> PixelLatticeShift;
-#endif
-
 uint4 SamplerCore(uint SampleIndex, inout uint Seed)
 {
 #if RANDSEQ == RANDSEQ_PURERANDOM
@@ -271,81 +161,34 @@ uint4 SamplerCore(uint SampleIndex, inout uint Seed)
 	uint SobolIndex = FastOwenScrambling(SampleIndex, EvolveSobolSeed(Seed));
 	// now get Sobol' point from this index
 	uint4 Result = uint4(SobolIndex, SobolIndex, 0, 0);
-#if RANDSEQ_VERSION == 1
 	// "SZ Sequences: Binary-Based (0,2^q)-Sequences"
 	//  Abdalla G. M. Ahmed, Matt Pharr, Victor Ostromoukhov, Hui Huang
-	// 	  https://arxiv.org/pdf/2505.20434
-	//    Preprint, May 2025
-	// We only need dimension 2 because 3 can be derived via the formula Uy = P * Ux
-	// TODO: The paper suggests a faster method via Equation 42
-	Result.z = SobolIndex & 3;
-	Result.z ^= ((SobolIndex >>  2) & 1) * 6u;
-	Result.z ^= ((SobolIndex >>  3) & 1) * 11u;
-	Result.z ^= ((SobolIndex >>  4) & 1) * 19u;
-	Result.z ^= ((SobolIndex >>  5) & 1) * 33u;
-	Result.z ^= ((SobolIndex >>  6) & 1) * 109u;
-	Result.z ^= ((SobolIndex >>  7) & 1) * 182u;
-	Result.z ^= ((SobolIndex >>  8) & 1) * 258u;
-	Result.z ^= ((SobolIndex >>  9) & 1) * 515u;
-	Result.z ^= ((SobolIndex >> 10) & 1) * 1547u;
-	Result.z ^= ((SobolIndex >> 11) & 1) * 2829u;
-	Result.z ^= ((SobolIndex >> 12) & 1) * 4897u;
-	Result.z ^= ((SobolIndex >> 13) & 1) * 8498u;
-	Result.z ^= ((SobolIndex >> 14) & 1) * 28086u;
-	Result.z ^= ((SobolIndex >> 15) & 1) * 46811u;
-	Result.z ^= ((SobolIndex >> 16) & 1) * 65539u;
-	Result.z ^= ((SobolIndex >> 17) & 1) * 131073u;
-	Result.z ^= ((SobolIndex >> 18) & 1) * 393229u;
-	Result.z ^= ((SobolIndex >> 19) & 1) * 720902u;
-	Result.z ^= ((SobolIndex >> 20) & 1) * 1245234u;
-	Result.z ^= ((SobolIndex >> 21) & 1) * 2162707u;
-	Result.z ^= ((SobolIndex >> 22) & 1) * 7143643u;
-	Result.z ^= ((SobolIndex >> 23) & 1) * 11927661u;
-	Result.z ^= ((SobolIndex >> 24) & 1) * 16909057u;
-	Result.z ^= ((SobolIndex >> 25) & 1) * 33751298u;
-	Result.z ^= ((SobolIndex >> 26) & 1) * 101387526u;
-	Result.z ^= ((SobolIndex >> 27) & 1) * 185402891u;
-	Result.z ^= ((SobolIndex >> 28) & 1) * 320942611u;
-	Result.z ^= ((SobolIndex >> 29) & 1) * 556929825u;
-	Result.z ^= ((SobolIndex >> 30) & 1) * 1840700269u;
-	Result.z ^= ((SobolIndex >> 31) & 1) * 3067833782u;
-   	// y and w component can be computed without iteration by starting with Result.xz
+	// 	  https://dl.acm.org/doi/epdf/10.1145/3763272
+	//    Siggraph Asia 2025
+	// Optimized evaluation of Dimension 2 via Equation 42
+	Result.z = ((SobolIndex >> 1) & 0x50550550u) ^
+	           ((SobolIndex     ) & 0xaf5af5afu) ^
+	           ((SobolIndex << 1) & 0xa0aa0aa0u);
+	Result.z ^=  Result.z               << 16;
+	Result.z ^= (Result.z & 0x00FF00FF) <<  8;
+	Result.z ^= (Result.z & 0x0F0F0F0F) <<  4;
+	Result.z = ((Result.z >> 1) & 0x50550550u) ^
+	           ((Result.z     ) & 0x5fa5fa5fu) ^
+	           ((Result.z << 1) & 0xa4ea4ea4u) ^
+	           ((Result.z << 2) & 0xc84c84c8u) ^
+	           ((Result.z << 3) & 0x08808808u);
+	Result.zw = reversebits(Result.z);
+   	// Result.yw can be computed without iteration by starting with Result.xz
 	// "An Implementation Algorithm of 2D Sobol Sequence Fast, Elegant, and Compact"
 	// Abdalla Ahmed, EGSR 2024
 	// See listing (19) in the paper
 	// The code is different here because we want the output to be bit-reversed, but
 	// the methodology is the same
-	Result.w = Result.z;
 	Result.yw ^=  Result.yw               >> 16;
 	Result.yw ^= (Result.yw & 0xFF00FF00) >>  8;
 	Result.yw ^= (Result.yw & 0xF0F0F0F0) >>  4;
 	Result.yw ^= (Result.yw & 0xCCCCCCCC) >>  2;
 	Result.yw ^= (Result.yw & 0xAAAAAAAA) >>  1;
-#else
-	// y component can be computed without iteration
-	// "An Implementation Algorithm of 2D Sobol Sequence Fast, Elegant, and Compact"
-	// Abdalla Ahmed, EGSR 2024
-	// See listing (19) in the paper
-	// The code is different here because we want the output to be bit-reversed, but
-	// the methodology is the same
-	Result.y ^=  Result.y               >> 16;
-	Result.y ^= (Result.y & 0xFF00FF00) >>  8;
-	Result.y ^= (Result.y & 0xF0F0F0F0) >>  4;
-	Result.y ^= (Result.y & 0xCCCCCCCC) >>  2;
-	Result.y ^= (Result.y & 0xAAAAAAAA) >>  1;
-	// Remaining bits require a bit-vector by matrix multiplication (in F2, so multiplication is 'and' and addition is 'xor')
-#if RANDSEQ_UNROLL_SOBOL // unrolled version seems faster in most cases
-	UNROLL for (uint b = 0; b < 32; b++)
-	{
-		uint IndexBit = (SobolIndex >> b) & 1;		// bitfield extract
-#else // however a loop with early exit seems faster in the path tracing context (possibly due to low occupancy?)
-	for (uint b = 0; SobolIndex; SobolIndex >>= 1, b++)
-	{
-		uint IndexBit = SobolIndex & 1;
-#endif
-		Result.zw ^= IndexBit * SobolMatrices[b];
-	}
-#endif
 	// finally scramble the points to avoid structured artifacts
 	Result.x = FastOwenScrambling(Result.x, EvolveSobolSeed(Seed));
 	Result.y = FastOwenScrambling(Result.y, EvolveSobolSeed(Seed));
