Substrate - Toon - fixed shadow application on diffuse when unified toon diffusion is disabled.
And improved how shadows are optionally accounted for in the unified diffuse path.

#rnx

#Synced-CL 51401593

diff --git a/Engine/Shaders/Private/DeferredLightPixelShaders.usf b/Engine/Shaders/Private/DeferredLightPixelShaders.usf
--- a/Engine/Shaders/Private/DeferredLightPixelShaders.usf
+++ b/Engine/Shaders/Private/DeferredLightPixelShaders.usf
@@ -29,7 +29,7 @@
 #endif
 
 #define SUBSTRATE_OUTPUT_TOON_DATA (SUBSTRATE_ENABLED && SUBSTRATE_TILETYPE >= 2 && SUBSTRATE_EXPERIMENTAL_TOON_UNIFIED_DIFFUSE)
-#define SUBSTRATE_ENABLE_TOON_DIFFUSE_RAMP (SUBSTRATE_OUTPUT_TOON_DATA ? 0 : 1)
+#define SUBSTRATE_LIGHTPASS_APPLIES_TOON_DIFFUSE_RAMP (SUBSTRATE_OUTPUT_TOON_DATA ? 0 : 1)
 
 #if SUBSTRATE_GBUFFER_FORMAT==0 || USE_HAIR_LIGHTING
 #define SUBSTRATE_LOAD_FROM_MATERIALCONTAINER 0	// Load from a specified BSDF, not from MaterialBuffer.
@@ -697,7 +697,7 @@ void DeferredLightApplyToonDiffusePS(
 		FToonProfileEvaluateSettings ToonEvalSettings = GetToonProfileEvaluateSettings(ToonProfileId);
 		float ToonProfileDiffuseDither = EvaluateToonProfileDiffuseDither(ToonProfileId, ToonPatternUV, ToonEvalSettings);
 
-		const float NoL = ToonData.a;
+		const float NoL = ToonData.a * 2.0 - 1.0;
 		const float3 ToonDiffuse = ToonData.rgb * EvaluateToonProfileDiffuseRamp(ToonProfileId, NoL + ToonProfileDiffuseDither);
 		OutColor = float4(ToonDiffuse, Luminance(ToonDiffuse));
 	}
diff --git a/Engine/Shaders/Private/Substrate/SubstrateEvaluation.ush b/Engine/Shaders/Private/Substrate/SubstrateEvaluation.ush
--- a/Engine/Shaders/Private/Substrate/SubstrateEvaluation.ush
+++ b/Engine/Shaders/Private/Substrate/SubstrateEvaluation.ush
@@ -57,8 +57,8 @@
 #define SUBSTRATE_DIFFUSE_COLOR_BOOST 0
 #endif
 
-#ifndef SUBSTRATE_ENABLE_TOON_DIFFUSE_RAMP
-#define SUBSTRATE_ENABLE_TOON_DIFFUSE_RAMP 1 
+#ifndef SUBSTRATE_LIGHTPASS_APPLIES_TOON_DIFFUSE_RAMP
+#define SUBSTRATE_LIGHTPASS_APPLIES_TOON_DIFFUSE_RAMP 1 
 #endif
 
 #ifndef SUBSTRATE_USE_FULLYSIMPLIFIED_MATERIAL
@@ -1408,10 +1408,10 @@ FSubstrateEvaluateResult SubstrateEvaluateBSDFCommon(FSubstrateBSDFContext BSDFC
 			#endif
 			}
 			
-			const bool bApplyDiffuseRamp	= SUBSTRATE_ENABLE_TOON_DIFFUSE_RAMP;
-			const bool bApplyDiffuseShadow	= !SUBSTRATE_ENABLE_TOON_DIFFUSE_RAMP && ToonEvalSettings.bUnifiedDiffuseIncludeShadow;
+			const bool bApplyDiffuseRamp	= SUBSTRATE_LIGHTPASS_APPLIES_TOON_DIFFUSE_RAMP;
+			const bool bApplyDiffuseShadow	= SUBSTRATE_LIGHTPASS_APPLIES_TOON_DIFFUSE_RAMP || (!SUBSTRATE_LIGHTPASS_APPLIES_TOON_DIFFUSE_RAMP && !ToonEvalSettings.bUnifiedDiffuseIncludeShadow);
 
-			Sample.ToonData					= AreaLight.NoL * (bApplyDiffuseShadow ? 1.0f : TransmittanceShadow);
+			Sample.ToonData					= (AreaLight.NoL*0.5+0.5) * (bApplyDiffuseShadow ? 1.0f : TransmittanceShadow);
 
 			// Diffuse
 			const float3 DiffuseRamp		= bApplyDiffuseRamp ? EvaluateToonProfileDiffuseRamp(ToonProfileId, BSDFContext.Context.NoL + ToonProfileDiffuseDither) : 1.0f;
