Add the extra attenuation back for InverseSquared local lights when MERGED_LOCAL_LIGHTS_MOBILE permutation is active

diff --git a/Engine/Shaders/Private/DeferredLightingCommon.ush b/Engine/Shaders/Private/DeferredLightingCommon.ush
--- a/Engine/Shaders/Private/DeferredLightingCommon.ush
+++ b/Engine/Shaders/Private/DeferredLightingCommon.ush
@@ -258,8 +258,8 @@ float GetLocalLightAttenuation(
 	if (LightData.bInverseSquared)
 	{
 		LightMask = Square( saturate( 1 - Square( DistanceSqr * Square(LightData.InvRadius) ) ) );
-		// This extra attenuation has been added for supporting existing translucent shading model on mobile
-		#if SHADING_PATH_MOBILE && TRANSLUCENCY_SH_LIGHTING
+		// This extra attenuation has been added for supporting existing translucent shading model and merged local lights on mobile
+		#if SHADING_PATH_MOBILE && (TRANSLUCENCY_SH_LIGHTING || MERGED_LOCAL_LIGHTS_MOBILE > 0)
 		if (!LightData.bRectLight)
 		{ 
 			LightMask *= 1.0f / (DistanceSqr + 1.0f);
