Substrate - Toon - Added support for transmittance shadow to virtual shadow map.

#rnx
#Synced-CL 51456025

diff --git a/Engine/Shaders/Private/VirtualShadowMaps/VirtualShadowMapTransmissionCommon.ush b/Engine/Shaders/Private/VirtualShadowMaps/VirtualShadowMapTransmissionCommon.ush
--- a/Engine/Shaders/Private/VirtualShadowMaps/VirtualShadowMapTransmissionCommon.ush
+++ b/Engine/Shaders/Private/VirtualShadowMaps/VirtualShadowMapTransmissionCommon.ush
@@ -11,6 +11,9 @@ VirtualShadowMaps/VirtualShadowMapTransmissionCommon.ush:
 #include "../TransmissionCommon.ush"
 #include "VirtualShadowMapProjectionCommon.ush"
 #include "../Substrate/SubstrateSubsurface.ush"
+#if SUBSTRATE_ENABLED
+#include "../ToonProfileCommon.ush"
+#endif
 
 FSubsurfaceOpacityMFP GetSubsurfaceOpacityFromGbuffer(uint2 PixelPos)
 {
@@ -20,8 +23,24 @@ FSubsurfaceOpacityMFP GetSubsurfaceOpacityFromGbuffer(uint2 PixelPos)
 
 	FSubstrateSubsurfaceData LoadStrat = SubstrateLoadSubsurfaceData(Substrate.MaterialTextureArray, Substrate.FirstSliceStoringSubstrateSSSData, PixelPos);
 
-	// Diffusion and DiffuseProfile are not allowed to translate their MFP into opacity in order to match legacy behavior
-	SubsurfaceOpacityMFP = SubstrateGetSubsurfaceOpacityMFP(LoadStrat.Header, false /*bAllowDiffuse*/, false /*bAllowDiffuseProfile*/);
+	const uint SSSType = SubstrateSubSurfaceHeaderGetSSSType(LoadStrat.Header);
+
+	#if SUBSTRATE_ENABLED
+	if (SSSType == SSS_TYPE_TRANSMITTANCE_SHADOW)
+	{
+		SubsurfaceOpacityMFP.bDataIsOpacity = true;
+		SubsurfaceOpacityMFP.Data = 0; // Unused by VSM for density as Extinction.
+
+		const float3 MFPcm = SubstrateSubSurfaceHeaderGetTransmittanceShadowMFPInMeters(LoadStrat.Header) * METER_TO_CENTIMETER;
+		SubsurfaceOpacityMFP.Density = 1.0f / max(0.000001f, MFPcm.x); // Recover extinction coefficient as 1/cm from MFP in cm.
+	}
+	else
+	#endif 
+		if (SSSType != SSS_TYPE_NONE)
+	{
+		// Diffusion and DiffuseProfile are not allowed to translate their MFP into opacity in order to match legacy behavior
+		SubsurfaceOpacityMFP = SubstrateGetSubsurfaceOpacityMFP(LoadStrat.Header, false /*bAllowDiffuse*/, false /*bAllowDiffuseProfile*/);
+	}
 
 #else // SUBSTRATE_GBUFFER_FORMAT==1
 	FGBufferData GBufferData = GetGBufferDataUint(PixelPos);
@@ -40,6 +59,21 @@ FSubsurfaceOpacityMFP GetSubsurfaceOpacityFromGbuffer(uint2 PixelPos)
 		SubsurfaceOpacityMFP.Data = min(GBufferData.CustomData.a, 0.99f);
 		SubsurfaceOpacityMFP.Density = SubsurfaceDensityFromOpacity(SubsurfaceOpacityMFP.Data);
 	}
+	#if SUBSTRATE_ENABLED
+	else if (GBufferData.ShadingModelID == SHADINGMODELID_SUBSTRATE_TOON)
+	{
+		SubsurfaceOpacityMFP.bDataIsOpacity = true;
+		SubsurfaceOpacityMFP.Data = 0; // Unused by VSM for density as Extinction.
+
+		float ToonProfileId = 0.0f;
+		float2 PatternUVs = 0.0f;
+		UnpackToonCustomData(PackRGBA8(GBufferData.CustomData), PatternUVs, ToonProfileId);
+
+		// Density is extinction as 1/cm (not 1/m) to match with the transmittance shadow evaluation from ManualPCF below.
+		FToonProfileShadowSettings ToonProfileShadowSettings = GetToonProfileShadowSettings(ToonProfileId);
+		SubsurfaceOpacityMFP.Density = ToonProfileShadowSettings.ShadowExtinctionCoefficientInInvCentimeter;
+	}
+	#endif
 #endif // SUBSTRATE_GBUFFER_FORMAT==1
 
 	return SubsurfaceOpacityMFP;
