Temporarly disabled RectAtlas mip filtering using quad swizzling while swizzling code is fixed on certain platforms.

#Synced-CL 51536989

diff --git a/Engine/Shaders/Private/RectLightAtlas.usf b/Engine/Shaders/Private/RectLightAtlas.usf
--- a/Engine/Shaders/Private/RectLightAtlas.usf
+++ b/Engine/Shaders/Private/RectLightAtlas.usf
@@ -80,13 +80,12 @@ void MainVS(
 ////////////////////////////////////////////////////////////////////////////////////////////////
 // Add texture
 
-#define PERMUTATION_WAVE_MIP (PLATFORM_SUPPORTS_SM6_0_WAVE_OPERATIONS)
+#define PERMUTATION_WAVE_MIP (0 && PLATFORM_SUPPORTS_SM6_0_WAVE_OPERATIONS)
 
 #if PERMUTATION_WAVE_MIP
 void FillMip(RWTexture2D<float4> OutAtlasMip, uint2 InCoord, float4 InColor, bool bWrite)
 {
-	float Out = InColor;
-	float3 MipColorH = InColor + QuadReadAcrossX(InColor).xyz;
+	float3 MipColorH = InColor.xyz + QuadReadAcrossX(InColor).xyz;
 	float3 MipColor1 = MipColorH + QuadReadAcrossY(MipColorH).xyz;
 	if (all((InCoord & 1) == 0) && bWrite)
 	{
diff --git a/Engine/Source/Runtime/Renderer/Private/RectLightTextureManager.cpp b/Engine/Source/Runtime/Renderer/Private/RectLightTextureManager.cpp
--- a/Engine/Source/Runtime/Renderer/Private/RectLightTextureManager.cpp
+++ b/Engine/Source/Runtime/Renderer/Private/RectLightTextureManager.cpp
@@ -108,7 +108,7 @@ static const FIntPoint InvalidOrigin = FIntPoint(-1, -1);
 
 static bool SupportWaveQuadPassing(EShaderPlatform In)
 {
-	return GRHISupportsWaveOperations && RHISupportsWaveOperations(In);
+	return false;  //GRHISupportsWaveOperations && RHISupportsWaveOperations(In);
 }
 
 static FIntPoint GetSourceResolution(const FRHITexture* Tex, const FVector4f& ScaleOffset)
