Substrate - fix for SoftSourceRadius to always affect the look, even when SourceRadius=0, which is the correct legacy behavior.

#tests PC editor with substrate enabled with movable lights.
#rnx

#Synced-CL 51536922

diff --git a/Engine/Shaders/Private/Substrate/SubstrateEvaluation.ush b/Engine/Shaders/Private/Substrate/SubstrateEvaluation.ush
--- a/Engine/Shaders/Private/Substrate/SubstrateEvaluation.ush
+++ b/Engine/Shaders/Private/Substrate/SubstrateEvaluation.ush
@@ -625,37 +625,22 @@ FSubstrateEvaluateResult SubstrateEvaluateBSDFCommon(FSubstrateBSDFContext BSDFC
 
 						// Generalized microfacet specular
 						{
-							if(IntegrationType == INTEGRATION_PUNCTUAL_LIGHT)
+							// The energy term needs to be computed before the D evaluation as this function modifies Alpha2Spec.
+							const float Energy = EnergyNormalization(Alpha2Spec, BSDFContext.Context.VoH, AreaLight);
+
+							// GGX lobe
+							{
+								D = Substrate_D_GGX(SafeRoughness, Alpha2Spec, BSDFContext.Context.NoH);
+							}
+
+							// Glint shading only in the complex path
 							{
-								// GGX lobe
-								{
-									D = Substrate_D_GGX(SafeRoughness, Alpha2Spec, BSDFContext.Context.NoH);
-								}
-
-								// Glint shading only in the complex path
-								{
-									D = Substrate_D_Glint(BSDFContext, SafeRoughness, D);
-								}
+								D = Substrate_D_Glint(BSDFContext, SafeRoughness, D);
 							}
-							else
+
+							// Area light energy compensation
 							{
-								// The energy term needs to be computed before the D evaluation as this function modifies Alpha2Spec.
-								const float Energy = EnergyNormalization(Alpha2Spec, BSDFContext.Context.VoH, AreaLight);
-
-								// GGX lobe
-								{
-									D = Substrate_D_GGX(SafeRoughness, Alpha2Spec, BSDFContext.Context.NoH);
-								}
-
-								// Glint shading only in the complex path
-								{
-									D = Substrate_D_Glint(BSDFContext, SafeRoughness, D);
-								}
-
-								// Area light energy compensation
-								{
-									D *= Energy;
-								}
+								D *= Energy;
 							}
 
 							Vis = Substrate_Vis_GGX(SafeRoughness, Alpha2Spec, BSDFContext.Context.NoV, AreaLight.NoL);
@@ -749,15 +734,8 @@ FSubstrateEvaluateResult SubstrateEvaluateBSDFCommon(FSubstrateBSDFContext BSDFC
 						// Special override for roughness for supporting area light integrator with Sphere/Tube/Disk light, which modifies/increase roughness.
 						{
 							float Alpha2SpecHaze = Pow4(HazeSafeRoughness);
-							if (IntegrationType == INTEGRATION_PUNCTUAL_LIGHT)
-							{
-								HazeD = Substrate_D_GGX(HazeSafeRoughness, Alpha2SpecHaze, BSDFContext.Context.NoH);
-							}
-							else
-							{
-								const float Energy = EnergyNormalization(Alpha2SpecHaze, BSDFContext.Context.VoH, AreaLight);
-								HazeD = Substrate_D_GGX(HazeSafeRoughness, Alpha2SpecHaze, BSDFContext.Context.NoH) * Energy;
-							}
+							const float Energy = EnergyNormalization(Alpha2SpecHaze, BSDFContext.Context.VoH, AreaLight);
+							HazeD = Substrate_D_GGX(HazeSafeRoughness, Alpha2SpecHaze, BSDFContext.Context.NoH) * Energy;
 							HazeVis = Substrate_Vis_GGX(HazeSafeRoughness, Alpha2SpecHaze, BSDFContext.Context.NoV, AreaLight.NoL);
 							const float H_PDF = VisibleGGXPDF(BSDFContext.TangentV, BSDFContext.TangentH, Alpha2SpecHaze);
 							HazePDF = RayPDFToReflectionRayPDF(BSDFContext.Context.VoH, H_PDF);
@@ -1277,15 +1255,8 @@ FSubstrateEvaluateResult SubstrateEvaluateBSDFCommon(FSubstrateBSDFContext BSDFC
 
 				// Generalized microfacet specular
 				float Alpha2Spec = Alpha2;
-				if(IntegrationType == INTEGRATION_PUNCTUAL_LIGHT)
-				{
-					D = Substrate_D_GGX(SafeRoughness, Alpha2Spec, BSDFContext.Context.NoH);
-				}
-				else
-				{
-					const float Energy = EnergyNormalization(Alpha2Spec, BSDFContext.Context.VoH, AreaLight);
-					D = Substrate_D_GGX(SafeRoughness, Alpha2Spec, BSDFContext.Context.NoH) * Energy;
-				}
+				const float Energy = EnergyNormalization(Alpha2Spec, BSDFContext.Context.VoH, AreaLight);
+				D = Substrate_D_GGX(SafeRoughness, Alpha2Spec, BSDFContext.Context.NoH) * Energy;
 				Vis = Substrate_Vis_GGX(SafeRoughness, Alpha2Spec, BSDFContext.Context.NoV, AreaLight.NoL);
 				const float H_PDF = VisibleGGXPDF(BSDFContext.TangentV, BSDFContext.TangentH, Alpha2Spec);
 				PDF = RayPDFToReflectionRayPDF(BSDFContext.Context.VoH, H_PDF);
