VirtualTexture: Don't write empty feedback from shaders.
In previous screen space feedback approaches we wrote empty feedback to potentially overwrite feedback from occluded geometry.
Now that we have append list style feedback, this is no longer relevant and we should drop empty feedback and save the write and subsequent processing.
Note that the occlusion case should be handled by the depth prepass which is turned on when virtual texture is enabled.
#rnx
#Synced-CL 51537754

diff --git a/Engine/Shaders/Private/VirtualTextureCommon.ush b/Engine/Shaders/Private/VirtualTextureCommon.ush
--- a/Engine/Shaders/Private/VirtualTextureCommon.ush
+++ b/Engine/Shaders/Private/VirtualTextureCommon.ush
@@ -92,7 +92,7 @@ void FinalizeVirtualTextureFeedback(in FVirtualTextureFeedbackParams Params, in
 	uint PixelTileIndex = (PixelTilePos.y << FeedbackViewParams.y) + PixelTilePos.x;
 
 	// This code will only run every few pixels...
-	[branch] if (PixelTileIndex == FeedbackViewParams.z)
+	[branch] if (PixelTileIndex == FeedbackViewParams.z && Params.Request != 0xFFFFFFFF)
 	{
 		uint WriteOffset = 0;
 		InterlockedAdd(OutputBuffer[0], 1, WriteOffset);
