Fix incorrect tangent during material evaluation when using discret hair strands LOD.

When using culling, active point might not be consecutive. This is an issue during material evaluation asonly the resolved segment's point ID is stored per pixel. The knowledge of the next 'valid' point ID is lost. To ensure corrent tangent interpolation during material evaluation, the adjacent point's tangent is updated with the current point value

#Synced-CL 51620156

diff --git a/Engine/Shaders/Private/HairStrands/HairStrandsTangent.usf b/Engine/Shaders/Private/HairStrands/HairStrandsTangent.usf
--- a/Engine/Shaders/Private/HairStrands/HairStrandsTangent.usf
+++ b/Engine/Shaders/Private/HairStrands/HairStrandsTangent.usf
@@ -60,7 +60,7 @@ void MainCS(uint DispatchThreadId : SV_DispatchThreadID)
 	}
 #endif
 
-	if (IndexCurr < MaxPointCount)
+	if (bValid && IndexCurr < MaxPointCount)
 	{
 		const FHairControlPoint CPCurr = ReadHairControlPoint(PositionBuffer,IndexCurr);
 		const FHairControlPoint CPPrev = ReadHairControlPoint(PositionBuffer,IndexPrev);
@@ -85,8 +85,21 @@ void MainCS(uint DispatchThreadId : SV_DispatchThreadID)
 		// TODO: switch to rotation minimizing frames to avoid the discontinuity
 		float3x3 Basis = GetTangentBasis(T);
 
+		// Tangent
 		OutputTangentBuffer[IndexCurr*2]   = PackNormal(Basis[0], 0);
 		OutputTangentBuffer[IndexCurr*2+1] = PackNormal(Basis[2], 127);
+
+		// When using culling, active point might not be consecutive. This is an issue during material evaluation asonly the resolved segment's 
+		// point ID is stored per pixel. The knowledge of the next 'valid' point ID is lost. To ensure corrent tangent interpolation during material 
+		// evaluation, the adjacent point's tangent is updated with the current point value
+		#if PERMUTATION_CULLING
+		if (IndexNext != IndexCurr+1 && IndexCurr+1 < PointCount && CPCurr.Type != HAIR_CONTROLPOINT_END)
+		{
+			const uint CurrNextIndex = IndexCurr + 1;
+			OutputTangentBuffer[CurrNextIndex*2]   = PackNormal(Basis[0], 0);
+			OutputTangentBuffer[CurrNextIndex*2+1] = PackNormal(Basis[2], 127);
+		}
+		#endif
 	}
 }
 #endif // SHADER_TANGENT
\ No newline at end of file
