Tonemap: Dead code removal.

#Synced-CL 51657854

diff --git a/Engine/Shaders/Private/PostProcessTonemap.usf b/Engine/Shaders/Private/PostProcessTonemap.usf
--- a/Engine/Shaders/Private/PostProcessTonemap.usf
+++ b/Engine/Shaders/Private/PostProcessTonemap.usf
@@ -548,10 +548,6 @@ float4 TonemapCommonPS(
 		// RETURN_COLOR not needed unless writing to SceneColor
 		OutColor.rgb = OutDeviceColor;
 
-		#if USE_LINEAR_FLOAT_RT
-			OutColor.rgb = sRGBToLinear( OutColor.rgb );
-		#endif
-
 		BRANCH
 		if(GetOutputDevice() == TONEMAPPER_OUTPUT_LinearEXR)
 		{
