Substrate - Toon - removed temporary controls after validation.
- Energy consevation is not used.
- GGX specular is used.

#rnx
#Synced-CL 51808331

diff --git a/Engine/Shaders/Private/Substrate/SubstrateEvaluation.ush b/Engine/Shaders/Private/Substrate/SubstrateEvaluation.ush
--- a/Engine/Shaders/Private/Substrate/SubstrateEvaluation.ush
+++ b/Engine/Shaders/Private/Substrate/SubstrateEvaluation.ush
@@ -1314,10 +1314,7 @@ FSubstrateEvaluateResult SubstrateEvaluateBSDFCommon(FSubstrateBSDFContext BSDFC
 				TransmittanceShadow *= EvaluateToonProfileShadowHatching(ToonProfileId, ToonPatternUV, ToonProfileShadowSettings, TransmittanceShadow);
 			}
 
-			float DirectionalAlbedo_SpecularTransmission = 1.0f;
 			float3 Specular = 0.0f;
-			FBxDFEnergyType MSScale = 1.0f;
-			if (ToonEvalSettings.bSpecularUsesGGX)
 			{
 				const float SafeRoughness = MakeRoughnessSafe(Roughness, SUBSTRATE_MIN_GGX_ROUGHNESS);
 				float Alpha2 = Pow4(SafeRoughness);
@@ -1326,20 +1323,13 @@ FSubstrateEvaluateResult SubstrateEvaluateBSDFCommon(FSubstrateBSDFContext BSDFC
 				const float3 F0 = SpecularColor;
 				const float3 F90 = 1.0f;
 
-				if (ToonEvalSettings.bEnergyConservation)
-				{
-					FBxDFEnergyTerms EnergyTerms = ComputeGGXSpecEnergyTerms(SafeRoughness, BSDFContext.Context.NoV, F0, F90);
-					DirectionalAlbedo_SpecularTransmission = ComputeEnergyPreservation(EnergyTerms);
-					MSScale = ComputeEnergyConservation(EnergyTerms);
-				}
-
 				const bool bIsRectLight = IntegrationType == INTEGRATION_AREA_LIGHT_RECT;
 				if (bIsRectLight)
 				{
 					// In this case, we set D to 1 and Vis will contain the area light / GGX integration
 					float3 RectLightSpec = RectGGXApproxLTC(SafeRoughness, /*F0*/ 1.0, /*F90*/ 1.0, BSDFContext.N, BSDFContext.V, AreaLight.Rect, AreaLight.Texture);
 
-					Specular = RectLightSpec * MSScale;
+					Specular = RectLightSpec;
 				}
 				else
 				{
@@ -1354,19 +1344,6 @@ FSubstrateEvaluateResult SubstrateEvaluateBSDFCommon(FSubstrateBSDFContext BSDFC
 
 				// TODO H_PDF and PDF
 			}
-			else
-			{
-				const float RoughnessExponent = saturate(1.0 - Roughness);
-				const float RoughnessExponent2 = RoughnessExponent * RoughnessExponent; // Arbitrary: this is so that the roughness control appears smooth and linear without big jumps due to 8bit encoding.
-
-				const float MaxSpecularExponent = 2048.0f;
-				const float SpecularExponentScale = MaxSpecularExponent - 1.0f;
-			#if 0
-				Specular = pow(saturate(dot(AreaLight.SpecularL, BSDFContext.R)), max(1.0f, 1.0 + RoughnessExponent2 * SpecularExponentScale));	// NoR based is also interesting visually.
-			#else
-				Specular = pow(saturate(BSDFContext.Context.NoH), max(1.0f, 1.0 + RoughnessExponent2 * SpecularExponentScale));					// NoH based
-			#endif
-			}
 			
 			const bool bApplyDiffuseRamp	= SUBSTRATE_LIGHTPASS_APPLIES_TOON_DIFFUSE_RAMP;
 			const bool bApplyDiffuseShadow	= SUBSTRATE_LIGHTPASS_APPLIES_TOON_DIFFUSE_RAMP || (!SUBSTRATE_LIGHTPASS_APPLIES_TOON_DIFFUSE_RAMP && !ToonEvalSettings.bUnifiedDiffuseIncludeShadow);
@@ -1375,11 +1352,11 @@ FSubstrateEvaluateResult SubstrateEvaluateBSDFCommon(FSubstrateBSDFContext BSDFC
 
 			// Diffuse
 			const float3 DiffuseRamp		= bApplyDiffuseRamp ? EvaluateToonProfileDiffuseRamp(ToonProfileId, BSDFContext.Context.NoL + ToonProfileDiffuseDither) : 1.0f;
-			Sample.DiffuseColor				= DirectionalAlbedo_SpecularTransmission * DiffuseColor * DiffuseRamp;
-			Sample.DiffusePathValue			= DirectionalAlbedo_SpecularTransmission * Diffuse_Lambert(Sample.DiffuseColor);
-			Sample.IntegratedDiffuseValue	= DirectionalAlbedo_SpecularTransmission * (bApplyDiffuseShadow ? TransmittanceShadow : 1.0f) * AreaLight.Falloff * AreaLight.FalloffColor * Sample.DiffusePathValue; //  * AreaLight.NoL for smoother edges? But then transmittance is blocked. So this is left in the hand of the artist.
+			Sample.DiffuseColor				= DiffuseColor * DiffuseRamp;
+			Sample.DiffusePathValue			= Diffuse_Lambert(Sample.DiffuseColor);
+			Sample.IntegratedDiffuseValue	= (bApplyDiffuseShadow ? TransmittanceShadow : 1.0f) * AreaLight.Falloff * AreaLight.FalloffColor * Sample.DiffusePathValue; //  * AreaLight.NoL for smoother edges? But then transmittance is blocked. So this is left in the hand of the artist.
 
-			const float3 SpecularRamp = SpecularColor * MSScale * EvaluateToonProfileSpecularRamp(ToonProfileId, saturate(dot(Specular, 1.0f/3.0f) + ToonProfileSpecularDither));
+			const float3 SpecularRamp = SpecularColor * EvaluateToonProfileSpecularRamp(ToonProfileId, saturate(dot(Specular, 1.0f/3.0f) + ToonProfileSpecularDither));
 			Sample.SpecularPathValue = SpecularRamp;
 			Sample.IntegratedSpecularValue += ShadowTerms.SurfaceShadow * AreaLight.NoL * AreaLight.Falloff * AreaLight.FalloffColor * Sample.SpecularPathValue;
 
@@ -2142,22 +2119,9 @@ FSubstrateEnvLightResult SubstrateEvaluateForEnvLight(FSubstrateBSDFContext BSDF
 			float3 SpecEnvBRDF = 0.0f;
 			float DirectionalAlbedo_SpecularTransmission = 1.0f;
 			float3 MSScale = 1.0f;
-			if (ToonEvalSettings.bEnergyConservation)
-			{
-				FBxDFEnergyTermsRGB SpecularEnergyTerms = ComputeGGXSpecEnergyTermsRGB_Env(SafeRoughness, BSDFContext.Context.NoV, F0, F90);
-				FBxDFEnergyTermsRGB DiffuseEnergyTerms = ComputeDiffuseEnergyTermsRGB(SafeRoughness, BSDFContext.Context.NoV);
-				MSScale = ComputeEnergyConservation(SpecularEnergyTerms);
 
-				SpecEnvBRDF = MSScale * SpecularEnergyTerms.E;
-				DiffEnvBRDF = DiffuseEnergyTerms.E;
-
-				DirectionalAlbedo_SpecularTransmission = ComputeEnergyPreservation(SpecularEnergyTerms);
-			}
-			else
-			{
-				SpecEnvBRDF = EnvBRDF(F0, SafeRoughness, BSDFContext.Context.NoV);
-				DiffEnvBRDF = 1.0f;
-			}
+			SpecEnvBRDF = EnvBRDF(F0, SafeRoughness, BSDFContext.Context.NoV);
+			DiffEnvBRDF = 1.0f;
 
 			FToonProfileGISettings ToonProfileGISettings	= GetToonProfileGISettings(ToonProfileId);
 
diff --git a/Engine/Shaders/Private/ToonProfile.usf b/Engine/Shaders/Private/ToonProfile.usf
--- a/Engine/Shaders/Private/ToonProfile.usf
+++ b/Engine/Shaders/Private/ToonProfile.usf
@@ -22,8 +22,6 @@ float ShadowHatchingPatternStrength;
 
 float DiffuseIndirectScale;
 float SpecularIndirectScale;
-float bSpecularUsesGGX;
-float bEnergyConservation;
 float bUnifiedDiffuseIncludeShadow;
 float SpecularIndirectRampRepetition;
 
@@ -78,7 +76,7 @@ void ToonProfileUpdateAtlasTileCS(uint3 DispatchThreadId : SV_DispatchThreadID)
 			}
 			else if (PixelPos.x == TOON_PROFILE_EVAL_PARAMS2_X)
 			{
-				ToonProfileTexture[uint3(PixelPos, ToonProfileIndex)] = float4(bSpecularUsesGGX, bEnergyConservation, bUnifiedDiffuseIncludeShadow, 0.0f);
+				ToonProfileTexture[uint3(PixelPos, ToonProfileIndex)] = float4(bUnifiedDiffuseIncludeShadow, 0.0f, 0.0f, 0.0f);
 			}
 			else if (PixelPos.x == TOON_PROFILE_SHADOW_PARAMS_X)
 			{
diff --git a/Engine/Shaders/Private/ToonProfileCommon.ush b/Engine/Shaders/Private/ToonProfileCommon.ush
--- a/Engine/Shaders/Private/ToonProfileCommon.ush
+++ b/Engine/Shaders/Private/ToonProfileCommon.ush
@@ -44,8 +44,6 @@ struct FToonProfileEvaluateSettings
 	float DiffuseDitherPatternSize;
 	float SpecularDitherPatternStrength;
 	float SpecularDitherPatternSize;
-	float bSpecularUsesGGX;
-	float bEnergyConservation;
 	float bUnifiedDiffuseIncludeShadow;
 };
 
@@ -69,9 +67,7 @@ FToonProfileEvaluateSettings GetToonProfileEvaluateSettings(in float ToonProfile
 		ToonProfileCoord.yz);
 	float4 Data2 = View.ToonProfileTexture.SampleLevel(View.ToonProfileSampler, ToonProfileCoord2, 0).xyzw;
 
-	Ret.bSpecularUsesGGX = Data2.x;
-	Ret.bEnergyConservation = Data2.y;
-	Ret.bUnifiedDiffuseIncludeShadow = Data2.z;
+	Ret.bUnifiedDiffuseIncludeShadow = Data2.x;
 
 	return Ret;
 }
diff --git a/Engine/Source/Runtime/Engine/Classes/Engine/ToonProfile.h b/Engine/Source/Runtime/Engine/Classes/Engine/ToonProfile.h
--- a/Engine/Source/Runtime/Engine/Classes/Engine/ToonProfile.h
+++ b/Engine/Source/Runtime/Engine/Classes/Engine/ToonProfile.h
@@ -128,18 +128,6 @@ struct FToonProfileStruct
 	UPROPERTY(Category = "Global Illumination", EditAnywhere, meta = (UIMin = 0.0, UIMax = 1.0, ClampMin = 0.0))
 	float SpecularIndirectRampRepetition;
 
-	/**
-	 * Experimental.
-	 */
-	UPROPERTY(Category = "Experimental", EditAnywhere, meta = (DisplayName = "Specular Uses GGX"))
-	uint8 bSpecularUsesGGX : 1;
-
-	/**
-	 * Experimental.
-	 */
-	UPROPERTY(Category = "Experimental", EditAnywhere, meta = (DisplayName = "Energy Conservation"))
-	uint8 bEnergyConservation : 1;
-
 	/**
 	 * Experimental. Requires r.Substrate.Experimental.ToonUnifiedDiffuse=1.
 	 */
diff --git a/Engine/Source/Runtime/Engine/Private/Rendering/ToonProfile.cpp b/Engine/Source/Runtime/Engine/Private/Rendering/ToonProfile.cpp
--- a/Engine/Source/Runtime/Engine/Private/Rendering/ToonProfile.cpp
+++ b/Engine/Source/Runtime/Engine/Private/Rendering/ToonProfile.cpp
@@ -82,9 +82,7 @@ public:
 		SHADER_PARAMETER(float, ShadowHatchingPatternStrength)
 		SHADER_PARAMETER(float, DiffuseIndirectScale)
 		SHADER_PARAMETER(float, SpecularIndirectScale)
-		SHADER_PARAMETER(float, bSpecularUsesGGX)
 		SHADER_PARAMETER(float, bUnifiedDiffuseIncludeShadow)
-		SHADER_PARAMETER(float, bEnergyConservation)
 		SHADER_PARAMETER(float, SpecularIndirectRampRepetition)
 
 		SHADER_PARAMETER(float, ToonProfileShadowExtinctionCoefficientInInvCentimeter)
@@ -418,9 +416,7 @@ IPooledRenderTarget* FToonProfileTextureManager::GetAtlasTexture(FRDGBuilder& Gr
 
 			PassParameters->DiffuseIndirectScale = Data.DiffuseIndirectScale;
 			PassParameters->SpecularIndirectScale = Data.SpecularIndirectScale;
-			PassParameters->bSpecularUsesGGX = Data.bSpecularUsesGGX ? 1.0f : 0.0f;
 			PassParameters->bUnifiedDiffuseIncludeShadow = Data.bUnifiedDiffuseIncludeShadow ? 1.0f : 0.0f;
-			PassParameters->bEnergyConservation = Data.bEnergyConservation ? 1.0f : 0.0f;
 			PassParameters->SpecularIndirectRampRepetition = Data.SpecularIndirectRampRepetition;
 
 			auto CreateBufferFRomCurve = [&](FRuntimeCurveLinearColor ColorCurve)
@@ -558,9 +554,7 @@ FToonProfileStruct::FToonProfileStruct()
 	SpecularIndirectScale = 1.0f;
 	SpecularIndirectRampRepetition = 0.0f;
 
-	bSpecularUsesGGX = 0;
 	bUnifiedDiffuseIncludeShadow = 0;
-	bEnergyConservation = 0;
 
 	AddPointToCurve(ShadowRamp, 0.00, 0.0);
 	AddPointToCurve(ShadowRamp, 1.00, 1.0);
