Fixed a shader compilation error caused by a parameter name that uses a reserved keyword on certain platforms.

#Synced-CL 51768376

diff --git a/Engine/Shaders/Private/PostProcessMobile.usf b/Engine/Shaders/Private/PostProcessMobile.usf
--- a/Engine/Shaders/Private/PostProcessMobile.usf
+++ b/Engine/Shaders/Private/PostProcessMobile.usf
@@ -68,9 +68,9 @@ float4 BloomUpSizeAndInvSize;
 float4 BufferViewRect;
 
 // Point on circle.
-float2 Circle(float Start, float Points, float Point) 
+float2 Circle(float Start, float Points, float StepIndex) 
 {
-	float Rad = (3.141592 * 2.0 * (1.0 / Points)) * (Point + Start);
+	float Rad = (3.141592 * 2.0 * (1.0 / Points)) * (StepIndex + Start);
 	return float2(sin(Rad), cos(Rad));
 }
 
