Path Tracer: Improve shading response inside a point (or spot) light's radius

The previous behavior could lead to artifacts when looking "inside" large radius lights. This produced artifacts that were hard to work around, particularly for lights with SourceLength > 0 in volumes. The new behavior clamps the light's radius by the distance so that we can resolve these issues while keeping the lighting outside the problematic region unchanged.

In the case of capsule lights (SourceLength > 0), some further tweaking was needed to avoid introducing a "bright line" in the central core which is visually inconsistent with the deferred renderer's treatment of this issue. As the source length increases, we limit how low the relative radius is for the purposes of calculation the energy output from the surface of the light to bound the overall brightness as we approach the light's center.

#Synced-CL 51839045

diff --git a/Engine/Shaders/Private/PathTracing/Light/PathTracingCapsuleLight.ush b/Engine/Shaders/Private/PathTracing/Light/PathTracingCapsuleLight.ush
--- a/Engine/Shaders/Private/PathTracing/Light/PathTracingCapsuleLight.ush
+++ b/Engine/Shaders/Private/PathTracing/Light/PathTracingCapsuleLight.ush
@@ -18,21 +18,28 @@ float3 CapsuleNormal(float3 Pos, float3 Center, float3 Axis)
 	return normalize(pa - h * Axis);
 }
 
-float GetCapsuleRadius(int LightId)
+float GetCapsuleRadius(int LightId, float3 Origin, float3 Axis, float Length)
 {
+    float h = clamp(dot(Axis, Origin), -0.5 * Length, 0.5 * Length);
+	float3 ClosestPointOnAxis = Origin - Axis * h;
+	float DistanceToAxisSqr = dot(ClosestPointOnAxis, ClosestPointOnAxis);
+
 	// don't allow 0 area capsules since we need to divide the light power by area to get radiance
-	return max(GetRadius(LightId), 0.01);
+    // Also, shrink capsule if we are inside it
+    return max(min(sqrt(0.99 * DistanceToAxisSqr), GetRadius(LightId)), 0.01);
 }
 
 FLightHit CapsuleLight_TraceLight(FRayDesc Ray, int LightId)
 {
 	float3 TranslatedLightPosition = GetTranslatedPosition(LightId);
-	float LightRadius = GetCapsuleRadius(LightId);
-	float LightRadius2 = Pow2(LightRadius);
-	float SourceLength = GetSourceLength(LightId);
+    float3 Origin = TranslatedLightPosition - Ray.Origin;
+    float3 Axis = GetdPdv(LightId);
+    float SourceLength = GetSourceLength(LightId);
+	float Radius = GetCapsuleRadius(LightId, Origin, Axis, SourceLength);
+	float Radius2 = Pow2(Radius);
 
 #if 0 // DEBUG code: visualize the bounding rectangle against the capsule intersection test
-	FCapsuleSphericalBounds CapsuleBounds = CapsuleGetSphericalBounds(TranslatedLightPosition - Ray.Origin, GetdPdv(LightId), LightRadius, SourceLength);
+	FCapsuleSphericalBounds CapsuleBounds = CapsuleGetSphericalBounds(TranslatedLightPosition - Ray.Origin, GetdPdv(LightId), Radius, SourceLength);
 
 	float ConeSolidAngle = CapsuleBounds.ConeSolidAngle;
 	float RectSolidAngle = CapsuleBounds.SphericalRect.SolidAngle;
@@ -68,11 +75,10 @@ FLightHit CapsuleLight_TraceLight(FRayDesc Ray, int LightId)
 	// ray points toward the plane and intersect it?
 	if (t > Ray.TMin && t < Ray.TMax)
 	{
-		float CapsArea = LightRadius2;
-		float BodyArea = 0.5 * LightRadius * SourceLength;
+        float CapsuleArea = PI * Radius * (Radius + 0.5 * SourceLength);
 		float3 LightPower = GetColor(LightId);
-		float3 LightRadiance = LightPower / (PI * (CapsArea + BodyArea));
-		float t = CapsuleIntersect(Ray.Direction, TranslatedLightPosition - Ray.Origin, GetdPdv(LightId), LightRadius2, SourceLength);
+        float3 LightRadiance = LightPower / CapsuleArea;
+		float t = CapsuleIntersect(Ray.Direction, TranslatedLightPosition - Ray.Origin, GetdPdv(LightId), Radius2, SourceLength);
 		if (t > Ray.TMin && t < Ray.TMax)
 		{
 			LightRadiance *= float3(0, 1, 0);
@@ -82,13 +88,12 @@ FLightHit CapsuleLight_TraceLight(FRayDesc Ray, int LightId)
 
 	// missed the bounding rect, test capsule anyway to see if any part of the capsule was missing
 	{
-		float CapsArea = LightRadius2;
-		float BodyArea = 0.5 * LightRadius * SourceLength;
+        float CapsuleArea = PI * Radius * (Radius + 0.5 * SourceLength);
 		float3 LightPower = GetColor(LightId);
-		float3 LightRadiance = LightPower / (PI * (CapsArea + BodyArea));
+        float3 LightRadiance = LightPower / CapsuleArea;
 		{
 			// exact
-			float t = CapsuleIntersect(Ray.Direction, TranslatedLightPosition - Ray.Origin, GetdPdv(LightId), LightRadius2, SourceLength);
+			float t = CapsuleIntersect(Ray.Direction, TranslatedLightPosition - Ray.Origin, GetdPdv(LightId), Radius2, SourceLength);
 			if (t > Ray.TMin && t < Ray.TMax)
 			{
 				LightRadiance *= float3(1, 0, 0);
@@ -97,25 +102,24 @@ FLightHit CapsuleLight_TraceLight(FRayDesc Ray, int LightId)
 		}
 	}
 #else
-	float3 Axis = GetdPdv(LightId);
-	float t = CapsuleIntersect(Ray.Direction, TranslatedLightPosition - Ray.Origin, Axis, LightRadius2, SourceLength);
+	float t = CapsuleIntersect(Ray.Direction, TranslatedLightPosition - Ray.Origin, Axis, Radius2, SourceLength);
 	if (t > Ray.TMin && t < Ray.TMax)
 	{
-		float CapsArea = LightRadius2;
-		float BodyArea = 0.5 * LightRadius * SourceLength;
+        float CapsRadius = max(Radius, 0.5 * GetRadius(LightId)); // shrink as we go into the light, but not too much
+        float CapsuleArea = PI * CapsRadius * (CapsRadius + 0.5 * SourceLength);
 		float3 LightPower = GetColor(LightId);
-		float3 LightRadiance = LightPower / (PI * (CapsArea + BodyArea));
+        float3 LightRadiance = LightPower / CapsuleArea;
 		float3 LightDirection = TranslatedLightPosition - Ray.Origin;
 		float LightDistanceSquared = dot(LightDirection, LightDirection);
 #if CAPSULE_SOLIDANGLE_SAMPLING
-		FCapsuleSphericalBounds Bounds = CapsuleGetSphericalBounds(LightDirection, Axis, LightRadius, SourceLength);
+		FCapsuleSphericalBounds Bounds = CapsuleGetSphericalBounds(LightDirection, Axis, Radius, SourceLength);
 		float Pdf = 1.0 / GetCapsuleBoundsInversePdf(Ray.Direction, Bounds);
 #else
 		float3 Normal = CapsuleNormal(Ray.Origin + t * Ray.Direction, TranslatedLightPosition, Axis * SourceLength);
 
 		float CosTheta = saturate(-dot(Normal, Ray.Direction));
 
-		float Pdf = CosTheta > 0 ? t * t / (4 * PI * CosTheta * (CapsArea + BodyArea)) : 0.0;
+		float Pdf = CosTheta > 0 ? t * t / (4 * PI * CosTheta * Radius * (Radius + 0.5 * SourceLength)) : 0.0;
 #endif
 		return CreateLightHit(LightRadiance, Pdf, t);
 	}
@@ -132,25 +136,18 @@ FLightSample CapsuleLight_SampleLight(
 )
 {
 	// Capsule case
-	// #dxr_todo: only sample the visible portion of the capsule and account for the 1/d^2 falloff down the axis
-
-	float Radius = GetCapsuleRadius(LightId);
-	float Radius2 = Radius * Radius;
-	float SourceLength = GetSourceLength(LightId);
-
-	// the caps are two halves of a full sphere
-	// the body is a cylinder
-	// the common factor of 4*PI is accounted for at the end
-	float CapsArea = Radius2;
-	float BodyArea = 0.5 * Radius * SourceLength;
-
-	float3 Axis = GetdPdv(LightId);
+	float3 LightDirection = GetTranslatedPosition(LightId) - TranslatedWorldPos;
+    float3 Axis = GetdPdv(LightId);
+    float SourceLength = GetSourceLength(LightId);
+	float Radius = GetCapsuleRadius(LightId, LightDirection, Axis, SourceLength);
+	float Radius2 = Pow2(Radius);
 
 #if CAPSULE_SOLIDANGLE_SAMPLING
-	float3 LightPower = GetColor(LightId);
-	float3 LightRadiance = LightPower / (PI * (CapsArea + BodyArea));
-
-	float3 LightDirection = GetTranslatedPosition(LightId) - TranslatedWorldPos;
+    float CapsRadius = max(Radius, 0.5 * min(SourceLength, GetRadius(LightId))); // shrink as we go into the light, but not too much
+    float CapsuleArea = PI * CapsRadius * (CapsRadius + 0.5 * SourceLength);
+    float3 LightPower = GetColor(LightId);
+    float3 LightRadiance = LightPower / CapsuleArea;
+	
 	FCapsuleSphericalBounds Bounds = CapsuleGetSphericalBounds(LightDirection, Axis, Radius, SourceLength);
 
 	float4 Result = SampleCapsuleBounds(Bounds, RandSample);
@@ -174,6 +171,11 @@ FLightSample CapsuleLight_SampleLight(
 	return NullLightSample();
 #else
 	// plain area sampling
+	// the caps are two halves of a full sphere
+	// the body is a cylinder
+	// the common factor of 4*PI is accounted for at the end
+	float CapsArea = Radius2;
+	float BodyArea = 0.5 * Radius * SourceLength;
 	float Prob = CapsArea / (CapsArea + BodyArea);
 
 	float3 LightPoint;
diff --git a/Engine/Shaders/Private/PathTracing/Light/PathTracingPointLight.ush b/Engine/Shaders/Private/PathTracing/Light/PathTracingPointLight.ush
--- a/Engine/Shaders/Private/PathTracing/Light/PathTracingPointLight.ush
+++ b/Engine/Shaders/Private/PathTracing/Light/PathTracingPointLight.ush
@@ -35,9 +35,11 @@ FLightHit PointLight_TraceLight(FRayDesc Ray, int LightId)
 	else
 	{
 		float3 TranslatedLightPosition = GetTranslatedPosition(LightId);
-		float LightRadius = GetRadius(LightId);
-		float LightRadius2 = Pow2(LightRadius);
 		float3 oc = Ray.Origin - TranslatedLightPosition;
+        float LightDistanceSquared = dot(oc, oc);
+		// Clamp the radius such that we are never fully inside the sphere
+		float LightRadius2 = min(0.99 * LightDistanceSquared, Pow2(GetRadius(LightId)));
+
 		float b = dot(oc, Ray.Direction);
 		// "Precision Improvements for Ray / Sphere Intersection" - Ray Tracing Gems (2019)
 		// https://link.springer.com/content/pdf/10.1007%2F978-1-4842-4427-2_7.pdf
@@ -48,13 +50,11 @@ FLightHit PointLight_TraceLight(FRayDesc Ray, int LightId)
 			float t = -b - sqrt(h);
 			if (t > Ray.TMin && t < Ray.TMax)
 			{
-				float LightDistanceSquared = dot(oc, oc);
-
 				// #dxr_todo: sphere area is 4*pi*r^2 -- but the factor of 4 is missing for some reason?
 				float3 LightPower = GetColor(LightId);
 				float3 LightRadiance = LightPower / (PI * LightRadius2);
 
-				float SinThetaMax2 = saturate(LightRadius2 / LightDistanceSquared);
+                float SinThetaMax2 = LightRadius2 / LightDistanceSquared;
 				float OneMinusCosThetaMax = SinThetaMax2 < 0.01 ? SinThetaMax2 * (0.5 + 0.125 * SinThetaMax2) : 1 - sqrt(1 - SinThetaMax2);
 
 				float SolidAngle = 2 * PI * OneMinusCosThetaMax;
@@ -99,12 +99,10 @@ FLightSample PointLight_SampleLight(
 	{
 		// Point light case
 
-		// Sample the solid angle subtended by the sphere (which could be singgular, in which case the PDF will be infinite)
-		float Radius = GetRadius(LightId);
-		float Radius2 = Pow2(Radius);
-
-		// #dxr_todo: come up with a better definition when we are inside the light
-		float SinThetaMax2 = saturate(Radius2 / LightDistanceSquared);
+		// Sample the solid angle subtended by the sphere (which could be singular, in which case the PDF will be infinite)
+		// Clamp the radius such that we are never fully inside the sphere
+		float Radius2 = min(0.99 * LightDistanceSquared, Pow2(GetRadius(LightId)));
+        float SinThetaMax2 = Radius2 / LightDistanceSquared;
 
 		// #dxr_todo: find a better way of handling the region inside the light than just clamping to 1.0 here
 		float4 DirAndPdf = UniformSampleConeRobust(RandSample, SinThetaMax2);
