Path Tracer: Cleanup the light loop for PrimaryRays debug mode and add new visualization modes

This adds various ifdefs that are useful for performance comparisons. When in RIS light loop mode, accumulate the radiance of unshadowed lights directly to focus sampling on shadowing calculations.

Add several debug visualization modes, including triangle/cluster/instance ID. Restore support for all the light grid and decal visualization modes. Add a cvar to independently control the scale of the light/decal count visualizations without needing to change the actual light grid.

Add ambient occlusion shading to some debug modes to improve readability. This works well with denoisers such as DLSS, or can be disabled by setting the max distance to a negative value.

Add max path intensity support to Primary Rays mode

Tweak RIS and sample clamp methods to encourage compiler to generate code without branches.

Adjust random sample assigment for area lights to ensure the higher quality (x,y) pair is used to sample the light's surface.

#Synced-CL 51851216

diff --git a/Engine/Shaders/Private/PathTracing/Light/PathTracingLightGrid.ush b/Engine/Shaders/Private/PathTracing/Light/PathTracingLightGrid.ush
--- a/Engine/Shaders/Private/PathTracing/Light/PathTracingLightGrid.ush
+++ b/Engine/Shaders/Private/PathTracing/Light/PathTracingLightGrid.ush
@@ -74,16 +74,16 @@ FLightLoopCount LightGridLookup(float3 TranslatedWorldPos)
 	return Result;
 }
 
-float3 LightGridVisualize(FLightLoopCount LoopCount, int VisualizeMode)
+float3 LightGridVisualize(FLightLoopCount LoopCount, int DebugMode, uint MaxCount)
 {
 	const float3 OutsideColor = 0.18;
 	const float3 EmptyColor = 0.36;
 
 	if (LoopCount.LightGridOffset != ~0u)
 	{
-		switch (VisualizeMode)
+		switch (DebugMode)
 		{
-			case 2:
+			case PATH_TRACER_DEBUG_VIZ_LIGHT_GRID_LIST_HASH:
 			{
 				// color by unique light list
 				uint H = 0;
@@ -91,9 +91,13 @@ float3 LightGridVisualize(FLightLoopCount LoopCount, int VisualizeMode)
 				{
 					H = StrongIntegerHash(H + LightGridData[LoopCount.LightGridOffset + Index]);
 				}
+                if (LoopCount.NumLights == SceneInfiniteLightCount)
+                {
+                    return EmptyColor;
+                }
 				return HUEtoRGB((H & 0xFFFFFF) * 5.96046447754e-08);
 			}
-			case 3:
+			case PATH_TRACER_DEBUG_VIZ_LIGHT_GRID_SINGULAR:
 			{
 				// color by presence of singular lights
 				if (LoopCount.NumLights == SceneInfiniteLightCount)
@@ -109,7 +113,7 @@ float3 LightGridVisualize(FLightLoopCount LoopCount, int VisualizeMode)
 					return float3(0.2, 1, 0.2); // green - only point lights
 				}
 			}
-			case 4:
+			case PATH_TRACER_DEBUG_VIZ_LIGHT_GRID_AXIS:
 			{
 				// color by chosen axis
 				int Axis = (LoopCount.LightGridOffset / LightGridMaxCount) % 3;
@@ -126,8 +130,12 @@ float3 LightGridVisualize(FLightLoopCount LoopCount, int VisualizeMode)
 				{
 					return EmptyColor;
 				}
+                if (N > MaxCount)
+                {
+                    return float3(1, 0.1, 1);
+                }
 				float Max = LightGridMaxCount;
-				float t = saturate(float(N) / Max);
+				float t = saturate(float(N) / MaxCount);
 				return BlueGreenRedColorMap(t);
 			}
 		}
@@ -155,7 +163,7 @@ FLightLoopCount LightGridLookup(float3 TranslatedWorldPos)
 	return Result;
 }
 
-float3 LightGridVisualize(FLightLoopCount LoopCount, int VisualizeMode)
+float3 LightGridVisualize(FLightLoopCount LoopCount, int DebugMode, uint MaxCount)
 {
 	return 0.5;
 }
diff --git a/Engine/Shaders/Private/PathTracing/PathTracingCommon.ush b/Engine/Shaders/Private/PathTracing/PathTracingCommon.ush
--- a/Engine/Shaders/Private/PathTracing/PathTracingCommon.ush
+++ b/Engine/Shaders/Private/PathTracing/PathTracingCommon.ush
@@ -34,6 +34,11 @@
 #define PATH_TRACING_PAYLOAD_INPUT_FLAG_RAY_TYPE_INDIRECT_VOLUME       4
 #define PATH_TRACING_PAYLOAD_INPUT_FLAG_HAS_SCENE_DEPTH               (1 << 4) // Indicates the payload contains a scene depth value (exclusive with dbuffer usage)
 #define PATH_TRACING_PAYLOAD_INPUT_FLAG_IGNORE_TRANSLUCENT            (1 << 5) // Indicates this ray should not process any translucent geometry (so we can skip them in our opaque traces)
+#define PATH_TRACING_PAYLOAD_INPUT_FLAG_DEBUG_INFO_MASK               (3 << 6) // Reserve 2 bits for debug-info (ask material to give us back some debug information)
+#define PATH_TRACING_PAYLOAD_INPUT_FLAG_DEBUG_TRIANGLE_ID             (1 << 6)
+#define PATH_TRACING_PAYLOAD_INPUT_FLAG_DEBUG_INSTANCE_ID             (2 << 6)
+#define PATH_TRACING_PAYLOAD_INPUT_FLAG_DEBUG_CLUSTER_ID              (3 << 6)
+#define PATH_TRACING_PAYLOAD_INPUT_FLAG_UNUSED8                       (1 << 8) // Free for future use
 
 #define PATH_TRACING_PAYLOAD_NORMAL_BITS  		32
 #define PATH_TRACING_PAYLOAD_TANGENT_ANGLE_BITS 16 						// Could be dropped to 12 or even 10 if we need more room in the future
diff --git a/Engine/Shaders/Private/PathTracing/PathTracingCore.ush b/Engine/Shaders/Private/PathTracing/PathTracingCore.ush
--- a/Engine/Shaders/Private/PathTracing/PathTracingCore.ush
+++ b/Engine/Shaders/Private/PathTracing/PathTracingCore.ush
@@ -348,28 +348,24 @@ struct FPathState
 	}
 #endif
 
-};
 
-void AccumulateRadiance(inout float3 TotalRadiance, float3 PathRadiance)
-{
-	PathRadiance *= BaseExposure;
-	float MaxPathRadiance = max3(PathRadiance.x, PathRadiance.y, PathRadiance.z);
+    void AccumulateRadiance(float3 PathRadiance)
+    {
+        PathRadiance *= BaseExposure;
 #if PATH_TRACER_KERNEL_TYPE == PATH_TRACER_KERNEL_CAMERA
-	float MaxIntensity = MaxHalfFloat; // camera rays, clamp at a fairly high level for half-float final storage into SceneColor buffer
+        // camera rays, clamp at a fairly high level for half-float final storage into SceneColor buffer
+        float MaxIntensity = MaxHalfFloat;
 #else
-	// User asked for path contributions to be clamped to reduce fireflies from indirect paths
-	// This puts an upper bound on variance within the pixel at the expense of bias
-	float MaxIntensity = MaxPathIntensity;
+        // User asked for path contributions to be clamped to reduce fireflies from indirect paths
+        // This puts an upper bound on variance within the pixel at the expense of bias
+        float MaxIntensity = MaxPathIntensity;
 #endif
-	if (MaxPathRadiance > MaxIntensity)
-	{
-		// adjust sample while keeping color
-		PathRadiance *= MaxIntensity / MaxPathRadiance;
-	}
-
-	// accumulate into the total
-	TotalRadiance += PathRadiance;
-}
+        // compute factor that will clamp the intensity but preserve the color
+        float Adjust = MaxIntensity / max(max(MaxIntensity, PathRadiance.x), max(PathRadiance.y, PathRadiance.z));
+        // accumulate into the total
+        Radiance += PathRadiance * Adjust;
+    }
+};
 
 bool HasDecalData(FDecalShaderPayload DecalPayload)
 {
@@ -819,7 +815,7 @@ FPackedPathTracingPayload TraceTransparentRay(inout FPathState PathState, int Bo
 					PathState.GetNormal()
 				);
 				HitSample = CloudPayload.GetVolumetricCallableShaderRISContext();
-				AccumulateRadiance(PathState.Radiance, CloudPayload.GetVolumetricCallableShaderOutputRadiance());
+				PathState.AccumulateRadiance(CloudPayload.GetVolumetricCallableShaderOutputRadiance());
 			#if PATH_TRACER_KERNEL_TYPE == PATH_TRACER_KERNEL_CAMERA
 				PathState.Alpha += CloudPayload.GetVolumetricCallableShaderOutputAlpha();
 			#endif
@@ -933,7 +929,7 @@ FPackedPathTracingPayload TraceTransparentRay(inout FPathState PathState, int Bo
 						}
 					}
 					// accumulate all at once so clamping applies to entire volume contribution
-					AccumulateRadiance(PathState.Radiance, VolumeRadiance);
+					PathState.AccumulateRadiance(VolumeRadiance);
 					// restore interval start now that we are finished with ray marching
 					Interval.VolumeTMin = StartVolumeTMin;
 
@@ -980,7 +976,7 @@ FPackedPathTracingPayload TraceTransparentRay(inout FPathState PathState, int Bo
 		// add in surface emission
 		if (PathState.ShouldAccumulateEmissive())
 		{
-			AccumulateRadiance(PathState.Radiance, PathState.PathThroughput * PackedPayload.UnpackRadiance());
+			PathState.AccumulateRadiance(PathState.PathThroughput * PackedPayload.UnpackRadiance());
 		}
 
 		if (PackedPayload.HitT == RAY_DEFAULT_T_MAX)
@@ -1832,7 +1828,7 @@ bool PathTracingKernel(inout FPathState PathState, int Bounce)
 							}
 							LightContrib *= TraceTransparentVisibilityRay(LightRay, 1.0, PathState.GetPathRoughness(), MissShaderIndex, bCastShadows, bCloudCastShadows, CloudFactor, PathState.RandSequence);
 
-							AccumulateRadiance(PathState.Radiance, LightContrib);
+							PathState.AccumulateRadiance(LightContrib);
 						}
 					}
 				}
@@ -1857,7 +1853,7 @@ bool PathTracingKernel(inout FPathState PathState, int Bounce)
 							}
 							MaterialRay.TMax = LightResult.HitT;
 							LightContrib *= TraceTransparentVisibilityRay(MaterialRay, 1.0, PathState.GetPathRoughness(), MissShaderIndex, bCastShadows, bCloudCastShadows, CloudFactor, PathState.RandSequence);
-							AccumulateRadiance(PathState.Radiance, LightContrib);
+							PathState.AccumulateRadiance(LightContrib);
 						}
 					}
 				}
@@ -1958,11 +1954,11 @@ bool PathTracingKernel(inout FPathState PathState, int Bounce)
 			int LightId;
 			float LightPickPdf = 0;
 
-			SelectLight(RandSample.x * LightPickingCdfSum, LightLoopCount.NumLights, LightPickingCdf, LightId, LightPickPdf);
+			SelectLight(RandSample.z * LightPickingCdfSum, LightLoopCount.NumLights, LightPickingCdf, LightId, LightPickPdf);
 			LightPickPdf /= LightPickingCdfSum;
 
 			LightId = GetLightId(LightId, LightLoopCount);
-			FLightSample LightSample = SampleLight(LightId, RandSample.yz, TranslatedWorldPos, WorldNormal);
+			FLightSample LightSample = SampleLight(LightId, RandSample.xy, TranslatedWorldPos, WorldNormal);
 
 			LightSample.RadianceOverPdf /= LightPickPdf;
 			LightSample.Pdf *= LightPickPdf;
@@ -2009,7 +2005,7 @@ bool PathTracingKernel(inout FPathState PathState, int Bounce)
 					const uint MissShaderIndex = GetLightMissShaderIndex(LightId);
 					LightContrib *= TraceTransparentVisibilityRay(LightRay, AvgRoughness, PathState.GetPathRoughness(), MissShaderIndex, bCastShadows, bCloudCastShadows, CloudFactor, PathState.RandSequence);
 
-					AccumulateRadiance(PathState.Radiance, LightContrib);
+					PathState.AccumulateRadiance(LightContrib);
 				}
 			}
 		}
@@ -2170,7 +2166,7 @@ bool PathTracingKernel(inout FPathState PathState, int Bounce)
 				const uint MissShaderIndex = GetLightMissShaderIndex(LightId);
 				const float ShadowRayRoughness = ApproximateCaustics ? PathState.GetPathRoughness() : 0.0;
 				LightContrib *= TraceTransparentVisibilityRay(LightRay, ShadowRayRoughness, PreviousPathRoughness, MissShaderIndex, bCastShadows, bCloudCastShadows, CloudFactor, PathState.RandSequence);
-				AccumulateRadiance(PathState.Radiance, LightContrib);
+				PathState.AccumulateRadiance(LightContrib);
 			}
 		}
 	}
diff --git a/Engine/Shaders/Private/PathTracing/PathTracingDebug.usf b/Engine/Shaders/Private/PathTracing/PathTracingDebug.usf
--- a/Engine/Shaders/Private/PathTracing/PathTracingDebug.usf
+++ b/Engine/Shaders/Private/PathTracing/PathTracingDebug.usf
@@ -2,24 +2,37 @@
 
 #define RANDSEQ_SFC_TEXTURE 1
 
-#define LIGHT_LOOP_METHOD 	2 // 0: stochastic (like ref PT) 1: loop over all lights 2: RIS (take N samples but trace a single shadow ray)
+#define LIGHT_LOOP_METHOD           2 // 0: stochastic (like ref PT) 1: loop over all lights 2: RIS (take 1 sample per light, but only keep one to trace a single shadow ray)
 
-#define USE_BSDF_SAMPLING 1 // 0: light sampling only (no MIS) 1: light sampling and BSDF sampling (with MIS)
+#define EVALUATE_MATERIAL_BSDF      1 // Use the full material eval/sample when true, otherwise substitute a plain lambert calculation for speed
+#define USE_BSDF_SAMPLING           1 // 0: light sampling only (no MIS) 1: light sampling and BSDF sampling (with MIS)
+
+#if LIGHT_LOOP_METHOD == 0
+#define LIGHT_SELECTION_METHOD      1 // 0: uniform sampling, 1: CDF (keeps stratification), 2: RIS (destroys stratification but no array needed)
+#define BSDF_SAMPLING_METHOD        0 // 0: trace all, 1: use RIS among all valid light hits (seems slower (!))
+#endif
+
+
+#if LIGHT_LOOP_METHOD == 2
+#define ACCUMULATE_UNSHADOWED_LIGHTS 1 // When true, in LIGHT_LOOP_METHOD == 2, accumulate any unshadowed light directly, skipping RIS to reduce noise
+#define FORCE_TRACE_INFINITE_LIGHTS  1 // When true, in LIGHT_LOOP_METHOD == 2, trace a ray for infinite lights instead of letting them do RIS
+#endif
+
+#define CLAMP_PATH_INTENSITY        1 // Apply MaxPathIntensity to limit intensity of indirect paths like the reference path tracer
 
 #define USE_PATH_TRACING_LIGHT_GRID	1
 #define USE_RAY_TRACING_DECAL_GRID	1
 
-#ifndef PATH_TRACER_PRIMARY_RAYS // Should be set by permutation
-#define PATH_TRACER_PRIMARY_RAYS 0
+#ifndef PATH_TRACER_PRIMARY_RAYS            // Should be set by permutation (set for editor visibility)
+#define PATH_TRACER_PRIMARY_RAYS    1
 #endif
 
-#ifndef PATH_TRACER_INCLUDE_TRANSLUCENT // Should be set by permutation
-#define PATH_TRACER_INCLUDE_TRANSLUCENT   0 // include objects tagged as TRANSLUCENT when tracing rays?
+#ifndef PATH_TRACER_INCLUDE_TRANSLUCENT     // Should be set by permutation (set for editor visibility)
+#define PATH_TRACER_INCLUDE_TRANSLUCENT   1 // include objects tagged as TRANSLUCENT when tracing rays?
 #endif
 
 
 #include "../Common.ush"
-#include "../BlueNoise.ush"
 #include "../RayTracing/RayTracingCommon.ush"
 #include "../RayTracing/RayTracingDecalGrid.ush"
 #include "PathTracingCommon.ush"
@@ -49,6 +62,13 @@ int SampleIndex;
 float SampleBlendFactor;
 int MaxBounces;
 int EnableEmissive;
+#if CLAMP_PATH_INTENSITY
+float MaxPathIntensity;
+#endif
+#if PATH_TRACER_PRIMARY_RAYS == 0
+float AmbientOcclusionMaxDistance;
+uint VisualizeMaxCount;
+#endif
 #if PATH_TRACER_KERNEL_TYPE != PATH_TRACER_KERNEL_CAMERA
 int FirstBounce;
 #endif
@@ -368,8 +388,12 @@ struct FPathState
 		float4 OldSpecular = SampleIndex > 0 ? RWSceneSpecularColor[TextureIndex]   : 0;
 		float4 OldNormal   = SampleIndex > 0 ? RWSceneNormal[TextureIndex]   : 0;
 		float OldRoughness = SampleIndex > 0 ? RWSceneRoughness[TextureIndex] : 0;
-
-		const float4 PixelValue = float4(Radiance * View.PreExposure, Alpha);
+#if CLAMP_PATH_INTENSITY
+        // Clamping is done with PreExposure folded in so that clamping is relative to current exposure level
+		const float4 PixelValue = float4(Radiance, Alpha);
+#else
+        const float4 PixelValue = float4(Radiance * View.PreExposure, Alpha);
+#endif
 
 		const float3 Diffuse = GetDiffuseColor();
 		const float3 Specular = GetSpecularColor();
@@ -392,6 +416,28 @@ struct FPathState
 		RWSceneDepth[TextureIndex] = max(OldDepth, ConvertToDeviceZ(DepthZ));
 	}
 	#endif
+
+    void AccumulateRadiance(float3 PathRadiance)
+    {
+#if CLAMP_PATH_INTENSITY
+        PathRadiance *= View.PreExposure;
+    #if PATH_TRACER_KERNEL_TYPE == PATH_TRACER_KERNEL_CAMERA
+        // camera rays, clamp at a fairly high level for half-float final storage into SceneColor buffer
+        float MaxIntensity = MaxHalfFloat;
+    #else
+        // User asked for path contributions to be clamped to reduce fireflies from indirect paths
+        // This puts an upper bound on variance within the pixel at the expense of bias
+        float MaxIntensity = MaxPathIntensity;
+    #endif
+        // compute factor that will clamp the intensity but preserve the color
+        float Adjust = MaxIntensity / max(max(MaxIntensity, PathRadiance.x), max(PathRadiance.y, PathRadiance.z));
+        // accumulate into the total
+        Radiance += PathRadiance * Adjust;
+#else
+        Radiance += PathRadiance;
+#endif
+    }
+
 };
 
 
@@ -429,7 +475,7 @@ FPackedPathTracingPayload TraceTransparentRay(inout FPathState PathState, int Bo
 	FRISContext HitSample = InitRISContext(RandSample.x);
 
 	float3 ResultThroughput = 0;
-	FPackedPathTracingPayload ResultPayload = (FPackedPathTracingPayload) 0;
+    FPackedPathTracingPayload ResultPayload; ResultPayload.SetMiss();
 	FPackedPathTracingPayload OpaquePackedPayload = InitPathTracingPayload(PathState.GetFirstScatterType(), PathState.GetPathRoughness());
 	OpaquePackedPayload.SetDBufferA(float4(0, 0, 0, 1));
 	OpaquePackedPayload.SetDBufferB(float4(0, 0, 0, 1));
@@ -585,7 +631,7 @@ FPackedPathTracingPayload TraceTransparentRay(inout FPathState PathState, int Bo
 
 		if (PATH_TRACER_KERNEL_TYPE == PATH_TRACER_KERNEL_CAMERA || EnableEmissive)
 		{
-			PathState.Radiance += PathState.PathThroughput * PackedPayload.UnpackRadiance();
+            PathState.AccumulateRadiance(PathState.PathThroughput * PackedPayload.UnpackRadiance());
 		}
 
 	#if PATH_TRACER_KERNEL_TYPE == PATH_TRACER_KERNEL_CAMERA
@@ -679,7 +725,21 @@ FPackedPathTracingPayload TraceTransparentRay(inout FPathState PathState, int Bo
 	return ResultPayload;
 }
 
-
+float LightContribToSelectionWeight(float3 C)
+{
+#if 0
+    return C.x + C.y + C.z;
+#elif 1
+    // Seems to behave best overall
+    return Luminance(C);
+#elif 0
+    return dot(C, float3(0.25, 0.5, 0.25));
+#elif 0
+    return max3(C.x, C.y, C.z);
+#else
+    return 1;
+#endif
+}
 
 bool PathTracingDebugCore(inout FPathState PathState, int Bounce)
 {
@@ -712,155 +772,344 @@ bool PathTracingDebugCore(inout FPathState PathState, int Bounce)
 	// Generate main random number for direct+indirect sampling
 	float4 RandSample = PathState.RandSequence.Get4D();
 
+#if EVALUATE_MATERIAL_BSDF
+	// Sample real material
+	FMaterialSample MaterialSample = SampleMaterial(-PathState.Ray.Direction, Payload, RandSample.xyz);
+#else
+    // Hard-coded diffuse lobe for benchmarking
+	const float4 SampledValue = CosineSampleHemisphereWorld(RandSample.xy, Payload.WorldNormal);
+    FMaterialSample MaterialSample = CreateMaterialSample(SampledValue.xyz, Payload.BSDFOpacity * Payload.DiffuseColor, SampledValue.w, 1.0, 1.0, PATHTRACER_SCATTER_DIFFUSE);
+#endif
+
 #if LIGHT_LOOP_METHOD == 0
+
+    #if LIGHT_SELECTION_METHOD == 0
+    float LightPickPdf = rcp(LightLoopCount.NumLights);
+    {
+        int LightId = GetLightId(min(int(RandSample.z * LightLoopCount.NumLights), LightLoopCount.NumLights - 1), LightLoopCount);
+    #else
+        #if LIGHT_SELECTION_METHOD == 1
+         // Store CDF in an array, then sample from it (maintains stratification)
 	float LightPickingCdf[RAY_TRACING_LIGHT_COUNT_MAXIMUM];
-	float LightPickingCdfSum = 0;
+    float LightPickingCdfSum = 0;
+        #elif LIGHT_SELECTION_METHOD == 2
+        // Use RIS to pick a light in a single pass (worse stratifiication, but less storage)
+    FRISContext LightPicker = InitRISContext(RandSample.z);
+    int SampledLightId = 0;
+    float LightPickPdf = 0;
+        #endif
 
 	for (uint Index = 0, Num = LightLoopCount.NumLights; Index < Num; ++Index)
 	{
-		uint LightIndex = GetLightId(Index, LightLoopCount);
-		float LightEstimate = EstimateLight(LightIndex, TranslatedWorldPos, WorldNormal, PrimitiveLightingChannelMask, false);
+		uint LightId = GetLightId(Index, LightLoopCount);
+		float LightEstimate = EstimateLight(LightId, TranslatedWorldPos, WorldNormal, PrimitiveLightingChannelMask, false);
+        LightEstimate *= PathState.HasMadeDiffuseOrSpecularScatter() ? GetLightIndirectScale(LightId, PathState.GetPathRoughness()) : 1.0;
+        #if LIGHT_SELECTION_METHOD == 1
 		LightPickingCdfSum += LightEstimate;
 		LightPickingCdf[Index] = LightPickingCdfSum;
+        #elif LIGHT_SELECTION_METHOD == 2
+        if (LightPicker.Accept(LightEstimate))
+        {
+            // Store this index as our chosen sample (so far)
+            SampledLightId = LightId;
+            LightPickPdf = LightEstimate;
+        }
+        #endif
 	}
 
-	FMaterialSample MaterialSample = SampleMaterial(-PathState.Ray.Direction, Payload, RandSample.xyz);
+        #if LIGHT_SELECTION_METHOD == 1
 	if (LightPickingCdfSum > 0)
 	{
 		// init worked
-		int LightId;
+        int SampledIndex;
 		float LightPickPdf = 0;
 
-		SelectLight(RandSample.x * LightPickingCdfSum, LightLoopCount.NumLights, LightPickingCdf, LightId, LightPickPdf);
+		SelectLight(RandSample.z * LightPickingCdfSum, LightLoopCount.NumLights, LightPickingCdf, SampledIndex, LightPickPdf);
 		LightPickPdf /= LightPickingCdfSum;
-		float2 LightRandSample = RandSample.yz;
-		LightId = GetLightId(LightId, LightLoopCount);
-#else // LIGHT_LOOP_METHOD != 0
+        int LightId = GetLightId(SampledIndex, LightLoopCount);
+        #elif LIGHT_SELECTION_METHOD == 2
+    if (LightPicker.HasSample())
+    {
+        LightPickPdf /= LightPicker.GetNormalization();
+        int LightId = SampledLightId;
+        #endif // LIGHT_SELECTION_METHOD
+    #endif // LIGHT_SELECTION_METHOD
 		float2 LightRandSample = RandSample.xy;
+		
+        // Generate a sample on the light source
+        FLightSample LightSample = SampleLight(LightId, LightRandSample, TranslatedWorldPos, WorldNormal);
+        if (LightSample.Pdf > 0)
+        {
+            LightSample.RadianceOverPdf /= LightPickPdf;
+            LightSample.Pdf *= LightPickPdf;
+
+            float2 LightDiffuseSpecularScale = float2(GetLightDiffuseScale(LightId), GetLightSpecularScale(LightId));
+    #if EVALUATE_MATERIAL_BSDF
+            // Evaluate real material
+            FMaterialEval MaterialEval = EvalMaterial(-PathState.Ray.Direction, LightSample.Direction, Payload, LightDiffuseSpecularScale);
+    #else
+            // Hard-coded diffuse lobe for benchmarking
+            FMaterialEval MaterialEval = CreateMaterialEval(Payload.DiffuseColor * LightDiffuseSpecularScale.x, saturate(dot(WorldNormal, LightSample.Direction)) / PI);
+    #endif
+            float3 LightContrib = PathState.PathThroughput * LightSample.RadianceOverPdf * MaterialEval.Weight * MaterialEval.Pdf;
+            LightContrib *= PathState.HasMadeDiffuseOrSpecularScatter() ? GetLightIndirectScale(LightId, PathState.GetPathRoughness()) : 1.0;
+
+    #if USE_BSDF_SAMPLING
+            LightContrib *= MISWeightPower(LightSample.Pdf, MaterialEval.Pdf);
+    #endif
+            // Trace shadow ray immediately
+            FRayDesc ShadowRay;
+            ShadowRay.Origin = TranslatedWorldPos;
+            ShadowRay.Direction = LightSample.Direction;
+            ShadowRay.TMin = 0;
+            ShadowRay.TMax = LightSample.Distance;
+            const float PositionBiasSign = Payload.IsMaterialTransmissive() ? sign(dot(Payload.WorldGeoNormal, ShadowRay.Direction)) : 1.0;
+            const bool bCastShadows = CastsShadow(LightId);
+            const uint MissShaderIndex = GetLightMissShaderIndex(LightId);
+            ApplyRayBias(ShadowRay.Origin, PackedPayload.HitT, PositionBiasSign * PackedPayload.UnpackWorldGeoNormal());
+            LightContrib *= TraceShadowRay(ShadowRay, GetAverageRoughness(Payload), PathState.GetPathRoughness(), MissShaderIndex, bCastShadows);
+            PathState.AccumulateRadiance(LightContrib);
+        }
+    } // LightPickingCdfSum > 0
+    #if USE_BSDF_SAMPLING == 1
+
+        #if BSDF_SAMPLING_METHOD == 1
+    FRISContext BSDFSamplePicker = InitRISContext(RandSample.w);
+    float SampledTMax = 0;
+    float3 SampledLightContrib = 0;
+    int SampledLightId = 0;
+        #endif
+
+        #if LIGHT_SELECTION_METHOD == 1
+    float PreviousCdfValue = 0.0;
+        #endif
+    for (uint Index = 0, Num = LightLoopCount.NumMISLights; Index < Num; ++Index)
+    {
+        const uint LightId = GetLightId(Index, LightLoopCount);
+        #if LIGHT_SELECTION_METHOD == 1
+        const float CdfValue = LightPickingCdf[Index];
+        const float LightPickPdf = (CdfValue - PreviousCdfValue) / LightPickingCdfSum;
+        PreviousCdfValue = CdfValue;
+        #endif
+        FRayDesc ShadowRay;
+        ShadowRay.Origin = TranslatedWorldPos;
+        ShadowRay.Direction = MaterialSample.Direction;
+        ShadowRay.TMin = 0;
+        ShadowRay.TMax = RAY_DEFAULT_T_MAX;
+
+        FLightHit LightResult = TraceLight(ShadowRay, LightId);
+        if (LightResult.IsHit())
+        {
+            ShadowRay.TMax = LightResult.HitT;
+            float3 LightContrib = PathState.PathThroughput * MaterialSample.Weight * LightResult.Radiance;
+            LightContrib *= MaterialSample.ScatterType == PATHTRACER_SCATTER_DIFFUSE ? GetLightDiffuseScale(LightId) : GetLightSpecularScale(LightId);
+            float IndirectScale = PathState.HasMadeDiffuseOrSpecularScatter() ? GetLightIndirectScale(LightId, PathState.GetPathRoughness()) : 1.0;
+            LightContrib *= IndirectScale;
+
+            #if LIGHT_SELECTION_METHOD == 2
+            // We didn't store the light estimate in an array, so need to redo it now to get the probability of having picked this light
+            float LightPickPdf = 0;
+            if (LightPicker.HasSample())
+            {
+                float LightEstimate = EstimateLight(LightId, TranslatedWorldPos, WorldNormal, PrimitiveLightingChannelMask, false);
+                LightEstimate *= IndirectScale;
+                LightPickPdf = LightEstimate / LightPicker.GetNormalization();
+            }
+            #endif
+            LightContrib *= MISWeightPower(MaterialSample.Pdf, LightResult.Pdf * LightPickPdf);
+            #if BSDF_SAMPLING_METHOD == 0
+            const bool bCastShadows = CastsShadow(LightId);
+            const uint MissShaderIndex = GetLightMissShaderIndex(LightId);
+            ApplyRayBias(ShadowRay.Origin, PackedPayload.HitT, MaterialSample.PositionBiasSign * PackedPayload.UnpackWorldGeoNormal());
+            LightContrib *= TraceShadowRay(ShadowRay, MaterialSample.Roughness, PathState.GetPathRoughness(), MissShaderIndex, bCastShadows);
+            PathState.AccumulateRadiance(LightContrib);
+            #elif BSDF_SAMPLING_METHOD == 1
+            float SelectionWeight = LightContribToSelectionWeight(LightContrib);
+            if (BSDFSamplePicker.Accept(SelectionWeight))
+            {
+                SampledLightId = LightId;
+                SampledLightContrib = LightContrib / SelectionWeight;
+                SampledTMax = ShadowRay.TMax;
+            }
+            #endif
+        }
+    }
+        #if BSDF_SAMPLING_METHOD == 1
+    if (BSDFSamplePicker.HasSample())
+    {
+        SampledLightContrib *= BSDFSamplePicker.GetNormalization();
+        const bool bCastShadows = CastsShadow(SampledLightId);
+        const uint MissShaderIndex = GetLightMissShaderIndex(SampledLightId);
+        FRayDesc ShadowRay;
+        ShadowRay.Origin = TranslatedWorldPos;
+        ShadowRay.Direction = MaterialSample.Direction;
+        ShadowRay.TMin = 0;
+        ShadowRay.TMax = SampledTMax;
+        ApplyRayBias(ShadowRay.Origin, PackedPayload.HitT, MaterialSample.PositionBiasSign * PackedPayload.UnpackWorldGeoNormal());
+        SampledLightContrib *= TraceShadowRay(ShadowRay, MaterialSample.Roughness, PathState.GetPathRoughness(), MissShaderIndex, bCastShadows);
+        PathState.AccumulateRadiance(SampledLightContrib);
+    }
+        #endif
+    #endif // USE_BSDF_SAMPLING == 1
 
-#if LIGHT_LOOP_METHOD == 2
-		FRISContext TraceSample = InitRISContext(RandSample.z);
-		FRayDesc ShadowRaySample;
-		float3 ShadowRayContrib = 0;
-		float ShadowRayRoughness = 0;
-		uint ShadowRayLightId;
-#endif
-		FMaterialSample MaterialSample = SampleMaterial(-PathState.Ray.Direction, Payload, RandSample.xyz);
-		for (uint Index = 0, Num = LightLoopCount.NumLights; Index < Num; ++Index)
-		{
-			uint LightId = GetLightId(Index, LightLoopCount);
-			const float LightPickPdf = 1;
-#endif // LIGHT_LOOP_METHOD != 0
-
-			// Generate a sample on the light source
-			FLightSample LightSample = SampleLight(LightId, LightRandSample, TranslatedWorldPos, WorldNormal);
-			if (LightSample.Pdf > 0)
-			{
-				LightSample.RadianceOverPdf /= LightPickPdf;
-				LightSample.Pdf *= LightPickPdf;
-
-#if 1
-				// Evaluate real material
-				float2 LightDiffuseSpecularScale = float2(GetLightDiffuseScale(LightId), GetLightSpecularScale(LightId));
-				FMaterialEval MaterialEval = EvalMaterial(-PathState.Ray.Direction, LightSample.Direction, Payload, LightDiffuseSpecularScale);
-				float3 LightContrib = PathState.PathThroughput * LightSample.RadianceOverPdf * MaterialEval.Weight * MaterialEval.Pdf;
-				uint ScatterType = PathState.GetFirstScatterType();
-				if (ScatterType == PATHTRACER_SCATTER_DIFFUSE || ScatterType == PATHTRACER_SCATTER_SPECULAR)
-				{
-					LightContrib *= GetLightIndirectScale(LightId, PathState.GetPathRoughness());
-				}
-
-#if USE_BSDF_SAMPLING
-				LightContrib *= MISWeightPower(LightSample.Pdf, MaterialEval.Pdf);
-#endif
-
-#else
-				// Hard-coded diffuse lobe for benchmarking
-				float3 MaterialEvalWeight = PackedPayload.UnpackDiffuseColor();
-				float  MaterialEvalPdf = saturate(dot(WorldNormal, LightSample.Direction)) / PI;
-				float3 LightContrib = PathThroughput * LightSample.RadianceOverPdf * MaterialEvalWeight * MaterialEvalPdf;
-#endif
-				{
-					FRayDesc ShadowRay;
-					ShadowRay.Origin = TranslatedWorldPos;
-					ShadowRay.Direction = LightSample.Direction;
-					ShadowRay.TMin = 0;
-					ShadowRay.TMax = LightSample.Distance;
-
-					const float SignedPositionBias = Payload.IsMaterialTransmissive() ? sign(dot(Payload.WorldGeoNormal, LightSample.Direction)) : 1.0;
-					ApplyRayBias(ShadowRay.Origin, PackedPayload.HitT, SignedPositionBias * PackedPayload.UnpackWorldGeoNormal());
+#else // LIGHT_LOOP_METHOD != 0
+    float2 LightRandSample = RandSample.xy;
 
-#if LIGHT_LOOP_METHOD == 2
-					float SelectionWeight = LightContrib.x + LightContrib.y + LightContrib.z;
-					if (TraceSample.Accept(SelectionWeight))
-					{
-						ShadowRaySample = ShadowRay;
-						ShadowRayContrib = LightContrib / SelectionWeight;
-						ShadowRayRoughness = GetAverageRoughness(Payload);
-						ShadowRayLightId = LightId;
-					}
-#else
-					// Trace shadow ray immediately
-					const bool bCastShadows = CastsShadow(LightId);
-					const uint MissShaderIndex = GetLightMissShaderIndex(LightId);
-					PathState.Radiance += LightContrib * TraceShadowRay(ShadowRay, GetAverageRoughness(Payload), PathState.GetPathRoughness(), MissShaderIndex, bCastShadows);
-#endif
-				}
-			}
+    #if LIGHT_LOOP_METHOD == 2
+    FRISContext TraceSample = InitRISContext(RandSample.z);
+    float4 ShadowRayDirLength = 0;
+	float3 ShadowRayContrib = 0;
+	uint ShadowRayLightId = 0;
+    #endif
 
-#if USE_BSDF_SAMPLING
-			// Generate a sample by intersecting a BSDF ray with the light
+    for (uint Index = 0, Num = LightLoopCount.NumLights; Index < Num; ++Index)
+	{
+        uint LightId = GetLightId(Index, LightLoopCount);
+        const bool bCastShadows = CastsShadow(LightId);
+        const uint MissShaderIndex = GetLightMissShaderIndex(LightId);
+        float IndirectScale = PathState.HasMadeDiffuseOrSpecularScatter() ? GetLightIndirectScale(LightId, PathState.GetPathRoughness()) : 1.0;
+        // Generate a sample on the light source
+        FLightSample LightSample = SampleLight(LightId, LightRandSample, TranslatedWorldPos, WorldNormal);
+        if (LightSample.Pdf > 0)
+        {
+            float2 LightDiffuseSpecularScale = float2(GetLightDiffuseScale(LightId), GetLightSpecularScale(LightId));
+
+    #if EVALUATE_MATERIAL_BSDF
+            // Evaluate real material
+            FMaterialEval MaterialEval = EvalMaterial(-PathState.Ray.Direction, LightSample.Direction, Payload, LightDiffuseSpecularScale);
+    #else
+            // Hard-coded diffuse lobe for benchmarking
+            FMaterialEval MaterialEval = CreateMaterialEval(Payload.DiffuseColor * LightDiffuseSpecularScale.x, saturate(dot(WorldNormal, LightSample.Direction)) / PI);
+    #endif
+            float3 LightContrib = PathState.PathThroughput * LightSample.RadianceOverPdf * MaterialEval.Weight * MaterialEval.Pdf;
+            LightContrib *= IndirectScale;
+
+    #if USE_BSDF_SAMPLING
+            LightContrib *= MISWeightPower(LightSample.Pdf, MaterialEval.Pdf);
+    #endif
+
+    #if LIGHT_LOOP_METHOD == 2
+        #if ACCUMULATE_UNSHADOWED_LIGHTS
+            if (bCastShadows == false && MissShaderIndex == 0)
+            {
+                PathState.AccumulateRadiance(LightContrib);
+            }
+            else
+        #endif
+        #if FORCE_TRACE_INFINITE_LIGHTS
+            if (Index < SceneInfiniteLightCount)
+            {
+                // Trace shadow ray immediately
+                FRayDesc ShadowRay;
+                ShadowRay.Origin = TranslatedWorldPos;
+                ShadowRay.Direction = LightSample.Direction;
+                ShadowRay.TMin = 0;
+                ShadowRay.TMax = LightSample.Distance;
+                const float PositionBiasSign = Payload.IsMaterialTransmissive() ? sign(dot(Payload.WorldGeoNormal, ShadowRay.Direction)) : 1.0;
+                ApplyRayBias(ShadowRay.Origin, PackedPayload.HitT, PositionBiasSign * PackedPayload.UnpackWorldGeoNormal());
+                PathState.AccumulateRadiance(LightContrib * TraceShadowRay(ShadowRay, GetAverageRoughness(Payload), PathState.GetPathRoughness(), MissShaderIndex, bCastShadows));
+            }
+            else
+        #endif
+            {
+                float SelectionWeight = LightContribToSelectionWeight(LightContrib);
+                if (TraceSample.Accept(SelectionWeight))
+                {
+                    ShadowRayDirLength = float4(LightSample.Direction, LightSample.Distance);
+                    ShadowRayContrib = LightContrib / SelectionWeight;
+                    ShadowRayLightId = LightId;
+                }
+            }
+    #else // LIGHT_LOOP_METHOD != 2
+            // Trace shadow ray immediately
+            FRayDesc ShadowRay;
+            ShadowRay.Origin = TranslatedWorldPos;
+            ShadowRay.Direction = LightSample.Direction;
+            ShadowRay.TMin = 0;
+            ShadowRay.TMax = LightSample.Distance;
+            const float PositionBiasSign = Payload.IsMaterialTransmissive() ? sign(dot(Payload.WorldGeoNormal, ShadowRay.Direction)) : 1.0;
+            ApplyRayBias(ShadowRay.Origin, PackedPayload.HitT, PositionBiasSign * Payload.WorldGeoNormal);
+            PathState.AccumulateRadiance(LightContrib * TraceShadowRay(ShadowRay, GetAverageRoughness(Payload), PathState.GetPathRoughness(), MissShaderIndex, bCastShadows));
+    #endif // LIGHT_LOOP_METHOD != 2
+        } // LightSample.Pdf > 0
+
+    #if USE_BSDF_SAMPLING
+        // Generate a sample by intersecting a BSDF ray with the light
+		{
+            FRayDesc ShadowRay;
+			ShadowRay.Origin = TranslatedWorldPos;
+			ShadowRay.Direction = MaterialSample.Direction;
+			ShadowRay.TMin = 0;
+			ShadowRay.TMax = RAY_DEFAULT_T_MAX;
+
+			FLightHit LightResult = TraceLight(ShadowRay, LightId);
+			if (LightResult.IsHit())
 			{
-				FRayDesc ShadowRay;
-				ShadowRay.Origin = TranslatedWorldPos;
-				ShadowRay.Direction = MaterialSample.Direction;
-				ShadowRay.TMin = 0;
-				ShadowRay.TMax = RAY_DEFAULT_T_MAX;
-
-				ApplyRayBias(ShadowRay.Origin, PackedPayload.HitT, MaterialSample.PositionBiasSign * PackedPayload.UnpackWorldGeoNormal());
-
-				FLightHit LightResult = TraceLight(ShadowRay, LightId);
-				if (LightResult.IsHit())
-				{
-					ShadowRay.TMax = LightResult.HitT;
-					float3 LightContrib = PathState.PathThroughput * MaterialSample.Weight * LightResult.Radiance;
-					LightContrib *= MaterialSample.ScatterType == PATHTRACER_SCATTER_DIFFUSE ? GetLightDiffuseScale(LightId) : GetLightSpecularScale(LightId);
-					uint ScatterType = PathState.GetFirstScatterType();
-					if (ScatterType == PATHTRACER_SCATTER_DIFFUSE || ScatterType == PATHTRACER_SCATTER_SPECULAR)
-					{
-						LightContrib *= GetLightIndirectScale(LightId, PathState.GetPathRoughness());
-					}
-					LightContrib *= MISWeightPower(MaterialSample.Pdf, LightResult.Pdf * LightPickPdf);
-#if LIGHT_LOOP_METHOD == 2
-					// Defer tracing shadow until later
-					float SelectionWeight = LightContrib.x + LightContrib.y + LightContrib.z;
-					if (TraceSample.Accept(SelectionWeight))
-					{
-						ShadowRaySample = ShadowRay;
-						ShadowRayContrib = LightContrib / SelectionWeight;
-						ShadowRayRoughness = MaterialSample.Roughness;
-						ShadowRayLightId = LightId;
-					}
-#else
-					const bool bCastShadows = CastsShadow(LightId);
-					const uint MissShaderIndex = GetLightMissShaderIndex(LightId);
-					PathState.Radiance += LightContrib * TraceShadowRay(ShadowRay, MaterialSample.Roughness, PathState.GetPathRoughness(), MissShaderIndex, bCastShadows);
-#endif
-				}
-			}
-#endif // USE_BSDF_SAMPLING
+				ShadowRay.TMax = LightResult.HitT;
+				float3 LightContrib = PathState.PathThroughput * MaterialSample.Weight * LightResult.Radiance;
+				LightContrib *= MaterialSample.ScatterType == PATHTRACER_SCATTER_DIFFUSE ? GetLightDiffuseScale(LightId) : GetLightSpecularScale(LightId);
+                LightContrib *= IndirectScale;
+				LightContrib *= MISWeightPower(MaterialSample.Pdf, LightResult.Pdf);
+        #if LIGHT_LOOP_METHOD == 2
+            #if ACCUMULATE_UNSHADOWED_LIGHTS
+                if (bCastShadows == false && MissShaderIndex == 0)
+                {
+                    PathState.AccumulateRadiance(LightContrib);
+                }
+                else
+            #endif
+            #if  FORCE_TRACE_INFINITE_LIGHTS
+                if (Index < SceneInfiniteLightCount)
+                {
+                    ApplyRayBias(ShadowRay.Origin, PackedPayload.HitT, MaterialSample.PositionBiasSign * PackedPayload.UnpackWorldGeoNormal());
+                    PathState.AccumulateRadiance(LightContrib * TraceShadowRay(ShadowRay, GetAverageRoughness(Payload), PathState.GetPathRoughness(), MissShaderIndex, bCastShadows));
+                }
+                else
+            #endif
+                {
+                    // Defer tracing shadow until later
+                    float SelectionWeight = LightContribToSelectionWeight(LightContrib);
+                    if (TraceSample.Accept(SelectionWeight))
+                    {
+                        ShadowRayDirLength = float4(MaterialSample.Direction, LightResult.HitT);
+                        ShadowRayContrib = LightContrib / SelectionWeight;
+                        ShadowRayLightId = LightId | 0x80000000u;
+                    }
+                }
+        #else // LIGHT_LOOP_METHOD != 2
+                ApplyRayBias(ShadowRay.Origin, PackedPayload.HitT, MaterialSample.PositionBiasSign * PackedPayload.UnpackWorldGeoNormal());
+                PathState.AccumulateRadiance(LightContrib * TraceShadowRay(ShadowRay, MaterialSample.Roughness, PathState.GetPathRoughness(), MissShaderIndex, bCastShadows));
+        #endif
+            } // if (LightResult.IsHit())
 		}
+    #endif // USE_BSDF_SAMPLING
+    } // End light loop
+
 #if LIGHT_LOOP_METHOD == 2
-		// Light loop gave us a valid candidate sample, trace it now
-		if (TraceSample.HasSample())
-		{
-			const bool bCastShadows = CastsShadow(ShadowRayLightId);
-			const uint MissShaderIndex = GetLightMissShaderIndex(ShadowRayLightId);
-			ShadowRayContrib *= TraceSample.GetNormalization();
-			PathState.Radiance += ShadowRayContrib * TraceShadowRay(ShadowRaySample, ShadowRayRoughness, PathState.GetPathRoughness(), MissShaderIndex, bCastShadows);
-		}
+    // Light loop gave us a valid candidate sample, trace it now
+    if (TraceSample.HasSample())
+    {
+        const bool bIsBSDFSample = (ShadowRayLightId & 0x80000000u) != 0;
+        ShadowRayLightId &= ~0x80000000u;
+        const bool bCastShadows = CastsShadow(ShadowRayLightId);
+        const uint MissShaderIndex = GetLightMissShaderIndex(ShadowRayLightId);
+        const float ShadowRayRoughness = bIsBSDFSample ? MaterialSample.Roughness : GetAverageRoughness(Payload);
+        const float PositionBiasSign = bIsBSDFSample ? MaterialSample.PositionBiasSign : (Payload.IsMaterialTransmissive() ? sign(dot(Payload.WorldGeoNormal, ShadowRayDirLength.xyz)) : 1.0);
+        FRayDesc ShadowRay;
+        ShadowRay.Origin = TranslatedWorldPos;
+        ShadowRay.Direction = ShadowRayDirLength.xyz;
+        ShadowRay.TMin = 0;
+        ShadowRay.TMax = ShadowRayDirLength.w;
+
+        ApplyRayBias(ShadowRay.Origin, PackedPayload.HitT, PositionBiasSign * PackedPayload.UnpackWorldGeoNormal());
+
+        ShadowRayContrib *= TraceSample.GetNormalization();
+        PathState.AccumulateRadiance(ShadowRayContrib * TraceShadowRay(ShadowRay, ShadowRayRoughness, PathState.GetPathRoughness(), MissShaderIndex, bCastShadows));
+    }
 #endif
 
+#endif // LIGHT_LOOP_METHOD != 0
+
 	// don't generate another ray on the last bounce or second to last bounce if emissive materials are disabled
 	if (bLastBounce || (Bounce + 1 == MaxBounces && EnableEmissive == 0))
 	{
@@ -958,6 +1207,33 @@ float3 GetPathTracingShadingModelColor(uint PathTracingShadingModelID)
 }
 
 
+bool DebugModeRequiresFakeShading()
+{
+    // Does this debug mode require a bit of "fake" shading to enhance readability?
+    switch (DebugMode)
+    {
+        case PATH_TRACER_DEBUG_VIZ_SHADING_MODEL:
+        case PATH_TRACER_DEBUG_VIZ_LIGHTING_CHANNEL_MASK:
+        case PATH_TRACER_DEBUG_VIZ_WORLD_POSITION:
+        case PATH_TRACER_DEBUG_VIZ_LIGHT_GRID_COUNT:
+        case PATH_TRACER_DEBUG_VIZ_LIGHT_GRID_AXIS:
+        case PATH_TRACER_DEBUG_VIZ_LIGHT_GRID_LIST_HASH:
+        case PATH_TRACER_DEBUG_VIZ_LIGHT_GRID_SINGULAR:
+        case PATH_TRACER_DEBUG_VIZ_DECAL_GRID_COUNT:
+        case PATH_TRACER_DEBUG_VIZ_DECAL_GRID_AXIS:
+        case PATH_TRACER_DEBUG_VIZ_DECAL_GRID_LIST_HASH:
+        case PATH_TRACER_DEBUG_VIZ_VOLUME_LIGHT_COUNT:
+        case PATH_TRACER_DEBUG_VIZ_HITKIND:
+        case PATH_TRACER_DEBUG_VIZ_TRANSPARENCY_COUNT:
+        case PATH_TRACER_DEBUG_VIZ_BACKFACING_NORMAL:
+        case PATH_TRACER_DEBUG_VIZ_TRIANGLE_ID:
+        case PATH_TRACER_DEBUG_VIZ_INSTANCE_ID:
+        case PATH_TRACER_DEBUG_VIZ_CLUSTER_ID:
+            return true;
+    }
+    return false;
+}
+
 FPathTracingResult PathTracingDebugCore(FRayDesc Ray, uint2 LaunchIndex)
 {
 	FPathTracingResult Result = (FPathTracingResult) 0;
@@ -974,7 +1250,15 @@ FPathTracingResult PathTracingDebugCore(FRayDesc Ray, uint2 LaunchIndex)
 	OpaquePackedPayload.SetDBufferB(float4(0, 0, 0, 1));
 	OpaquePackedPayload.SetDBufferC(float4(0, 0, 0, 1));
 	OpaquePackedPayload.SetStochasticSlabRand(RandSample);
-	OpaquePackedPayload.SetFlag(PATH_TRACING_PAYLOAD_INPUT_FLAG_IGNORE_TRANSLUCENT);
+    uint DebugFlag = 0;
+    switch (DebugMode)
+    {
+        case PATH_TRACER_DEBUG_VIZ_TRIANGLE_ID: DebugFlag = PATH_TRACING_PAYLOAD_INPUT_FLAG_DEBUG_TRIANGLE_ID; break;
+        case PATH_TRACER_DEBUG_VIZ_INSTANCE_ID: DebugFlag = PATH_TRACING_PAYLOAD_INPUT_FLAG_DEBUG_INSTANCE_ID; break;
+        case PATH_TRACER_DEBUG_VIZ_CLUSTER_ID:  DebugFlag = PATH_TRACING_PAYLOAD_INPUT_FLAG_DEBUG_CLUSTER_ID; break;
+    }
+	OpaquePackedPayload.SetFlag(PATH_TRACING_PAYLOAD_INPUT_FLAG_IGNORE_TRANSLUCENT | DebugFlag);
+
 
 	TraceRay(
 		TLAS,
@@ -1006,17 +1290,16 @@ FPathTracingResult PathTracingDebugCore(FRayDesc Ray, uint2 LaunchIndex)
 				if (LightProb > 0)
 				{
 					NumLights++;
-					if (NumLights == RAY_TRACING_LIGHT_COUNT_MAXIMUM)
-					{
-						// reached max
-						break;
-					}
 				}
 			}
 		}
 
-		float NumLightsMax = min(SceneLightCount - SceneInfiniteLightCount, RAY_TRACING_LIGHT_COUNT_MAXIMUM);
-		Result.Radiance.rgb = NumLights > 0 ? BlueGreenRedColorMap(saturate(float(NumLights) / NumLightsMax)) : 0.18;
+        Result.Radiance.rgb = 0.18;
+        if (NumLights > 0)
+        {
+            Result.Radiance.rgb = (NumLights > VisualizeMaxCount) ? float3(1, 0.1, 1) : BlueGreenRedColorMap(saturate(float(NumLights) / VisualizeMaxCount));
+        }
+
 		if (OpaquePackedPayload.IsHit())
 		{
 			// add some fake shading to help with definition
@@ -1033,6 +1316,7 @@ FPathTracingResult PathTracingDebugCore(FRayDesc Ray, uint2 LaunchIndex)
 		PackedPayload.SetDBufferB(float4(0, 0, 0, 1));
 		PackedPayload.SetDBufferC(float4(0, 0, 0, 1));
 		PackedPayload.SetStochasticSlabRand(RandSample);
+        PackedPayload.SetFlag(DebugFlag);
 		if (OpaquePackedPayload.IsHit())
 		{
 			const float3 ViewZ = View.ViewToTranslatedWorld[2].xyz;
@@ -1084,9 +1368,6 @@ FPathTracingResult PathTracingDebugCore(FRayDesc Ray, uint2 LaunchIndex)
 			const float3 TranslatedWorldPos = Ray.Origin + PackedPayload.HitT * Ray.Direction;
 			const float3 WorldNormal = HitPayload.WorldNormal;
 
-			FLightLoopCount LightLoopCount = LightGridLookup(TranslatedWorldPos);
-
-			float FakeShading = abs(dot(WorldNormal, Ray.Direction)); // cheap shading, just to give some definition to otherwise "solid color" modes
 			float3 Color = 0;
 			float ExtraAlpha = 0;
 			switch (DebugMode)
@@ -1113,29 +1394,29 @@ FPathTracingResult PathTracingDebugCore(FRayDesc Ray, uint2 LaunchIndex)
 					int3 Coord = int3(round(AbsoluteWorldPosition / CellWidth));
 					uint HashIndex = GridHash(int4(Coord, 0)).x;
 					Color = HUEtoRGB((HashIndex & 0xFFFFFF) * 5.96046447754e-08);
-					Color *= FakeShading;
 					break;
 				}
 				case PATH_TRACER_DEBUG_VIZ_WORLD_TANGENT: 				Color = PackedPayload.UnpackWorldTangent() * 0.5 + 0.5; break;
 				case PATH_TRACER_DEBUG_VIZ_LIGHT_GRID_COUNT:
 				case PATH_TRACER_DEBUG_VIZ_LIGHT_GRID_AXIS:
+                case PATH_TRACER_DEBUG_VIZ_LIGHT_GRID_LIST_HASH:
+                case PATH_TRACER_DEBUG_VIZ_LIGHT_GRID_SINGULAR:
 				{
-
-					Color = LightGridVisualize(LightLoopCount, DebugMode == PATH_TRACER_DEBUG_VIZ_LIGHT_GRID_AXIS ? 4 : 1);
-					Color *= FakeShading;
+                    FLightLoopCount LightLoopCount = LightGridLookup(TranslatedWorldPos);
+					Color = LightGridVisualize(LightLoopCount, DebugMode, VisualizeMaxCount);
 					break;
 				}
 				case PATH_TRACER_DEBUG_VIZ_DECAL_GRID_COUNT:
 				case PATH_TRACER_DEBUG_VIZ_DECAL_GRID_AXIS:
+                case PATH_TRACER_DEBUG_VIZ_DECAL_GRID_LIST_HASH:
 				{
 					FDecalLoopCount DecalLoopCount = DecalGridLookup(TranslatedWorldPos);
-					Color = PackedPayload.IsDecalReceiver() ? DecalGridVisualize(DecalLoopCount, DebugMode == PATH_TRACER_DEBUG_VIZ_DECAL_GRID_AXIS ? 3 : 1) : 0.18;
-					Color *= FakeShading;
+					Color = PackedPayload.IsDecalReceiver() ? DecalGridVisualize(DecalLoopCount, DebugMode, VisualizeMaxCount) : 0.18;
 					break;
 				}
 				case PATH_TRACER_DEBUG_VIZ_HITKIND:
 				{
-					Color = (PackedPayload.IsFrontFace() ? float3(0, 1, 0) : float3(1, 0, 0)) * FakeShading;
+					Color = (PackedPayload.IsFrontFace() ? float3(0, 1, 0) : float3(1, 0, 0));
 					break;
 				}
 				case PATH_TRACER_DEBUG_VIZ_SSS_COLOR:
@@ -1176,10 +1457,10 @@ FPathTracingResult PathTracingDebugCore(FRayDesc Ray, uint2 LaunchIndex)
 				case PATH_TRACER_DEBUG_VIZ_EYE_IRIS_MASK:      Color = (HitPayload.PathTracingShadingModelID == PATH_TRACING_SHADINGMODELID_EYE ) ?             HitPayload.GetEyeIrisMask()      : 0.0; break;
 				case PATH_TRACER_DEBUG_VIZ_EYE_IRIS_NORMAL:    Color = (HitPayload.PathTracingShadingModelID == PATH_TRACING_SHADINGMODELID_EYE ) ? 0.5 + 0.5 * HitPayload.GetEyeIrisNormal()    : 0.0; break;
 				case PATH_TRACER_DEBUG_VIZ_EYE_CAUSTIC_NORMAL: Color = (HitPayload.PathTracingShadingModelID == PATH_TRACING_SHADINGMODELID_EYE ) ? 0.5 + 0.5 * HitPayload.GetEyeCausticNormal() : 0.0; break;
-				case PATH_TRACER_DEBUG_VIZ_HAIR_PRIMITIVE_UV:  Color = (HitPayload.PathTracingShadingModelID == PATH_TRACING_SHADINGMODELID_HAIR) ? float3(HitPayload.GetHairPrimitiveUV(), 0.0) : 0.0; break;
+				case PATH_TRACER_DEBUG_VIZ_HAIR_PRIMITIVE_UV:  Color = (HitPayload.PathTracingShadingModelID == PATH_TRACING_SHADINGMODELID_HAIR) ? float3(HitPayload.GetHairPrimitiveUV(), 0.0) : 0.5; break;
 				case PATH_TRACER_DEBUG_VIZ_BACKFACING_NORMAL:
 				{
-					Color = FakeShading;
+                    Color = 1;
 					float K1 = dot(WorldNormal, Ray.Direction);
 					float K2 = dot(HitPayload.WorldGeoNormal, reflect(Ray.Direction, WorldNormal));
 					if (K1 > 1e-3)
@@ -1190,14 +1471,42 @@ FPathTracingResult PathTracingDebugCore(FRayDesc Ray, uint2 LaunchIndex)
 					{
 						Color = float3(1, 0, 0);
 					}
+					break;
 				}
+                case PATH_TRACER_DEBUG_VIZ_TRIANGLE_ID:
+                case PATH_TRACER_DEBUG_VIZ_INSTANCE_ID:
+                case PATH_TRACER_DEBUG_VIZ_CLUSTER_ID:
+                    // radiance is populated by the material with the appropriate color
+                    Color = HitPayload.Radiance;
+                    break;
 				default:
 				{
 					break;
 				}
 			}
-
-			Result.Radiance.rgb += PathThroughput * Color * HitPayload.BSDFOpacity;
+            float FakeShading = 1;
+            if (DebugModeRequiresFakeShading())
+            {
+                if (AmbientOcclusionMaxDistance > 0)
+                {
+                    FRayDesc AORay;
+                    AORay.Origin = TranslatedWorldPos;
+                    AORay.Direction = CosineSampleHemisphereWorld(RandomSequenceCreate(uint3(LaunchIndex, View.StateFrameIndex), 0, 1).Get2D(), WorldNormal).xyz;
+                    AORay.TMin = 0;
+                    AORay.TMax = AmbientOcclusionMaxDistance;
+                    ApplyRayBias(AORay.Origin, PackedPayload.HitT, HitPayload.WorldGeoNormal);
+                    FakeShading = Luminance(TraceShadowRay(AORay, 1, 0, 0, true));
+                }
+                else
+                {
+					// cheap shading, just to give some definition to otherwise "solid color" modes
+                    FakeShading = abs(dot(WorldNormal, Ray.Direction));
+                }
+            }
+            Color *= PathThroughput * HitPayload.BSDFOpacity;
+            Result.DiffuseColor += Color;
+            Result.WorldNormal += Luminance(PathThroughput) * HitPayload.BSDFOpacity * HitPayload.WorldNormal;
+            Result.Radiance.rgb += Color * FakeShading;
 			Result.Radiance.a += Luminance(PathThroughput) * ExtraAlpha;
 		}
 
@@ -1224,7 +1533,11 @@ FPathTracingResult PathTracingDebugCore(FRayDesc Ray, uint2 LaunchIndex)
 
 	if (DebugMode == PATH_TRACER_DEBUG_VIZ_TRANSPARENCY_COUNT)
 	{
-		Result.Radiance.rgb = TransparentHitCount > 0 ? BlueGreenRedColorMap(saturate((TransparentHitCount - 1) / 16.0)) : 0.5;
+        Result.Radiance.rgb = 0.5;
+        if (TransparentHitCount > 0)
+        {
+            Result.Radiance.rgb = TransparentHitCount > VisualizeMaxCount ? float3(1, 0.1, 1) : BlueGreenRedColorMap(saturate(float(TransparentHitCount - 1) / float(VisualizeMaxCount)));
+        }
 		if (OpaquePackedPayload.IsHit())
 		{
 			// add some fake shading to help with definition
diff --git a/Engine/Shaders/Private/PathTracing/PathTracingMaterialHitShader.usf b/Engine/Shaders/Private/PathTracing/PathTracingMaterialHitShader.usf
--- a/Engine/Shaders/Private/PathTracing/PathTracingMaterialHitShader.usf
+++ b/Engine/Shaders/Private/PathTracing/PathTracingMaterialHitShader.usf
@@ -1335,6 +1335,31 @@ RAY_TRACING_ENTRY_CLOSEST_HIT(PathTracingMaterialCHS,
 	}
 #endif
 
+    // Support for debug modes that require information only avaiable inside CHS
+    int DebugInfo = PackedPayload.GetFlags() & PATH_TRACING_PAYLOAD_INPUT_FLAG_DEBUG_INFO_MASK;
+    if (DebugInfo)
+    {
+        Payload.PathTracingShadingModelID = PATH_TRACING_SHADINGMODELID_LAMBERT;
+        Payload.SetBaseColor(0.0);
+        uint PayloadTriangleID = PrimitiveIndex();
+        uint PayloadInstanceID = InstanceID();
+#if PLATFORM_SUPPORTS_RAYTRACING_CLUSTER_OPS
+        uint PayloadClusterID = GetClusterID();
+#else
+        uint PayloadClusterID = 0;
+#endif
+        uint Bits = StrongIntegerHash(PayloadInstanceID);
+        if (DebugInfo == PATH_TRACING_PAYLOAD_INPUT_FLAG_DEBUG_CLUSTER_ID)
+        {
+            Bits += StrongIntegerHash(PayloadClusterID);
+        }
+        if (DebugInfo == PATH_TRACING_PAYLOAD_INPUT_FLAG_DEBUG_TRIANGLE_ID)
+        {
+            Bits += StrongIntegerHash(PayloadTriangleID);
+        }
+        Payload.Radiance = ((Bits >> int3(0, 10, 20)) & 1023) * (1.0f / 1023);
+    }
+
 	// Output depth for solid, masked 
 	// or when `output depth and velocity' is set and Transparency is below the clip value to match the behavior of the rasterizer path. 
 	// If refraction is enabled, we always output depth (but still check against the transparency to support glass with partial coverage
diff --git a/Engine/Shaders/Private/PathTracing/Utilities/PathTracingRIS.ush b/Engine/Shaders/Private/PathTracing/Utilities/PathTracingRIS.ush
--- a/Engine/Shaders/Private/PathTracing/Utilities/PathTracingRIS.ush
+++ b/Engine/Shaders/Private/PathTracing/Utilities/PathTracingRIS.ush
@@ -17,21 +17,13 @@ struct FRISContext
 	bool Accept(float Weight)
 	{
 		WeightSum += Weight;
-		if (Weight > 0)
-		{
-			const float OneMinusEpsilon = 0.99999994; // 32-bit float just before 1.0
-			RandSample *= WeightSum;
-			if (RandSample < Weight)
-			{
-				RandSample = min(RandSample / Weight, OneMinusEpsilon);
-				return true;
-			}
-			else
-			{
-				RandSample = min((RandSample - Weight) / (WeightSum - Weight), OneMinusEpsilon);
-			}
-		}
-		return false;
+        float Sum = WeightSum > 0 ? WeightSum : 1; // Avoid changing RandSample if we haven't accumulated anything yet
+        RandSample *= Sum;
+        bool Keep = RandSample < Weight;           // Should we accept this sample?
+        RandSample -= Keep ? 0 : Weight;           // Rescale RandSample for re-use
+        RandSample /= Keep ? Weight : Sum - Weight;
+        RandSample = min(RandSample, 0.99999994);  // Make sure we stay in [0,1)
+        return Keep;
 	}
 
 	bool HasSample()
diff --git a/Engine/Shaders/Private/RayTracing/RayTracingDecalGrid.ush b/Engine/Shaders/Private/RayTracing/RayTracingDecalGrid.ush
--- a/Engine/Shaders/Private/RayTracing/RayTracingDecalGrid.ush
+++ b/Engine/Shaders/Private/RayTracing/RayTracingDecalGrid.ush
@@ -7,6 +7,7 @@
 #include "../Common.ush"
 #include "../Visualization.ush"
 #include "/Engine/Shared/RayTracingDefinitions.h"
+#include "/Engine/Shared/PathTracingDefinitions.h"
 
 struct FDecalLoopCount {
 	uint NumDecals;
@@ -54,59 +55,63 @@ FDecalLoopCount DecalGridLookup(float3 TranslatedWorldPos)
 	return Result;
 }
 
-float3 DecalGridVisualize(FDecalLoopCount LoopCount, int VisualizeMode)
+float3 DecalGridVisualize(FDecalLoopCount LoopCount, int DebugMode, uint MaxCount)
 {
 	const float3 OutsideColor = 0.18;
 	const float3 EmptyColor = 0.36;
 
 	if (LoopCount.DecalGridOffset != ~0u)
 	{
-		switch (VisualizeMode)
+		switch (DebugMode)
 		{
-		case 2:
-		{
-			// color by unique decal list
-			uint H = 0;
-			for (int Index = 0, Num = LoopCount.NumDecals; Index < Num; Index++)
-			{
-				H = MurmurMix(H + RayTracingDecals.GridData[LoopCount.DecalGridOffset + Index]);
-			}
-			return IntToColor(H);
-		}
-		case 3:
-		{
-			// color by chosen axis
-			int Axis = (LoopCount.DecalGridOffset / RayTracingDecals.GridMaxCount) % 3;
-			if (Axis == 0) return float3(0.9, 0.2, 0.2);
-			if (Axis == 1) return float3(0.2, 0.9, 0.2);
-			if (Axis == 2) return float3(0.2, 0.2, 0.9);
-			return float3(0.8, 0.2, 0.8); // error
-		}
-		default:
-		{
-			// default mode - color by decal count
-			uint N = LoopCount.NumDecals;
-			if (N < 1)
-			{
-				return EmptyColor;
-			}
-			float Max = RayTracingDecals.GridMaxCount;
-			float t = saturate(float(N) / Max);
-			return BlueGreenRedColorMap(t);
-		}
+            case PATH_TRACER_DEBUG_VIZ_DECAL_GRID_LIST_HASH:
+            {
+                // color by unique decal list
+                uint H = 0;
+                for (int Index = 0, Num = LoopCount.NumDecals; Index < Num; Index++)
+                {
+                    H = MurmurMix(H + RayTracingDecals.GridData[LoopCount.DecalGridOffset + Index]);
+                }
+                if (LoopCount.NumDecals < 1)
+                {
+                    return EmptyColor;
+                }
+                return HUEtoRGB((H & 0xFFFFFF) * 5.96046447754e-08);
+            }
+            case PATH_TRACER_DEBUG_VIZ_DECAL_GRID_AXIS:
+            {
+                // color by chosen axis
+                int Axis = (LoopCount.DecalGridOffset / RayTracingDecals.GridMaxCount) % 3;
+                if (Axis == 0) return float3(0.9, 0.2, 0.2);
+                if (Axis == 1) return float3(0.2, 0.9, 0.2);
+                if (Axis == 2) return float3(0.2, 0.2, 0.9);
+                return float3(0.8, 0.2, 0.8); // error
+            }
+            default:
+            {
+                // default mode - color by decal count
+                uint N = LoopCount.NumDecals;
+                if (N < 1)
+                {
+                    return EmptyColor;
+                }
+                if (N > MaxCount)
+                {
+                    return float3(1, 0.1, 1);
+                }
+                float t = saturate(float(N) / MaxCount);
+                return BlueGreenRedColorMap(t);
+            }
 		}
 	}
 	// outside decal grid bounds
 	return OutsideColor;
 }
 
+
 uint GetDecalId(uint Index, FDecalLoopCount LoopCount)
 {
-	if (Index >= 0)
-	{
-		return RayTracingDecals.GridData[LoopCount.DecalGridOffset + Index];
-	}
-	return Index;
+	return RayTracingDecals.GridData[LoopCount.DecalGridOffset + Index];
 }
 
 #else
@@ -119,7 +124,7 @@ FDecalLoopCount DecalGridLookup(float3 TranslatedWorldPos)
 	return Result;
 }
 
-float3 DecalGridVisualize(FDecalLoopCount LoopCount, int VisualizeMode)
+float3 DecalGridVisualize(FDecalLoopCount LoopCount, int DebugMode, uint MaxCount)
 {
 	return 0.5;
 }
diff --git a/Engine/Shaders/Shared/PathTracingDefinitions.h b/Engine/Shaders/Shared/PathTracingDefinitions.h
--- a/Engine/Shaders/Shared/PathTracingDefinitions.h
+++ b/Engine/Shaders/Shared/PathTracingDefinitions.h
@@ -106,25 +106,31 @@
 #define PATH_TRACER_DEBUG_VIZ_ANISOTROPY				19 
 #define PATH_TRACER_DEBUG_VIZ_LIGHT_GRID_COUNT          20
 #define PATH_TRACER_DEBUG_VIZ_LIGHT_GRID_AXIS           21
-#define PATH_TRACER_DEBUG_VIZ_DECAL_GRID_COUNT 			22
-#define PATH_TRACER_DEBUG_VIZ_DECAL_GRID_AXIS			23
-#define PATH_TRACER_DEBUG_VIZ_VOLUME_LIGHT_COUNT        24
-#define PATH_TRACER_DEBUG_VIZ_HITKIND                   25
-#define PATH_TRACER_DEBUG_VIZ_TRANSPARENCY_COUNT		26
-#define PATH_TRACER_DEBUG_VIZ_SSS_COLOR                 27
-#define PATH_TRACER_DEBUG_VIZ_SSS_RADIUS                28
-#define PATH_TRACER_DEBUG_VIZ_SSS_WEIGHT                29
-#define PATH_TRACER_DEBUG_VIZ_SSS_PHASE                 30
-#define PATH_TRACER_DEBUG_VIZ_FUZZ_COLOR                31
-#define PATH_TRACER_DEBUG_VIZ_FUZZ_ROUGHNESS            32
-#define PATH_TRACER_DEBUG_VIZ_FUZZ_AMOUNT               33
-#define PATH_TRACER_DEBUG_VIZ_BSDF_OPACITY              34
-#define PATH_TRACER_DEBUG_VIZ_SUBSTRATE_WEIGHT_V		35
-#define PATH_TRACER_DEBUG_VIZ_SUBSTRATE_TRANSMITTANCE_N 36
-#define PATH_TRACER_DEBUG_VIZ_SUBSTRATE_COVERAGE_ABOVE_N 37
-#define PATH_TRACER_DEBUG_VIZ_SUBSTRATE_F90 			38
-#define PATH_TRACER_DEBUG_VIZ_EYE_IRIS_MASK             39
-#define PATH_TRACER_DEBUG_VIZ_EYE_CAUSTIC_NORMAL        40
-#define PATH_TRACER_DEBUG_VIZ_EYE_IRIS_NORMAL           41
-#define PATH_TRACER_DEBUG_VIZ_BACKFACING_NORMAL         42
-#define PATH_TRACER_DEBUG_VIZ_HAIR_PRIMITIVE_UV         43
+#define PATH_TRACER_DEBUG_VIZ_LIGHT_GRID_LIST_HASH      22
+#define PATH_TRACER_DEBUG_VIZ_LIGHT_GRID_SINGULAR       23
+#define PATH_TRACER_DEBUG_VIZ_DECAL_GRID_COUNT 			24
+#define PATH_TRACER_DEBUG_VIZ_DECAL_GRID_AXIS			25
+#define PATH_TRACER_DEBUG_VIZ_DECAL_GRID_LIST_HASH      26
+#define PATH_TRACER_DEBUG_VIZ_VOLUME_LIGHT_COUNT        27
+#define PATH_TRACER_DEBUG_VIZ_HITKIND                   28
+#define PATH_TRACER_DEBUG_VIZ_TRANSPARENCY_COUNT		29
+#define PATH_TRACER_DEBUG_VIZ_SSS_COLOR                 30
+#define PATH_TRACER_DEBUG_VIZ_SSS_RADIUS                31
+#define PATH_TRACER_DEBUG_VIZ_SSS_WEIGHT                32
+#define PATH_TRACER_DEBUG_VIZ_SSS_PHASE                 33
+#define PATH_TRACER_DEBUG_VIZ_FUZZ_COLOR                34
+#define PATH_TRACER_DEBUG_VIZ_FUZZ_ROUGHNESS            35
+#define PATH_TRACER_DEBUG_VIZ_FUZZ_AMOUNT               36
+#define PATH_TRACER_DEBUG_VIZ_BSDF_OPACITY              37
+#define PATH_TRACER_DEBUG_VIZ_SUBSTRATE_WEIGHT_V		38
+#define PATH_TRACER_DEBUG_VIZ_SUBSTRATE_TRANSMITTANCE_N 39
+#define PATH_TRACER_DEBUG_VIZ_SUBSTRATE_COVERAGE_ABOVE_N 40
+#define PATH_TRACER_DEBUG_VIZ_SUBSTRATE_F90 			41
+#define PATH_TRACER_DEBUG_VIZ_EYE_IRIS_MASK             42
+#define PATH_TRACER_DEBUG_VIZ_EYE_CAUSTIC_NORMAL        43
+#define PATH_TRACER_DEBUG_VIZ_EYE_IRIS_NORMAL           44
+#define PATH_TRACER_DEBUG_VIZ_BACKFACING_NORMAL         45
+#define PATH_TRACER_DEBUG_VIZ_HAIR_PRIMITIVE_UV         46
+#define PATH_TRACER_DEBUG_VIZ_TRIANGLE_ID               47
+#define PATH_TRACER_DEBUG_VIZ_INSTANCE_ID               48
+#define PATH_TRACER_DEBUG_VIZ_CLUSTER_ID                49
diff --git a/Engine/Source/Runtime/Renderer/Private/PathTracing.cpp b/Engine/Source/Runtime/Renderer/Private/PathTracing.cpp
--- a/Engine/Source/Runtime/Renderer/Private/PathTracing.cpp
+++ b/Engine/Source/Runtime/Renderer/Private/PathTracing.cpp
@@ -165,8 +165,11 @@ struct FPathTracingDebugModeInfo
 		REGISTER_VIZ_TYPE(TEXT("Anisotropy") 	    	, PATH_TRACER_DEBUG_VIZ_ANISOTROPY 				);
 		REGISTER_VIZ_TYPE(TEXT("LightGridCount") 	   	, PATH_TRACER_DEBUG_VIZ_LIGHT_GRID_COUNT		);
 		REGISTER_VIZ_TYPE(TEXT("LightGridAxis") 	   	, PATH_TRACER_DEBUG_VIZ_LIGHT_GRID_AXIS			);
+        REGISTER_VIZ_TYPE(TEXT("LightGridListHash")	   	, PATH_TRACER_DEBUG_VIZ_LIGHT_GRID_LIST_HASH    );
+        REGISTER_VIZ_TYPE(TEXT("LightGridSingular")     , PATH_TRACER_DEBUG_VIZ_LIGHT_GRID_SINGULAR     );
 		REGISTER_VIZ_TYPE(TEXT("DecalGridCount")	   	, PATH_TRACER_DEBUG_VIZ_DECAL_GRID_COUNT		);
 		REGISTER_VIZ_TYPE(TEXT("DecalGridAxis") 	   	, PATH_TRACER_DEBUG_VIZ_DECAL_GRID_AXIS			);
+        REGISTER_VIZ_TYPE(TEXT("DecalGridListHash")     , PATH_TRACER_DEBUG_VIZ_DECAL_GRID_LIST_HASH    );
 		REGISTER_VIZ_TYPE(TEXT("VolumeLightCount") 		, PATH_TRACER_DEBUG_VIZ_VOLUME_LIGHT_COUNT		);
 		REGISTER_VIZ_TYPE(TEXT("HitKind")            	, PATH_TRACER_DEBUG_VIZ_HITKIND             	);
 		REGISTER_VIZ_TYPE(TEXT("TransparencyCount")     , PATH_TRACER_DEBUG_VIZ_TRANSPARENCY_COUNT		);
@@ -187,6 +190,9 @@ struct FPathTracingDebugModeInfo
 		REGISTER_VIZ_TYPE(TEXT("EyeCausticNormal")      , PATH_TRACER_DEBUG_VIZ_EYE_CAUSTIC_NORMAL      );
 		REGISTER_VIZ_TYPE(TEXT("BackfacingNormal")      , PATH_TRACER_DEBUG_VIZ_BACKFACING_NORMAL       );
 		REGISTER_VIZ_TYPE(TEXT("HairPrimitiveUV")		, PATH_TRACER_DEBUG_VIZ_HAIR_PRIMITIVE_UV		);
+        REGISTER_VIZ_TYPE(TEXT("TriangleID")            , PATH_TRACER_DEBUG_VIZ_TRIANGLE_ID             );
+        REGISTER_VIZ_TYPE(TEXT("InstanceID")            , PATH_TRACER_DEBUG_VIZ_INSTANCE_ID             );
+        REGISTER_VIZ_TYPE(TEXT("ClusterID")             , PATH_TRACER_DEBUG_VIZ_CLUSTER_ID              );
 #undef REGISTER_VIZ_TYPE
 	}
 
@@ -223,6 +229,21 @@ TAutoConsoleVariable<int32> CVarPathTracingVisualizeSamplesPerPixel(
 	ECVF_RenderThreadSafe
 );
 
+TAutoConsoleVariable<float> CVarPathTracingVisualizeAmbientOcclusionMaxDistance(
+	TEXT("r.PathTracing.Visualize.AmbientOcclusionMaxDistance"),
+	200.0f,
+	TEXT("How far to trace ambient occlusion rays for fake shading effect? This only applies to some of the visualization modes. Negative values revert to a simple facing ratio shading."),
+	ECVF_RenderThreadSafe
+);
+
+TAutoConsoleVariable<int32> CVarPathTracingVisualizeMaxCount(
+	TEXT("r.PathTracing.Visualize.MaxCount"),
+	16,
+	TEXT("Number of lights (or decals) to use for the visualization modes involving turning a count into a color. The color ramp is blue to red, with pink indicating a cell beyond the target count. "),
+	ECVF_RenderThreadSafe
+);
+
+
 TAutoConsoleVariable<bool> CVarPathTracingVisualizeIndirectGuiding(
 	TEXT("r.PathTracing.Visualize.IndirectGuiding"),
 	false,
@@ -1583,6 +1604,9 @@ class FPathTracingDebugRG : public FGlobalShader
 		SHADER_PARAMETER(float, SampleBlendFactor)
 		SHADER_PARAMETER(int32, MaxBounces)
 		SHADER_PARAMETER(int32, EnableEmissive)
+        SHADER_PARAMETER(float, MaxPathIntensity)
+        SHADER_PARAMETER(float, AmbientOcclusionMaxDistance)
+        SHADER_PARAMETER(uint32, VisualizeMaxCount)
 
 		// scene lights
 		SHADER_PARAMETER_RDG_BUFFER_SRV(StructuredBuffer<FPathTracingLight>, SceneLights)
@@ -4103,6 +4127,14 @@ void FDeferredShadingSceneRenderer::RenderPathTracing(
 					{
 						PassParameters->EnableEmissive = View.FinalPostProcessSettings.PathTracingEnableEmissiveMaterials;
 					}
+                    PassParameters->MaxPathIntensity = CVarPathTracingMaxPathIntensity.GetValueOnRenderThread();
+                    if (PassParameters->MaxPathIntensity <= 0)
+                    {
+                        // cvar clamp disabled, use PPV value instad
+                        PassParameters->MaxPathIntensity = View.FinalPostProcessSettings.PathTracingMaxPathIntensity;
+                    }
+                    PassParameters->AmbientOcclusionMaxDistance = CVarPathTracingVisualizeAmbientOcclusionMaxDistance.GetValueOnRenderThread();
+                    PassParameters->VisualizeMaxCount = uint32(FMath::Max(1, CVarPathTracingVisualizeMaxCount.GetValueOnRenderThread()));
 					{
 						SetLightParameters(
 							GraphBuilder,
