Hair Strands: Fix minor shader warnings

#rnx

#Synced-CL 51851216

diff --git a/Engine/Shaders/Private/HairStrands/HairStrandsVertexFactory.ush b/Engine/Shaders/Private/HairStrands/HairStrandsVertexFactory.ush
--- a/Engine/Shaders/Private/HairStrands/HairStrandsVertexFactory.ush
+++ b/Engine/Shaders/Private/HairStrands/HairStrandsVertexFactory.ush
@@ -644,8 +644,8 @@ FMaterialPixelParameters GetMaterialPixelParameters(float3 RayOrigin, float3 Ray
 			HairStrandsVF.RootScale,
 			HairStrandsVF.TipScale);
 
-		float3 A = DFDemote(CalcWorldPosition(P0.Position, SceneData.Primitive.LocalToWorld));
-		float3 B = DFDemote(CalcWorldPosition(P1.Position, SceneData.Primitive.LocalToWorld));
+		float3 A = DFDemote(CalcWorldPosition(P0.Position, SceneData.Primitive.LocalToWorld).xyz);
+		float3 B = DFDemote(CalcWorldPosition(P1.Position, SceneData.Primitive.LocalToWorld).xyz);
 
 		// approximate position along the width of the ribbon
 		
