Heterogeneous volumes: Fix minor shader warnings

Also upgrade usage of the older RandomSequence API

#rnx

#Synced-CL 51851216

diff --git a/Engine/Shaders/Private/HeterogeneousVolumes/HeterogeneousVolumesPreshadingPipeline.usf b/Engine/Shaders/Private/HeterogeneousVolumes/HeterogeneousVolumesPreshadingPipeline.usf
--- a/Engine/Shaders/Private/HeterogeneousVolumes/HeterogeneousVolumesPreshadingPipeline.usf
+++ b/Engine/Shaders/Private/HeterogeneousVolumes/HeterogeneousVolumesPreshadingPipeline.usf
@@ -60,8 +60,8 @@ void MaxTexture3D(
 	uint3 Index = DispatchThreadId;
 	if (all(Index < TextureResolution))
 	{
-		float3 ValueLHS = InputTextureLHS[Index];
-		float3 ValueRHS = InputTextureRHS[Index];
+		float3 ValueLHS = InputTextureLHS[Index].xyz;
+		float3 ValueRHS = InputTextureRHS[Index].xyz;
 #if 0
 		if (any(ValueRHS > 0.0))
 		{
diff --git a/Engine/Shaders/Private/HeterogeneousVolumes/HeterogeneousVolumesVoxelGridShadows.usf b/Engine/Shaders/Private/HeterogeneousVolumes/HeterogeneousVolumesVoxelGridShadows.usf
--- a/Engine/Shaders/Private/HeterogeneousVolumes/HeterogeneousVolumesVoxelGridShadows.usf
+++ b/Engine/Shaders/Private/HeterogeneousVolumes/HeterogeneousVolumesVoxelGridShadows.usf
@@ -531,9 +531,7 @@ void RenderVolumetricShadowMapForLightWithVoxelGridCS(
 			bApplyShadowTransmittance
 		);
 
-		RandomSequence RandSequence;
-		RandomSequence_Initialize(RandSequence, LinearIndex, View.StateFrameIndex);
-		RayMarchingContext.RandSequence = RandSequence;
+		RayMarchingContext.RandSequence = RandomSequenceCreate(LinearIndex, View.StateFrameIndex);
 		uint StepCount = CalcStepCount(RayMarchingContext);
 
 		float3 ExpectedExtinction = 0.0;
@@ -596,13 +594,13 @@ void RenderVolumetricShadowMapForLightWithVoxelGridCS(
 
 		}
 
-		AVSM_SetLinkedListSampleData(RWVolumetricShadowLinkedListBuffer[PixelIndex + ElementIndex], AVSM_CreateSampleData(WorldHitT, Transmittance));
+		AVSM_SetLinkedListSampleData(RWVolumetricShadowLinkedListBuffer[PixelIndex + ElementIndex], AVSM_CreateSampleData(WorldHitT, Transmittance.x));
 		AVSM_SetLinkedListPrevPtr(RWVolumetricShadowLinkedListBuffer[PixelIndex + ElementIndex], PrevElementIndex);
 		AVSM_SetLinkedListNextPtr(RWVolumetricShadowLinkedListBuffer[PixelIndex + ElementIndex], AVSM_NULL_PTR);
 	}
 	else
 	{
-		AVSM_SetLinkedListSampleData(RWVolumetricShadowLinkedListBuffer[PixelIndex], AVSM_CreateSampleData(0.0, Transmittance));
+		AVSM_SetLinkedListSampleData(RWVolumetricShadowLinkedListBuffer[PixelIndex], AVSM_CreateSampleData(0.0, Transmittance.x));
 		AVSM_SetLinkedListPrevPtr(RWVolumetricShadowLinkedListBuffer[PixelIndex], AVSM_NULL_PTR);
 		AVSM_SetLinkedListNextPtr(RWVolumetricShadowLinkedListBuffer[PixelIndex], AVSM_NULL_PTR);
 	}
