Path Tracing: Turn off FORCE_TRACE_INFINITE_LIGHTS which was left enabled accidentally

#rnx
#Synced-CL 51893035

diff --git a/Engine/Shaders/Private/PathTracing/PathTracingDebug.usf b/Engine/Shaders/Private/PathTracing/PathTracingDebug.usf
--- a/Engine/Shaders/Private/PathTracing/PathTracingDebug.usf
+++ b/Engine/Shaders/Private/PathTracing/PathTracingDebug.usf
@@ -15,7 +15,7 @@
 
 #if LIGHT_LOOP_METHOD == 2
 #define ACCUMULATE_UNSHADOWED_LIGHTS 1 // When true, in LIGHT_LOOP_METHOD == 2, accumulate any unshadowed light directly, skipping RIS to reduce noise
-#define FORCE_TRACE_INFINITE_LIGHTS  1 // When true, in LIGHT_LOOP_METHOD == 2, trace a ray for infinite lights instead of letting them do RIS
+#define FORCE_TRACE_INFINITE_LIGHTS  0 // When true, in LIGHT_LOOP_METHOD == 2, trace a ray for infinite lights instead of letting them do RIS
 #endif
 
 #define CLAMP_PATH_INTENSITY        1 // Apply MaxPathIntensity to limit intensity of indirect paths like the reference path tracer
