Fix incorrect scene texture sampling uv in postprocess materials after TSR. This also caused incorrect worldposition.
Previously, SvPosition was used to sample scenetextures. Given that the postprocess output viewport may not match the main view viewport, this could produce misaligned sampling positions.
This issue was particularly visible with a screenpercentage set, or with an offset viewrect (e.g. r.Test.ViewRectOffset 1)
This also affected related variables such as ScreenPosition.

#Synced-CL 51867312

diff --git a/Engine/Plugins/Experimental/RemoteSession/Source/RemoteSession/Private/Channels/RemoteSessionARCameraChannel.cpp b/Engine/Plugins/Experimental/RemoteSession/Source/RemoteSession/Private/Channels/RemoteSessionARCameraChannel.cpp
--- a/Engine/Plugins/Experimental/RemoteSession/Source/RemoteSession/Private/Channels/RemoteSessionARCameraChannel.cpp
+++ b/Engine/Plugins/Experimental/RemoteSession/Source/RemoteSession/Private/Channels/RemoteSessionARCameraChannel.cpp
@@ -86,10 +86,6 @@ public:
 		OutEnvironment.SetDefine(TEXT("POST_PROCESS_MATERIAL"), 1);
 		OutEnvironment.SetDefine(TEXT("POST_PROCESS_MATERIAL_BEFORE_TONEMAP"), (Parameters.MaterialParameters.BlendableLocation != BL_SceneColorAfterTonemapping) ? 1 : 0);
 
-		EBlendableLocation Location = EBlendableLocation(Parameters.MaterialParameters.BlendableLocation);
-		bool bPostProcessMaterialBeforeUpscaling = (Location == BL_SceneColorBeforeDOF) || (Location == BL_SceneColorAfterDOF) || (Location == BL_TranslucencyAfterDOF);
-		OutEnvironment.SetDefine(TEXT("POST_PROCESS_MATERIAL_BEFORE_UPSCALING"), bPostProcessMaterialBeforeUpscaling ? 1 : 0 );
-
 	}
 };
 
diff --git a/Engine/Plugins/Runtime/AR/AppleAR/AppleARKit/Source/AppleARKit/Private/AppleARKitVideoOverlay.cpp b/Engine/Plugins/Runtime/AR/AppleAR/AppleARKit/Source/AppleARKit/Private/AppleARKitVideoOverlay.cpp
--- a/Engine/Plugins/Runtime/AR/AppleAR/AppleARKit/Source/AppleARKit/Private/AppleARKitVideoOverlay.cpp
+++ b/Engine/Plugins/Runtime/AR/AppleAR/AppleARKit/Source/AppleARKit/Private/AppleARKitVideoOverlay.cpp
@@ -159,10 +159,6 @@ public:
 		FMaterialShader::ModifyCompilationEnvironment(Parameters, OutEnvironment);
 		OutEnvironment.SetDefine(TEXT("POST_PROCESS_MATERIAL"), 1);
 		OutEnvironment.SetDefine(TEXT("POST_PROCESS_MATERIAL_BEFORE_TONEMAP"), (Parameters.MaterialParameters.BlendableLocation != BL_SceneColorAfterTonemapping) ? 1 : 0);
-		
-		EBlendableLocation Location = EBlendableLocation(Parameters.MaterialParameters.BlendableLocation);
-		bool bPostProcessMaterialBeforeUpscaling = (Location == BL_SceneColorBeforeDOF) || (Location == BL_SceneColorAfterDOF) || (Location == BL_TranslucencyAfterDOF);
-		OutEnvironment.SetDefine(TEXT("POST_PROCESS_MATERIAL_BEFORE_UPSCALING"), bPostProcessMaterialBeforeUpscaling ? 1 : 0 );
 		OutEnvironment.SetDefine(TEXT("POST_PROCESS_AR_PASSTHROUGH"), 1);
 
 		const bool bMobileForceDepthRead = MobileUsesFullDepthPrepass(Parameters.Platform);
diff --git a/Engine/Plugins/Runtime/AR/Google/GoogleARCore/Source/GoogleARCoreRendering/Private/GoogleARCorePassthroughCameraRenderer.cpp b/Engine/Plugins/Runtime/AR/Google/GoogleARCore/Source/GoogleARCoreRendering/Private/GoogleARCorePassthroughCameraRenderer.cpp
--- a/Engine/Plugins/Runtime/AR/Google/GoogleARCore/Source/GoogleARCoreRendering/Private/GoogleARCorePassthroughCameraRenderer.cpp
+++ b/Engine/Plugins/Runtime/AR/Google/GoogleARCore/Source/GoogleARCoreRendering/Private/GoogleARCorePassthroughCameraRenderer.cpp
@@ -108,10 +108,6 @@ public:
 		FMaterialShader::ModifyCompilationEnvironment(Parameters, OutEnvironment);
 		OutEnvironment.SetDefine(TEXT("POST_PROCESS_MATERIAL"), 1);
 		OutEnvironment.SetDefine(TEXT("POST_PROCESS_MATERIAL_BEFORE_TONEMAP"), (Parameters.MaterialParameters.BlendableLocation != BL_SceneColorAfterTonemapping) ? 1 : 0);
-		
-		EBlendableLocation Location = EBlendableLocation(Parameters.MaterialParameters.BlendableLocation);
-		bool bPostProcessMaterialBeforeUpscaling = (Location == BL_SceneColorBeforeDOF) || (Location == BL_SceneColorAfterDOF) || (Location == BL_TranslucencyAfterDOF);
-		OutEnvironment.SetDefine(TEXT("POST_PROCESS_MATERIAL_BEFORE_UPSCALING"), bPostProcessMaterialBeforeUpscaling ? 1 : 0 );
 
 	}
 };
diff --git a/Engine/Plugins/TextureGraph/Shaders/TextureGraphMaterialShader.usf b/Engine/Plugins/TextureGraph/Shaders/TextureGraphMaterialShader.usf
--- a/Engine/Plugins/TextureGraph/Shaders/TextureGraphMaterialShader.usf
+++ b/Engine/Plugins/TextureGraph/Shaders/TextureGraphMaterialShader.usf
@@ -62,7 +62,7 @@ void TextureGraphMaterialShaderPS(
 
 	// fill out other related material parameters
 	FPixelMaterialInputs PixelMaterialInputs;
-	CalcMaterialParametersPost(MaterialParameters, PixelMaterialInputs, SvPosition, true);
+	CalcMaterialParametersPost(MaterialParameters, PixelMaterialInputs, SvPosition, ViewportUV, true);
 	
 	// Extract the value we want to represent
 	half3 MatBaseColor = GetMaterialBaseColor(PixelMaterialInputs);
diff --git a/Engine/Shaders/Private/MaterialTemplate.ush b/Engine/Shaders/Private/MaterialTemplate.ush
--- a/Engine/Shaders/Private/MaterialTemplate.ush
+++ b/Engine/Shaders/Private/MaterialTemplate.ush
@@ -4740,22 +4740,25 @@ void CalcMaterialParameters(
 	CalcMaterialParametersEx(Parameters, PixelMaterialInputs, SvPosition, ScreenPosition, bIsFrontFace, TranslatedWorldPosition, TranslatedWorldPosition);
 }
 
+// Postprocess output viewport may not match main view viewport.
+// This function translates the 'true' SvPosition in the post pass to the equivalent SvPosition in the main view (pre upscaling, viewrect offset, ...).
+float4 PostSvPositionToViewSvPosition(float2 ViewportUV, float4 PostPassSvPosition)
+{
+	float2 PixelPos = (ViewportUV * View.ViewSizeAndInvSize.xy) + View.ViewRectMin.xy;
+	return float4(floor(PixelPos)+0.5, PostPassSvPosition.zw);
+}
+
 void CalcMaterialParametersPost(
 	in out FMaterialPixelParameters Parameters,
 	in out FPixelMaterialInputs PixelMaterialInputs,
-	float4 SvPosition,
+	float4 PostPassSvPosition,
+	float2 ViewportUV,
 	FIsFrontFace bIsFrontFace)
 {
-	float4 ScreenPosition = SvPositionToScreenPosition(SvPosition);
-#if POST_PROCESS_MATERIAL_BEFORE_UPSCALING
-	// In this case, SVPositionToTranslatedWorld account for the correct dynamic resolution scaling.
-	float3 TranslatedWorldPosition = SvPositionToTranslatedWorld(SvPosition);
-#else
-	// Otherwise, we must update SVPosition to account for the resolution scaling not happening after the upscaling stage.
-	float3 TranslatedWorldPosition = SvPositionToTranslatedWorld(SvPosition * float4(View.ResolutionFractionAndInv.xx, 1.0f, 1.0f));
-#endif
-
-	CalcMaterialParametersEx(Parameters, PixelMaterialInputs, SvPosition, ScreenPosition, bIsFrontFace, TranslatedWorldPosition, TranslatedWorldPosition);
+	float4 ViewSvPosition = PostSvPositionToViewSvPosition(ViewportUV, PostPassSvPosition);
+	float4 ScreenPosition = SvPositionToScreenPosition(ViewSvPosition);
+	float3 TranslatedWorldPosition = SvPositionToTranslatedWorld(ViewSvPosition);
+	CalcMaterialParametersEx(Parameters, PixelMaterialInputs, PostPassSvPosition, ScreenPosition, bIsFrontFace, TranslatedWorldPosition, TranslatedWorldPosition);
 }
 
 /** Assemble the transform from tangent space into world space */
diff --git a/Engine/Shaders/Private/PostProcessMaterialShaders.usf b/Engine/Shaders/Private/PostProcessMaterialShaders.usf
--- a/Engine/Shaders/Private/PostProcessMaterialShaders.usf
+++ b/Engine/Shaders/Private/PostProcessMaterialShaders.usf
@@ -268,7 +268,7 @@ void MainPS(
 	SvPosition.z = max(SvPosition.z, 1e-18);
 
 	// fill out other related material parameters
-	CalcMaterialParametersPost(Parameters, PixelMaterialInputs, SvPosition, true);
+	CalcMaterialParametersPost(Parameters, PixelMaterialInputs, SvPosition, ViewportUV, true);
 
 #if NEURAL_POSTPROCESS_PREPASS && MATERIAL_NEURAL_POST_PROCESS
 
@@ -414,7 +414,7 @@ void MainPS_VideoOverlay(
 	SvPosition.z = max(SvPosition.z, 1e-18);
 
 	// fill out other related material parameters
-	CalcMaterialParametersPost(Parameters, PixelMaterialInputs, SvPosition, true);
+	CalcMaterialParametersPost(Parameters, PixelMaterialInputs, SvPosition, ViewportUV, true);
 
 	// Grab emissive colour as output
 #if MATERIAL_OUTPUT_OPACITY_AS_ALPHA
@@ -519,7 +519,7 @@ void MainPS(
 	SvPosition.z = DeviceZ;
 	SvPosition.z = max(SvPosition.z, 1e-18);
 	// fill out other related material parameters
-	CalcMaterialParametersPost(Parameters, PixelMaterialInputs, SvPosition, true);
+	CalcMaterialParametersPost(Parameters, PixelMaterialInputs, SvPosition, ViewportUV, true);
 
 	float Alpha = 0.0f;
 	float3 EmissiveColor = 0.0f;
diff --git a/Engine/Source/Runtime/Renderer/Private/PostProcess/PostProcessMaterial.cpp b/Engine/Source/Runtime/Renderer/Private/PostProcess/PostProcessMaterial.cpp
--- a/Engine/Source/Runtime/Renderer/Private/PostProcess/PostProcessMaterial.cpp
+++ b/Engine/Source/Runtime/Renderer/Private/PostProcess/PostProcessMaterial.cpp
@@ -263,9 +263,6 @@ public:
 		OutEnvironment.SetDefine(TEXT("POST_PROCESS_MATERIAL_BEFORE_TONEMAP"), (Location == BL_SceneColorAfterTonemapping || Location == BL_ReplacingTonemapper) ? 0 : 1);
 		OutEnvironment.SetDefine(TEXT("POST_PROCESS_DISABLE_PRE_EXPOSURE_SCALE"), Parameters.MaterialParameters.bDisablePreExposureScale ? 1 : 0);
 
-		bool bPostProcessMaterialBeforeUpscaling = (Location == BL_SceneColorBeforeDOF) || (Location == BL_SceneColorAfterDOF) || (Location == BL_TranslucencyAfterDOF);
-		OutEnvironment.SetDefine(TEXT("POST_PROCESS_MATERIAL_BEFORE_UPSCALING"), bPostProcessMaterialBeforeUpscaling ? 1 : 0 );
-
 		if (IsMobilePlatform(Parameters.Platform))
 		{
 			OutEnvironment.SetDefine(TEXT("POST_PROCESS_MATERIAL_BEFORE_TONEMAP"), (Parameters.MaterialParameters.BlendableLocation != BL_SceneColorAfterTonemapping) ? 1 : 0);
