Fix incorrect lighting on mobile due to unnormalized GBuffer normals

The normals decoded from the mobile GBuffer using DecodeNormal() are not perfectly normalized (magnitude ~[REDACTED_PHONE]) due to compression/quantization artifacts when MOBILE_EXTENDED_GBUFFER is enabled.

This causes incorrect lighting calculations as dot products assume unit vectors. This ensures normals are properly renormalized after decode,
fixing subtle lighting intensity errors (~0.2-2%) and ensuring consistency with the desktop renderer which uses octahedral
encoding (OctahedronToUnitVector) that already produces normalized results.

#tests Console with odsc

diff --git a/Engine/Shaders/Private/DeferredShadingCommon.ush b/Engine/Shaders/Private/DeferredShadingCommon.ush
--- a/Engine/Shaders/Private/DeferredShadingCommon.ush
+++ b/Engine/Shaders/Private/DeferredShadingCommon.ush
@@ -709,7 +709,7 @@ FGBufferData MobileDecodeGBuffer(half4 InGBufferA, half4 InGBufferB, half4 InGBu
 	FGBufferData GBuffer = (FGBufferData)0;
 
 #if	MOBILE_EXTENDED_GBUFFER && !ALLOW_STATIC_LIGHTING		
-	GBuffer.WorldNormal = DecodeNormal(InGBufferA.rgb);
+	GBuffer.WorldNormal = normalize(DecodeNormal(InGBufferA.rgb));
 #else
 	GBuffer.WorldNormal = OctahedronToUnitVector(InGBufferA.rg * 2.0 - 1.0);
 	#if ALLOW_STATIC_LIGHTING
