Add F90 approximation for Substrate Blendable GBuffer format.

This approximation only works for conductor.

#Synced-CL 51943914

diff --git a/Engine/Shaders/Private/Substrate/Substrate.ush b/Engine/Shaders/Private/Substrate/Substrate.ush
--- a/Engine/Shaders/Private/Substrate/Substrate.ush
+++ b/Engine/Shaders/Private/Substrate/Substrate.ush
@@ -311,6 +311,8 @@ struct FSubstrateBSDF
 	// Return true if the BSDF has SSS diffusion enabled
 	bool HasSSSDiffusion();
 
+	void SubstrateBakeF90ToReflectivity(in float NoV);
+
 	void SubstrateBakeSpecularProfileToReflectivity(in float NoV);
 
 	void SubstrateBakeGlintToReflectivity();
@@ -2887,6 +2889,11 @@ FSubstrateData SubstrateHorizontalMixingParameterBlending(FSubstrateData Backgro
 	BSDF_SETHASFUZZ(ResultBSDF,			BSDF_GETHASFUZZ(ForegroundBSDF)					| BSDF_GETHASFUZZ(BackgroundBSDF));
 	BSDF_SETHASTRANSABOVE(ResultBSDF,	BSDF_GETHASTRANSABOVE(ForegroundBSDF)			| BSDF_GETHASTRANSABOVE(BackgroundBSDF));
 
+	// Bake the F90 into F0 of BSDFs to blend in order to retain some of the original look when using parameter blending or fully simplified materials.
+	BackgroundBSDF.SubstrateBakeF90ToReflectivity(BackgroundNoV);
+	ForegroundBSDF.SubstrateBakeF90ToReflectivity(ForegroundNoV);
+	BSDF_SETHASF90(ResultBSDF, 0);
+
 	// Bake the Specular Profiles into F0/F90 of BSDFs to blend in order to retain some of the original look when using parameter blending or fully simplified materials.
 #if SUBSTRATE_SPECPROFILE_ENABLED
 	BackgroundBSDF.SubstrateBakeSpecularProfileToReflectivity(BackgroundNoV);
@@ -3018,6 +3025,11 @@ FSubstrateData SubstrateAddParameterBlending(FSubstrateData A, FSubstrateData B,
 	BSDF_SETHASF90(ResultBSDF,			BSDF_GETHASF90(BBSDF)							| BSDF_GETHASF90(ABSDF));
 	BSDF_SETHASFUZZ(ResultBSDF,			BSDF_GETHASFUZZ(BBSDF)							| BSDF_GETHASFUZZ(ABSDF));
 
+	// Bake the F90 into F0 of BSDFs to blend in order to retain some of the original look when using parameter blending or fully simplified materials.
+	ABSDF.SubstrateBakeF90ToReflectivity(ANoV);
+	BBSDF.SubstrateBakeF90ToReflectivity(BNoV);
+	BSDF_SETHASF90(ResultBSDF, 0);
+
 	// Bake the Specular Profiles into F0/F90 of BSDFs to blend in order to retain some of the original look when using parameter blending or fully simplified materials.
 #if SUBSTRATE_SPECPROFILE_ENABLED
 	ABSDF.SubstrateBakeSpecularProfileToReflectivity(ANoV);
@@ -3160,6 +3172,11 @@ FSubstrateData SubstrateVerticalLayeringParameterBlending(in FSubstrateData Top,
 	BSDF_SETHASF90(ResultBSDF,			BSDF_GETHASF90(BaseBSDF)						| BSDF_GETHASF90(TopBSDF));			// F90: keep union of both even though it will be hard to get a match
 	BSDF_SETHASFUZZ(ResultBSDF,			BSDF_GETHASFUZZ(BaseBSDF)						| BSDF_GETHASFUZZ(TopBSDF));		// Fuzz: keep union of both even though it will be hard to get a match
 
+	// Bake the F90 into F0 of BSDFs to blend in order to retain some of the original look when using parameter blending or fully simplified materials.
+	TopBSDF.SubstrateBakeF90ToReflectivity(TopNoV);
+	BaseBSDF.SubstrateBakeF90ToReflectivity(BaseNoV);
+	BSDF_SETHASF90(ResultBSDF, 0);
+
 	// Bake the Specular Profiles into F0/F90 of BSDFs to blend in order to retain some of the original look when using parameter blending or fully simplified materials.
 #if SUBSTRATE_SPECPROFILE_ENABLED
 	TopBSDF.SubstrateBakeSpecularProfileToReflectivity(TopNoV);
@@ -3640,6 +3657,28 @@ bool FSubstrateBSDF::HasSSSDiffusion()
 	return SSSType == SSS_TYPE_DIFFUSION || SSSType == SSS_TYPE_DIFFUSION_PROFILE;
 }
 
+void FSubstrateBSDF::SubstrateBakeF90ToReflectivity(in float NoV)
+{
+	if (BSDF_GETHASF90(this))
+	{
+		float3 F90 = saturate(SLAB_F90(this));
+		// What is important is to maintain the hue and saturation, so we scale the color by the maximum of its components
+		const float Divisor = max3(F90.r, F90.g, F90.b);
+		F90 = Divisor > 0.0f ? F90 * rcp(Divisor) : 1.0f;
+
+		// Shift when F90 color is fully reached to better match reference
+		float3 F0Scale = lerp(1.f, F90, Pow5((1-NoV)*1.2f));
+
+		// Fade out F90 effect as roughness goes from 0 to 0.5. 
+		// On Adaptive GBuffer, the effect is no longer visible after roughness 0.5.
+		F0Scale = lerp(1.f, F0Scale, saturate(1.f-2.f*SLAB_ROUGHNESS(this)));		
+
+		SLAB_F0(this) *= F0Scale;
+		SLAB_F90(this) = 1.f;
+		BSDF_SETHASSPECPROFILE(this, 0);
+	}
+}
+
 void FSubstrateBSDF::SubstrateBakeSpecularProfileToReflectivity(in float NoV)
 {
 #if SUBSTRATE_SPECPROFILE_ENABLED
diff --git a/Engine/Shaders/Private/Substrate/SubstrateExport.ush b/Engine/Shaders/Private/Substrate/SubstrateExport.ush
--- a/Engine/Shaders/Private/Substrate/SubstrateExport.ush
+++ b/Engine/Shaders/Private/Substrate/SubstrateExport.ush
@@ -1088,10 +1088,14 @@ FExportResult SubstrateMaterialExportOut(
 			const float3 T = TangentBasis[0];
 
 			const float GreyScaleWeight = dot(CurrentBSDF.LuminanceWeightV, (1.0 / 3.0).xxx);
+			const float NoV = saturate(dot(SurfaceWorldNormal, V));
+
+			// Bake the F90 into F0 of BSDFs to blend in order to retain some of the original look when using parameter blending or fully simplified materials.
+			CurrentBSDF.SubstrateBakeF90ToReflectivity(NoV);
+			BSDF_SETHASF90(CurrentBSDF, 0);
 			
 			// Bake the Specular Profiles into F0/F90 of BSDFs to blend in order to retain some of the original look when using parameter blending or fully simplified materials.
 		#if SUBSTRATE_SPECPROFILE_ENABLED && MATERIAL_USES_SPECULAR_PROFILE
-			const float NoV = saturate(dot(SurfaceWorldNormal, V));
 			CurrentBSDF.SubstrateBakeSpecularProfileToReflectivity(NoV);
 			BSDF_SETHASSPECPROFILE(CurrentBSDF, 0);
 		#endif
