Substrate - Fixed unlit mode for blendable gbuffer when material are converted to GBuffer.

#tests PC editor dx12 repro on jira
#rnx

#Synced-CL 52001491

diff --git a/Engine/Shaders/Private/BasePassPixelShader.usf b/Engine/Shaders/Private/BasePassPixelShader.usf
--- a/Engine/Shaders/Private/BasePassPixelShader.usf
+++ b/Engine/Shaders/Private/BasePassPixelShader.usf
@@ -1944,6 +1944,10 @@ void FPixelShaderInOut_MainPS(
 		GBuffer.Depth = MaterialParameters.ScreenPosition.w;
 		// GBuffer.SelectiveOutputMask is setup above.
 		EmissiveLuminance = Export.EmissiveLuminance;
+
+	#if USE_DEVELOPMENT_SHADERS
+		TopLayerData.UnlitViewBaseColor = Export.UnlitViewBaseColor;
+	#endif
 		
 		// Patch material parameter to not have to change SetGBufferForShadingModel for now
 		MaterialParameters.WorldNormal = Export.WorldNormal;
diff --git a/Engine/Shaders/Private/Substrate/SubstrateExport.ush b/Engine/Shaders/Private/Substrate/SubstrateExport.ush
--- a/Engine/Shaders/Private/Substrate/SubstrateExport.ush
+++ b/Engine/Shaders/Private/Substrate/SubstrateExport.ush
@@ -1038,6 +1038,9 @@ struct FExportResult
 	float IndirectIrradianceScale;
 	float Curvature;
 	float Opacity;
+#if USE_DEVELOPMENT_SHADERS
+	float3 UnlitViewBaseColor;
+#endif
 };
 
 // This is used to export very simple data for lightmaps, material baking or preview on editor node.
@@ -1113,7 +1116,10 @@ FExportResult SubstrateMaterialExportOut(
 			Out.SubsurfaceProfileIDFloat= ConvertSubsurfaceProfileToFloat(SubstrateGetBSDFSubSurfaceProfileID(CurrentBSDF));
 			Out.WorldTangent			= T;
 			Out.ShadingModelID			= SubstrateGetLegacyShadingModels(CurrentBSDF);
-
+			
+		#if USE_DEVELOPMENT_SHADERS
+			Out.UnlitViewBaseColor		= CurrentBSDF.TopLayerDataWeight * SubstrateGetBSDFBaseColor(CurrentBSDF);
+		#endif
 
 #if MATERIAL_SHADINGMODEL_HAIR || !MATERIAL_IS_SUBSTRATEHIDDENCONVERSION // Hair shading model or if not a hidden conversion, then we cannot fully trust the shading model from simplified substrate nodes.
 			if (BSDFType == SUBSTRATE_BSDF_TYPE_HAIR)
