[VSM] Use strided addressing scheme when inigializing physical pages, instead of 2x2 pattern per lane

Strided access results in a more optimal memory access pattern, where each lane in a warp/wave loads/stores consecutive pixels.
Original code would result in every other pixel skipped per texture load/store.

#Synced-CL 52000591

diff --git a/Engine/Shaders/Private/VirtualShadowMaps/VirtualShadowMapPhysicalPageManagement.usf b/Engine/Shaders/Private/VirtualShadowMaps/VirtualShadowMapPhysicalPageManagement.usf
--- a/Engine/Shaders/Private/VirtualShadowMaps/VirtualShadowMapPhysicalPageManagement.usf
+++ b/Engine/Shaders/Private/VirtualShadowMaps/VirtualShadowMapPhysicalPageManagement.usf
@@ -947,25 +947,26 @@ StructuredBuffer<uint> PhysicalPagesToInitialize;
 void InitializePhysicalPagesIndirectCS(uint2 TileThreadID : SV_GroupThreadID, uint GroupIndex : SV_GroupID)
 {
 	FPhysicalPageMetaData MetaData;
-	uint3 BasePos = GetTileBasePos(TileThreadID, GroupIndex, PhysicalPagesToInitialize, MetaData);
-	bool bStaticCached = (MetaData.Flags & VSM_FLAG_STATIC_UNCACHED) == 0U;
+	const uint3 BasePos = GetTileOffset(GroupIndex, PhysicalPagesToInitialize, MetaData) + uint3(TileThreadID, 0);
+	const bool bStaticCached = (MetaData.Flags & VSM_FLAG_STATIC_UNCACHED) == 0U;
+	const uint Stride = TILE_THREAD_GROUP_SIZE_XY;
 
 	if (bStaticCached && (MetaData.Flags & VSM_EXTENDED_FLAG_VIEW_UNCACHED) == 0U)
 	{
 		// Initialize from the static page data
 		checkSlow(BasePos.z == 0U);
 		OutPhysicalPagePool[BasePos + uint3(0U, 0U, 0U)] = OutPhysicalPagePool[BasePos + uint3(0U, 0U, 1U)];
-		OutPhysicalPagePool[BasePos + uint3(1U, 0U, 0U)] = OutPhysicalPagePool[BasePos + uint3(1U, 0U, 1U)];
-		OutPhysicalPagePool[BasePos + uint3(0U, 1U, 0U)] = OutPhysicalPagePool[BasePos + uint3(0U, 1U, 1U)];
-		OutPhysicalPagePool[BasePos + uint3(1U, 1U, 0U)] = OutPhysicalPagePool[BasePos + uint3(1U, 1U, 1U)];
+		OutPhysicalPagePool[BasePos + uint3(Stride, 0U, 0U)] = OutPhysicalPagePool[BasePos + uint3(Stride, 0U, 1U)];
+		OutPhysicalPagePool[BasePos + uint3(0U, Stride, 0U)] = OutPhysicalPagePool[BasePos + uint3(0U, Stride, 1U)];
+		OutPhysicalPagePool[BasePos + uint3(Stride, Stride, 0U)] = OutPhysicalPagePool[BasePos + uint3(Stride, Stride, 1U)];
 	}
 	else
 	{
 		// Clear the page to zero
 		OutPhysicalPagePool[BasePos + uint3(0U, 0U, 0U)] = 0U;
-		OutPhysicalPagePool[BasePos + uint3(1U, 0U, 0U)] = 0U;
-		OutPhysicalPagePool[BasePos + uint3(0U, 1U, 0U)] = 0U;
-		OutPhysicalPagePool[BasePos + uint3(1U, 1U, 0U)] = 0U;
+		OutPhysicalPagePool[BasePos + uint3(Stride, 0U, 0U)] = 0U;
+		OutPhysicalPagePool[BasePos + uint3(0U, Stride, 0U)] = 0U;
+		OutPhysicalPagePool[BasePos + uint3(Stride, Stride, 0U)] = 0U;
 	}
 }
 
