Substrate - Toon - Added Anisotropy support.
Updated the Max energy lut to account for anisotropy (3rd dimenson on RGBA component for now to avoid 3D lut).
Added support for RGBA LUT to ShadingEnergyConservation.cpp.
Toon patternUVs now automatically applies frac().

#tests PC with toon shading and substrate blendable/adaptive.
#rnx

#virtualized
#Synced-CL 51842837

#Synced-CL 51842837

diff --git a/Engine/Shaders/Private/AnisotropyPassShader.usf b/Engine/Shaders/Private/AnisotropyPassShader.usf
--- a/Engine/Shaders/Private/AnisotropyPassShader.usf
+++ b/Engine/Shaders/Private/AnisotropyPassShader.usf
@@ -145,7 +145,7 @@ void MainPixelShader(
 		0.0f // Curvature
 		);
 
-	float Anisotropy = Export.Anisotropy;
+	float Anisotropy = clamp(Export.Anisotropy, -1.0f, 1.0f);
 
 	// Similar to CalculateAnisotropyTangent. SUBSTRATE_TODO: merge the code and account for bottom normal.
 	// Exported tangent and normal are already normalized (from storage in SharedLocalNormal).
diff --git a/Engine/Shaders/Private/BasePassPixelShader.usf b/Engine/Shaders/Private/BasePassPixelShader.usf
--- a/Engine/Shaders/Private/BasePassPixelShader.usf
+++ b/Engine/Shaders/Private/BasePassPixelShader.usf
@@ -1216,8 +1216,8 @@ void FPixelShaderInOut_MainPS(
 	GBuffer.Velocity = 0;
 #endif
 
-#if USES_GBUFFER && SUBSTRATE_GBUFFER_FORMAT==0
-	if (!LegacyShadingModelSupportsAnisotropy(GBuffer.ShadingModelID))
+#if !SUBSTRATE_ENABLED && USES_GBUFFER && SUBSTRATE_GBUFFER_FORMAT==0
+	if (!LegacyShadingModelSupportsAnisotropy(GBuffer.ShadingModelID)) // GBuffer.ShadingModelID is only valid when !SUBSTRATE_ENABLED as this stage.
 	{
 		RemoveAnisotropyFromSelectiveOutputMask(GBuffer.SelectiveOutputMask);
 	}
diff --git a/Engine/Shaders/Private/DeferredShadingCommon.ush b/Engine/Shaders/Private/DeferredShadingCommon.ush
--- a/Engine/Shaders/Private/DeferredShadingCommon.ush
+++ b/Engine/Shaders/Private/DeferredShadingCommon.ush
@@ -371,7 +371,8 @@ bool LegacyShadingModelSupportsAnisotropy(uint ShadingModelID)
 {
 	// See ShadingModels.ush for where this comes from.
 	return ShadingModelID == SHADINGMODELID_DEFAULT_LIT
-		|| ShadingModelID == SHADINGMODELID_CLEAR_COAT;
+		|| ShadingModelID == SHADINGMODELID_CLEAR_COAT
+		|| ShadingModelID == SHADINGMODELID_SUBSTRATE_TOON;
 }
 
 // all values that are output by the forward rendering pass
diff --git a/Engine/Shaders/Private/ShadingEnergyConservation.ush b/Engine/Shaders/Private/ShadingEnergyConservation.ush
--- a/Engine/Shaders/Private/ShadingEnergyConservation.ush
+++ b/Engine/Shaders/Private/ShadingEnergyConservation.ush
@@ -96,10 +96,24 @@ float GGXEnergyLookup(float Roughness, float NoV, float Eta)
 #endif
 }
 
-float GGXMaxEnergyLookup(float Roughness, float NoV)
+float GGXMaxEnergyLookup(float Roughness, float Anisotropy, float NoV)
 {
 	// Always evaluate since it is only used for toon shading as of today without any more options.
-	return GGX_MAX_LUT_LOAD_SCALE * View.ShadingEnergyGGXMaxSpecTexture.SampleLevel(SharedShadingEnergySampler, float2(NoV, Roughness * GGX_MAX_LUT_LOAD_UV_SCALE + GGX_MAX_LUT_LOAD_UV_BIAS), 0).r;
+	const float4 ValuesAlongNoV = GGX_MAX_LUT_LOAD_SCALE * View.ShadingEnergyGGXMaxSpecTexture.SampleLevel(SharedShadingEnergySampler, float2(Roughness, abs(Anisotropy)) * GGX_MAX_LUT_LOAD_UV_SCALE + GGX_MAX_LUT_LOAD_UV_BIAS, 0).rgba;
+
+	const float Step = 0.3333f;
+#if 1
+	const float4 StepScale = float4(0.0f, 0.3333f, 0.6666f, 1.0f);
+	float4 LerpFactors = 1.0f - saturate(abs(NoV - StepScale) / Step);
+	return dot(ValuesAlongNoV, LerpFactors);
+#else
+	if(NoV <= Step)
+		return lerp(ValuesAlongNoV.x, ValuesAlongNoV.y, saturate((NoV - 0.0f * Step) / Step));
+	else if(NoV <= 2.0f * Step)
+		return lerp(ValuesAlongNoV.y, ValuesAlongNoV.z, saturate((NoV - 1.0f * Step) / Step));
+	//else  if(NoV < 3.0f * Step)
+		return lerp(ValuesAlongNoV.z, ValuesAlongNoV.w, saturate((NoV - 2.0f * Step) / Step));
+#endif
 }
 
 float2 ClothEnergyLookup(float Roughness, float NoV)
diff --git a/Engine/Shaders/Private/ShadingEnergyConservationTable.usf b/Engine/Shaders/Private/ShadingEnergyConservationTable.usf
--- a/Engine/Shaders/Private/ShadingEnergyConservationTable.usf
+++ b/Engine/Shaders/Private/ShadingEnergyConservationTable.usf
@@ -24,23 +24,28 @@ RWTexture2D<float>  Output1Texture2D;
 // y: directional albedo of GGX+Schlick fresnel term
 RWTexture2D<float2> Output2Texture2D;
 // x: maximum of the lobe for normalisation of the lob w.r.t. its maximum value.
-RWTexture2D<float> Output3Texture2D;
+RWTexture2D<float4> Output3Texture2D;
 // x: directional albedo from outside->in, y: directional albedo from inside->out
 RWTexture3D<float2> OutputTexture3D;
 
+#if BUILD_ENERGY_TABLE == 4
+groupshared float MaxLobeValue[THREADGROUPSIZE_X][THREADGROUPSIZE_Y][4];
+[numthreads(THREADGROUPSIZE_X, THREADGROUPSIZE_Y, 4)]
+#else
 [numthreads(THREADGROUPSIZE_X, THREADGROUPSIZE_Y, 1)]
-void BuildEnergyTableCS(uint3 DispatchThreadId : SV_DispatchThreadID)
+#endif
+void BuildEnergyTableCS(uint3 DispatchThreadId : SV_DispatchThreadID, uint3 GroupThreadId : SV_GroupThreadID)
 {
 	const float3 Index = (float3(DispatchThreadId) + 0.5) / float(EnergyTableResolution);
 
-	// Map from thread ID to the exact roughness/viewing angle we want to compute
-	const float NoV = max(Index.x, 1e-8);
-
 #if BUILD_ENERGY_TABLE == 4
 	// In this case we want exact uvs [0.5/height, (height-0.5)/height] to map to roughness [0, 1] to catch the extrems.
 	// This approach could potentially be applied to all the LUTs.
-	const float Roughness = float(DispatchThreadId.y) / float(EnergyTableResolution - 1);
-#else
+	const float Roughness  =     float(DispatchThreadId.x) / float(EnergyTableResolution - 1);
+	const float Anisotropy =     float(DispatchThreadId.y) / float(EnergyTableResolution - 1);		// in [0,1] as abs(aniso)
+	const float NoV        = max(float(DispatchThreadId.z) / float(4                     - 1), 1e-8f);
+#else	// Map from thread ID to the exact roughness/viewing angle we want to compute
+	const float NoV = max(Index.x, 1e-8);
 	const float Roughness =  Index.y;
 #endif
 
@@ -76,6 +81,7 @@ void BuildEnergyTableCS(uint3 DispatchThreadId : SV_DispatchThreadID)
 			Weight = float2(GGXResult.x, GGXResult.x * Pow5(1 - dot(V, H)));
 		}
 #elif BUILD_ENERGY_TABLE == 4
+#if 0
 		const float3 H = ImportanceSampleVisibleGGX(RandSample.xy, Alpha, V).xyz;
 
 		const float3 L = reflect(-V, H);
@@ -84,8 +90,18 @@ void BuildEnergyTableCS(uint3 DispatchThreadId : SV_DispatchThreadID)
 			const float a2 = Alpha.x * Alpha.x;	// no anisotropy
 			const float NoH = H.z;
 			const float NoL = L.z;
-			Weight  = D_GGX(a2, NoH);
+			Weight = D_GGX(a2, NoH); // Isotropic lobe
 		}
+#else
+		float2 AlphaXY = 0;
+		GetAnisotropicRoughness(Alpha, Anisotropy, AlphaXY.x, AlphaXY.y);
+		const float3 H = ImportanceSampleVisibleGGX(RandSample.xy, AlphaXY, V, false).xyz;
+		const float3 L = reflect(-V, H);
+		if (L.z > 0)
+		{
+			Weight = max(0, D_GGXaniso(AlphaXY.x, AlphaXY.y, H.z, H.x, H.y));// Anisotropic lobe
+		}
+#endif
 #elif BUILD_ENERGY_TABLE == 1
 		const float3 H = ImportanceSampleVisibleGGX(RandSample.xy, Alpha, V, false).xyz;
 		{
@@ -144,7 +160,7 @@ void BuildEnergyTableCS(uint3 DispatchThreadId : SV_DispatchThreadID)
 
 #if BUILD_ENERGY_TABLE == 4
 		// Max of all samples
-		Result = max(Result, Weight);
+		Result = max(Result, Weight.x);
 #else
 		// average in
 		Result = lerp(Result, Weight, 1.0 / (Sample + 1));
@@ -164,6 +180,21 @@ void BuildEnergyTableCS(uint3 DispatchThreadId : SV_DispatchThreadID)
 #endif
 
 #if BUILD_ENERGY_TABLE == 4
-	Output3Texture2D[DispatchThreadId.xy] = saturate(GGX_MAX_LUT_STORE_SCALE * (1.0f / Result.x));
+
+	MaxLobeValue[GroupThreadId.x][GroupThreadId.y][GroupThreadId.z] = Result.x;
+	GroupMemoryBarrierWithGroupSync();
+
+	BRANCH
+	if (DispatchThreadId.z == 0)
+	{
+		// Take the max of the 4 NoV samples for this roughness/anisotropy and store it in the LUT. This will be used for energy compensation in toon shading to make sure the specular lobe can reach 1.0 at its peak.
+		const float4 MaxLobeValues = float4(
+			MaxLobeValue[GroupThreadId.x][GroupThreadId.y][0],
+			MaxLobeValue[GroupThreadId.x][GroupThreadId.y][1],
+			MaxLobeValue[GroupThreadId.x][GroupThreadId.y][2],
+			MaxLobeValue[GroupThreadId.x][GroupThreadId.y][3]);
+
+		Output3Texture2D[DispatchThreadId.xy] = saturate(GGX_MAX_LUT_STORE_SCALE * (1.0f / MaxLobeValues));
+	}
 #endif
 }
diff --git a/Engine/Shaders/Private/Substrate/Substrate.ush b/Engine/Shaders/Private/Substrate/Substrate.ush
--- a/Engine/Shaders/Private/Substrate/Substrate.ush
+++ b/Engine/Shaders/Private/Substrate/Substrate.ush
@@ -722,6 +722,7 @@ struct FSubstrateBSDF
 #define TOON_SPECULAR(X)				X.VGPRs[1].y
 #define TOON_ROUGHNESS(X)				X.VGPRs[1].z
 #define TOON_PATTERNUVS(X)				X.VGPRs[2].xy
+#define TOON_ANISOTROPY(X)				X.VGPRs[2].z
 
 ///////////////////////////////////////////////////////////////////////////////
 // Forward declarations
@@ -1985,11 +1986,12 @@ FSubstrateData GetSubstrateToonBSDF(
 	float3 BaseColor, float Metallic, float Specular, float Roughness,
 	float3 Emissive,
 	float2 PatternUVs,
+	float Anisotropy, float3 Tangent,
 	uint SharedLocalBasisIndex, inout uint SharedLocalBasisTypes)
 {
 	FSubstrateData SubstrateData = GetInitialisedSubstrateData();
 
-#if 0 // MATERIAL_USES_ANISOTROPY && SUBSTRATE_FASTPATH==0
+#if MATERIAL_USES_ANISOTROPY && SUBSTRATE_FASTPATH==0
 	const uint bHasAnisotropy = Anisotropy != 0 ? 1 : 0;
 	if (bHasAnisotropy)
 	{
@@ -1997,7 +1999,7 @@ FSubstrateData GetSubstrateToonBSDF(
 	}
 #else
 	const uint bHasAnisotropy = 0;
-//	Anisotropy = 0.0f;
+	Anisotropy = 0.0f;
 #endif
 
 #if (SUBSTRATE_FASTPATH==0 || MATERIALBLENDING_ANY_TRANSLUCENT) // !Fast path or rendering translucent materials
@@ -2013,7 +2015,8 @@ FSubstrateData GetSubstrateToonBSDF(
 	TOON_METALLIC					(SubstrateData.InlinedBSDF) = Metallic;
 	TOON_SPECULAR					(SubstrateData.InlinedBSDF) = Specular;
 	TOON_ROUGHNESS					(SubstrateData.InlinedBSDF) = Roughness;
-	TOON_PATTERNUVS					(SubstrateData.InlinedBSDF) = PatternUVs;
+	TOON_PATTERNUVS					(SubstrateData.InlinedBSDF) = frac(PatternUVs);	// Always apply fract since we cannot store out UVs out of [0,1] range.
+	TOON_ANISOTROPY					(SubstrateData.InlinedBSDF) = Anisotropy;
 
 #if SUBSTRATE_INLINE_SHADING
 	SubstrateData.InlinedBSDF.Coverage	= 1.0f;
@@ -3592,10 +3595,11 @@ void FSubstrateBSDF::SubstrateSanitizeBSDF()
 			BSDF_SETEMISSIVE(this, max(BSDF_GETEMISSIVE(this), 0.0));
 
 			TOON_BASECOLOR(this) = SanitizeDiffuseAlbedo(TOON_BASECOLOR(this));
-			TOON_METALLIC(this) = clamp(TOON_METALLIC(this), 0.0, 1.0);
-			TOON_SPECULAR(this) = clamp(TOON_SPECULAR(this), 0.0, 1.0);
+			TOON_METALLIC(this) = saturate(TOON_METALLIC(this));
+			TOON_SPECULAR(this) = saturate(TOON_SPECULAR(this));
 			TOON_ROUGHNESS(this) = clamp(TOON_ROUGHNESS(this), 0.0, 1.0);	// We do not use SanitizeRoughness which clamps to View.MinRoughness that is higher than 0.
 			TOON_PATTERNUVS(this) = clamp(TOON_PATTERNUVS(this), 0.0.xx, 1.0.xx);
+			TOON_ANISOTROPY(this) = clamp(TOON_ANISOTROPY(this), -1.0, 1.0);
 			break;
 		}
 	}
@@ -3977,7 +3981,7 @@ float SubstrateGetBSDFAnisotropy(in FSubstrateBSDF BSDF)
 	case SUBSTRATE_BSDF_TYPE_HAIR:				return 0;
 	case SUBSTRATE_BSDF_TYPE_EYE:				return 0;
 	case SUBSTRATE_BSDF_TYPE_SINGLELAYERWATER:	return 0;
-	case SUBSTRATE_BSDF_TYPE_TOON:				return 0;
+	case SUBSTRATE_BSDF_TYPE_TOON:				return BSDF_GETHASANISOTROPY(BSDF) ? TOON_ANISOTROPY(BSDF) : 0.f;
 	}
 	return 0.f;
 }
@@ -4972,6 +4976,11 @@ FSubstrateBSDF UnpackSubstrateBSDFIn(FSubstrateMaterialContainer SubstrateBuffer
 		
 				UnpackToonCustomData(DataY, TOON_PATTERNUVS(OutBSDF), TOON_PROFILEID(OutBSDF));
 
+				if (BSDF_GETHASANISOTROPY(OutBSDF))
+				{
+					TOON_ANISOTROPY(OutBSDF) = asfloat(SubstrateLoadUint1(SubstrateBuffer, SubstrateAddressing));
+				}
+
 				BSDF_SETSSSTYPE(OutBSDF, SSS_TYPE_TRANSMITTANCE_SHADOW);
 
 				return OutBSDF;
diff --git a/Engine/Shaders/Private/Substrate/SubstrateEvaluation.ush b/Engine/Shaders/Private/Substrate/SubstrateEvaluation.ush
--- a/Engine/Shaders/Private/Substrate/SubstrateEvaluation.ush
+++ b/Engine/Shaders/Private/Substrate/SubstrateEvaluation.ush
@@ -1302,7 +1302,8 @@ FSubstrateEvaluateResult SubstrateEvaluateBSDFCommon(FSubstrateBSDFContext BSDFC
 			FToonProfileEvaluateSettings ToonEvalSettings = GetToonProfileEvaluateSettings(ToonProfileId);
 
 			const float2 OffsetFraction = float2(0.0f, 0.0f);
-			const float2 ToonPatternUV = TOON_PATTERNUVS(BSDFContext.BSDF);// (BSDFContext.PixelCoord - View.ViewRectMinAndSize.xy)* View.ViewSizeAndInvSize.zw;
+			const float2 ToonPatternUV = TOON_PATTERNUVS(BSDFContext.BSDF);
+			const float Anisotropy = TOON_ANISOTROPY(BSDFContext.BSDF);
 			float ToonProfileDiffuseDither  = EvaluateToonProfileDiffuseDither(ToonProfileId, ToonPatternUV, ToonEvalSettings);
 			float ToonProfileSpecularDither = EvaluateToonProfileSpecularDither(ToonProfileId, ToonPatternUV, ToonEvalSettings);
 
@@ -1333,13 +1334,23 @@ FSubstrateEvaluateResult SubstrateEvaluateBSDFCommon(FSubstrateBSDFContext BSDFC
 				}
 				else
 				{
-
-					float D = Substrate_D_GGX(SafeRoughness, Alpha2Spec, BSDFContext.Context.NoH);
+					float D;
+					if (BSDF_GETHASANISOTROPY(BSDFContext.BSDF) && Anisotropy!=0)
+					{
+						float2 AlphaXY = 0;
+						float  Alpha = Square(SafeRoughness);
+						GetAnisotropicRoughness(Alpha, Anisotropy, AlphaXY.x, AlphaXY.y);
+						D = Substrate_D_GGX_Aniso(AlphaXY.x, AlphaXY.y, BSDFContext.Context.NoH, BSDFContext.Context.XoH, BSDFContext.Context.YoH);
+					}
+					else
+					{
+						D = Substrate_D_GGX(SafeRoughness, Alpha2Spec, BSDFContext.Context.NoH);
+					}
 					//const float3 FresnelTerm = Substrate_F_GGX(F0, F90, BSDFContext.Context.VoH);
 					//float Vis = Substrate_Vis_GGX(SafeRoughness, Alpha2Spec, BSDFContext.Context.NoV, AreaLight.NoL);
 
 					Specular =  D;// * Vis * FresnelTerm;
-					Specular *= GGXMaxEnergyLookup(SafeRoughness, BSDFContext.Context.NoV);
+					Specular *= GGXMaxEnergyLookup(SafeRoughness, Anisotropy, BSDFContext.Context.NoV);
 				}
 
 				// TODO H_PDF and PDF
diff --git a/Engine/Shaders/Private/Substrate/SubstrateExport.ush b/Engine/Shaders/Private/Substrate/SubstrateExport.ush
--- a/Engine/Shaders/Private/Substrate/SubstrateExport.ush
+++ b/Engine/Shaders/Private/Substrate/SubstrateExport.ush
@@ -890,11 +890,16 @@ void PackSubstrateOut(
 							uint Data1 = PackToonCustomData(TOON_PATTERNUVS(BSDF), TOON_PROFILEID(BSDF), Dither);
 							SUBSTRATE_STORE_UINT1(Data1);
 
+							if (BSDF_GETHASANISOTROPY(BSDF))
+							{
+								SUBSTRATE_STORE_UINT1(asuint(TOON_ANISOTROPY(BSDF)));
+							}
+
 							SSSDataType = SSS_TYPE_TRANSMITTANCE_SHADOW;
 
 							const float ShadowMFPInMeters = GetToonProfileShadowSettings(TOON_PROFILEID(BSDF)).ShadowMFPInMeters;
 							SSSDataMFP = ShadowMFPInMeters.xxx;
-							// 12 bytes
+							// 12 to 16 bytes
 						}
 						break;
 						//case SUBSTRATE_BSDF_TYPE_SINGLELAYERWATER: 
@@ -1280,6 +1285,7 @@ FExportResult SubstrateMaterialExportOut(
 				//	- Roughness
 				//	- Metallic
 				//	- Specular
+				//  - Anisotropy is stored in the anisotropy buffer
 
 				uint PackedToonCustomData = PackToonCustomData(TOON_PATTERNUVS(CurrentBSDF), TOON_PROFILEID(CurrentBSDF), Dither);
 				Out.CustomData = UnpackRGBA8(PackedToonCustomData);
diff --git a/Engine/Shaders/Private/Substrate/SubstrateRead.ush b/Engine/Shaders/Private/Substrate/SubstrateRead.ush
--- a/Engine/Shaders/Private/Substrate/SubstrateRead.ush
+++ b/Engine/Shaders/Private/Substrate/SubstrateRead.ush
@@ -126,6 +126,7 @@ FSubstrateGBufferBSDF SubstrateReadGBufferBSDF(in FScreenSpaceData ScreenSpaceDa
 				GBuffer.BaseColor, GBuffer.Metallic, GBuffer.Specular, GBuffer.Roughness,
 				Emissive,
 				PatternUVs,
+				GBuffer.Anisotropy, GBuffer.WorldTangent,
 				SharedLocalBasisIndex, SharedLocalBasisTypes).InlinedBSDF;
 		}
 		else
diff --git a/Engine/Shaders/Shared/SubstrateVisualizeDefinitions.h b/Engine/Shaders/Shared/SubstrateVisualizeDefinitions.h
--- a/Engine/Shaders/Shared/SubstrateVisualizeDefinitions.h
+++ b/Engine/Shaders/Shared/SubstrateVisualizeDefinitions.h
@@ -701,6 +701,12 @@ FSubstratePixelDebugData ConvertToSerializableSubstratePixelDebugData(uint2 InCo
 			Data.Closures[i].F0B			= SpecularColor.b;
 			Data.Closures[i].Roughness		= TOON_ROUGHNESS(BSDF);
 			Data.Closures[i].SSSPRofileID	= TOON_PROFILEID(BSDF);
+
+			Data.Closures[i].bHasAnisotropy = BSDF_GETHASANISOTROPY(BSDF);
+			if (Data.Closures[i].bHasAnisotropy)
+			{
+				Data.Closures[i].Anisotropy = TOON_ANISOTROPY(BSDF);
+			}
 			
 			Data.Closures[i].GlintUVDDXx	= TOON_PATTERNUVS(BSDF).x;
 			Data.Closures[i].GlintUVDDXy	= TOON_PATTERNUVS(BSDF).y;
diff --git a/Engine/Source/Runtime/Engine/Private/Materials/HLSLMaterialTranslator.cpp b/Engine/Source/Runtime/Engine/Private/Materials/HLSLMaterialTranslator.cpp
--- a/Engine/Source/Runtime/Engine/Private/Materials/HLSLMaterialTranslator.cpp
+++ b/Engine/Source/Runtime/Engine/Private/Materials/HLSLMaterialTranslator.cpp
@@ -13484,6 +13484,7 @@ int32 FHLSLMaterialTranslator::SubstrateToonBSDF(
 	int32 BaseColor, int32 Metallic, int32 Specular, int32 Roughness,
 	int32 EmissiveColor, int32 Normal,
 	int32 PatternUVs,
+	int32 Anisotropy, int32 Tangent,
 	const FString& SharedLocalBasisIndexMacro, FSubstrateOperator* PromoteToOperator)
 {
 	if (ToonProfileId <= INDEX_NONE)
@@ -13495,7 +13496,7 @@ int32 FHLSLMaterialTranslator::SubstrateToonBSDF(
 	if (PromoteToOperator)
 	{
 		return AddCodeChunk(
-			MCT_Substrate, TEXT("Parameters.%s.PromoteParameterBlendedBSDFToOperator(GetSubstrateToonBSDF(%s, %s, %s, %s, %s, %s, %s, %s, Parameters.%s.Types), %u, %u, %u, %u) /* Normal:%s */"),
+			MCT_Substrate, TEXT("Parameters.%s.PromoteParameterBlendedBSDFToOperator(GetSubstrateToonBSDF(%s, %s, %s, %s, %s, %s, %s, %s, %s, %s, Parameters.%s.Types), %u, %u, %u, %u) /* Normal:%s */"),
 			*Substrate::GetParametersSubstrateTreeName(SubstrateTranslatorData.CurrentSubstrateCompilationContext),
 			*SubstrateGetCastParameterCode(ToonProfileId, MCT_Float),
 			*SubstrateGetCastParameterCode(BaseColor, MCT_Float3),
@@ -13504,6 +13505,8 @@ int32 FHLSLMaterialTranslator::SubstrateToonBSDF(
 			*SubstrateGetCastParameterCode(Roughness, MCT_Float),
 			*SubstrateGetCastParameterCode(EmissiveColor, MCT_Float3),
 			*SubstrateGetCastParameterCode(PatternUVs, MCT_Float2),
+			*SubstrateGetCastParameterCode(Anisotropy, MCT_Float),
+			*SubstrateGetCastParameterCode(Tangent, MCT_Float3),
 			*SharedLocalBasisIndexMacro,
 			*Substrate::GetParametersSharedLocalBasesName(SubstrateTranslatorData.CurrentSubstrateCompilationContext),
 			PromoteToOperator->Index,
@@ -13515,7 +13518,7 @@ int32 FHLSLMaterialTranslator::SubstrateToonBSDF(
 	}
 
 	return AddCodeChunk(
-		MCT_Substrate, TEXT("GetSubstrateToonBSDF(%s, %s, %s, %s, %s, %s, %s, %s, Parameters.%s.Types)"),
+		MCT_Substrate, TEXT("GetSubstrateToonBSDF(%s, %s, %s, %s, %s, %s, %s, %s, %s, %s, Parameters.%s.Types)"),
 		*SubstrateGetCastParameterCode(ToonProfileId, MCT_Float),
 		*SubstrateGetCastParameterCode(BaseColor, MCT_Float3),
 		*SubstrateGetCastParameterCode(Metallic, MCT_Float),
@@ -13523,6 +13526,8 @@ int32 FHLSLMaterialTranslator::SubstrateToonBSDF(
 		*SubstrateGetCastParameterCode(Roughness, MCT_Float),
 		*SubstrateGetCastParameterCode(EmissiveColor, MCT_Float3),
 		*SubstrateGetCastParameterCode(PatternUVs, MCT_Float2),
+		*SubstrateGetCastParameterCode(Anisotropy, MCT_Float),
+		*SubstrateGetCastParameterCode(Tangent, MCT_Float3),
 		*SharedLocalBasisIndexMacro,
 		*Substrate::GetParametersSharedLocalBasesName(SubstrateTranslatorData.CurrentSubstrateCompilationContext)
 	);
diff --git a/Engine/Source/Runtime/Engine/Private/Materials/HLSLMaterialTranslator.h b/Engine/Source/Runtime/Engine/Private/Materials/HLSLMaterialTranslator.h
--- a/Engine/Source/Runtime/Engine/Private/Materials/HLSLMaterialTranslator.h
+++ b/Engine/Source/Runtime/Engine/Private/Materials/HLSLMaterialTranslator.h
@@ -1291,6 +1291,7 @@ protected:
 		int32 BaseColor, int32 Metallic, int32 Specular, int32 Roughness,
 		int32 EmissiveColor, int32 Normal,
 		int32 PatternUVs,
+		int32 Anisotropy, int32 Tangent,
 		const FString& SharedLocalBasisIndexMacro, FSubstrateOperator* PromoteToOperator) override;
 	virtual int32 SubstrateHorizontalMixing(int32 Background, int32 Foreground, int32 Mix, int OperatorIndex, uint32 MaxDistanceFromLeaves) override;
 	virtual int32 SubstrateHorizontalMixingParameterBlending(int32 Background, int32 Foreground, int32 HorizontalMixCodeChunk, int32 NormalMixCodeChunk, const FString& SharedLocalBasisIndexMacro, int32 BackgroundBSDFNormalCodeChunk, int32 ForegroundBSDFNormalCodeChunk, FSubstrateOperator* PromoteToOperator) override;
diff --git a/Engine/Source/Runtime/Engine/Private/Materials/MaterialExpressionSubstrate.cpp b/Engine/Source/Runtime/Engine/Private/Materials/MaterialExpressionSubstrate.cpp
--- a/Engine/Source/Runtime/Engine/Private/Materials/MaterialExpressionSubstrate.cpp
+++ b/Engine/Source/Runtime/Engine/Private/Materials/MaterialExpressionSubstrate.cpp
@@ -2523,6 +2523,8 @@ int32 UMaterialExpressionSubstrateToonBSDF::Compile(class FMaterialCompiler* Com
 		CompileWithDefaultFloat3(Compiler, EmissiveColor, 0.0f, 0.0f, 0.0f),
 		NormalCodeChunk,
 		CompileWithDefaultFloat2(Compiler, PatternUVs, 0.0f, 0.0f),
+		CompileWithDefaultFloat1(Compiler, Anisotropy, 0.0f),
+		CompileWithDefaultTangentWS(Compiler, Tangent),
 		Compiler->GetSubstrateSharedLocalBasisIndexMacro(NewRegisteredSharedLocalBasis),
 		!SubstrateOperator.bUseParameterBlending || (SubstrateOperator.bUseParameterBlending && SubstrateOperator.bRootOfParameterBlendingSubTree) ? &SubstrateOperator : nullptr);
 
@@ -2557,6 +2559,10 @@ EMaterialValueType UMaterialExpressionSubstrateToonBSDF::GetInputValueType(int32
 		return MCT_Float3; // Emissive Color
 	case 6:
 		return MCT_Float2; // Pattern UVs
+	case 7:
+		return MCT_Float1; // Anisotropy
+	case 8:
+		return MCT_Float3; // Tangent
 	}
 
 	check(false);
@@ -2577,6 +2583,11 @@ void UMaterialExpressionSubstrateToonBSDF::GatherSubstrateMaterialInfo(FSubstrat
 	{
 		SubstrateMaterialInfo.AddToonProfile(ToonProfile);
 	}
+
+	if (Anisotropy.IsConnected())
+	{
+		SubstrateMaterialInfo.AddPropertyConnected(MP_Anisotropy);
+	}
 }
 
 FSubstrateOperator* UMaterialExpressionSubstrateToonBSDF::SubstrateGenerateMaterialTopologyTree(
@@ -2588,10 +2599,10 @@ FSubstrateOperator* UMaterialExpressionSubstrateToonBSDF::SubstrateGenerateMater
 	SubstrateOperator.bBSDFWritesEmissive = EmissiveColor.IsConnected();
 	SubstrateOperator.BSDFFeatures |= ESubstrateBsdfFeature::Toon;
 
-//	if ( SUBSTRATE_TODO fetch data from the profile to set anisotopy if value is !0.
-//	{
-//		SubstrateOperator.BSDFFeatures |= ESubstrateBsdfFeature::Anisotropy;
-//	}
+	if (Anisotropy.IsConnected())
+	{
+		SubstrateOperator.BSDFFeatures |= ESubstrateBsdfFeature::Anisotropy;
+	}
 
 	return &SubstrateOperator;
 }
diff --git a/Engine/Source/Runtime/Engine/Private/Materials/SubstrateTranslatorCommon.cpp b/Engine/Source/Runtime/Engine/Private/Materials/SubstrateTranslatorCommon.cpp
--- a/Engine/Source/Runtime/Engine/Private/Materials/SubstrateTranslatorCommon.cpp
+++ b/Engine/Source/Runtime/Engine/Private/Materials/SubstrateTranslatorCommon.cpp
@@ -1282,6 +1282,13 @@ bool FSubstrateCompilationContext::SubstrateGenerateDerivedMaterialOperatorData(
 						SubstrateMaterialRequestedSizeByte += UintByteSize;
 						// We cannot add more uint here to stay compatible with legacy gbuffer when Substrate::IsSubstrateBlendableGBufferEnabled(ShaderPlatform).
 						// Possible work around: disable some toon shading feature that cannot be stored in the GBuffer, or increase Custom GBufferD size to 64 bits.
+
+						if (!Substrate::IsSubstrateBlendableGBufferEnabled(ShaderPlatform)
+							&& It.Has(ESubstrateBsdfFeature::Anisotropy))
+						{
+							SubstrateMaterialRequestedSizeByte += UintByteSize;
+						}
+
 						break;
 					}
 					case SUBSTRATE_BSDF_TYPE_UNLIT:
diff --git a/Engine/Source/Runtime/Engine/Public/MaterialCompiler.h b/Engine/Source/Runtime/Engine/Public/MaterialCompiler.h
--- a/Engine/Source/Runtime/Engine/Public/MaterialCompiler.h
+++ b/Engine/Source/Runtime/Engine/Public/MaterialCompiler.h
@@ -742,6 +742,7 @@ public:
 		int32 BaseColor, int32 Metallic, int32 Specular, int32 Roughness,
 		int32 EmissiveColor, int32 Normal,
 		int32 PatternUVs,
+		int32 Anisotropy, int32 Tangent,
 		const FString& SharedLocalBasisIndexMacro, FSubstrateOperator* PromoteToOperator) = 0;
 	virtual int32 SubstrateHorizontalMixing(int32 Background, int32 Foreground, int32 Mix, int OperatorIndex, uint32 MaxDistanceFromLeaves) = 0;
 	virtual int32 SubstrateHorizontalMixingParameterBlending(int32 Background, int32 Foreground, int32 HorizontalMixCodeChunk, int32 NormalMixCodeChunk, const FString& SharedLocalBasisIndexMacro, int32 BackgroundBSDFNormalCodeChunk, int32 ForegroundBSDFNormalCodeChunk, FSubstrateOperator* PromoteToOperator) = 0;
@@ -1547,6 +1548,7 @@ public:
 		int32 BaseColor, int32 Metallic, int32 Specular, int32 Roughness,
 		int32 EmissiveColor, int32 Normal,
 		int32 PatternUVs,
+		int32 Anisotropy, int32 Tangent,
 		const FString& SharedLocalBasisIndexMacro, FSubstrateOperator* PromoteToOperator) override
 	{
 		return Compiler->SubstrateToonBSDF(
@@ -1554,6 +1556,7 @@ public:
 			BaseColor, Metallic, Specular, Roughness,
 			EmissiveColor, Normal,
 			PatternUVs,
+			Anisotropy, Tangent,
 			SharedLocalBasisIndexMacro, PromoteToOperator);
 	}
 
diff --git a/Engine/Source/Runtime/Engine/Public/Materials/MaterialExpressionSubstrate.h b/Engine/Source/Runtime/Engine/Public/Materials/MaterialExpressionSubstrate.h
--- a/Engine/Source/Runtime/Engine/Public/Materials/MaterialExpressionSubstrate.h
+++ b/Engine/Source/Runtime/Engine/Public/Materials/MaterialExpressionSubstrate.h
@@ -904,6 +904,18 @@ class UMaterialExpressionSubstrateToonBSDF : public UMaterialExpressionSubstrate
 	UPROPERTY()
 	FExpressionInput PatternUVs;
 
+	/**
+	 * Controls the anisotropy factor of the roughness. Positive value elongates the specular lobe along the Tangent vector, Negative value elongates the specular lobe along the perpendicular of the Tangent. (type = float, unit = unitless).
+	 */
+	UPROPERTY()
+	FExpressionInput Anisotropy;
+
+	/**
+	* Take a surface tangent as input. The tangent is considered tangent or world space according to the space properties on the main material node. (type = float3, unit = unitless, defaults to vertex tangent)
+	*/
+	UPROPERTY()
+	FExpressionInput Tangent;
+
 	/** Toon profile used to control the appearance of this BSDF. */
 	UPROPERTY(EditAnywhere, BlueprintReadOnly, Category = Material, meta = (DisplayName = "Toon Profile"))
 	TObjectPtr<class UToonProfile> ToonProfile;
diff --git a/Engine/Source/Runtime/Engine/Public/SceneView.h b/Engine/Source/Runtime/Engine/Public/SceneView.h
--- a/Engine/Source/Runtime/Engine/Public/SceneView.h
+++ b/Engine/Source/Runtime/Engine/Public/SceneView.h
@@ -1238,7 +1238,7 @@ BEGIN_GLOBAL_SHADER_PARAMETER_STRUCT_WITH_CONSTRUCTOR(FViewUniformShaderParamete
 	SHADER_PARAMETER(uint32, bShadingEnergyConservation)
 	SHADER_PARAMETER(uint32, bShadingEnergyPreservation)
 	SHADER_PARAMETER_TEXTURE(Texture2D<float2>, ShadingEnergyGGXSpecTexture)
-	SHADER_PARAMETER_TEXTURE(Texture2D<float>,  ShadingEnergyGGXMaxSpecTexture)
+	SHADER_PARAMETER_TEXTURE(Texture2D<float4>,  ShadingEnergyGGXMaxSpecTexture)
 	SHADER_PARAMETER_TEXTURE(Texture3D<float2>, ShadingEnergyGGXGlassTexture)
 	SHADER_PARAMETER_TEXTURE(Texture2D<float2>, ShadingEnergyClothSpecTexture)
 	SHADER_PARAMETER_TEXTURE(Texture2D<float>, ShadingEnergyDiffuseTexture)
diff --git a/Engine/Source/Runtime/Renderer/Private/AnisotropyRendering.cpp b/Engine/Source/Runtime/Renderer/Private/AnisotropyRendering.cpp
--- a/Engine/Source/Runtime/Renderer/Private/AnisotropyRendering.cpp
+++ b/Engine/Source/Runtime/Renderer/Private/AnisotropyRendering.cpp
@@ -142,10 +142,10 @@ bool GetAnisotropyPassShaders(
 	return true;
 }
 
-static bool ShouldDraw(const FMaterial& Material, bool bMaterialUsesAnisotropy)
+static bool ShouldDraw(const FMaterial& Material, bool bMaterialUsesAnisotropy, bool bMaterialIsSubstrateToon)
 {
 	const bool bIsNotTranslucent = IsOpaqueOrMaskedBlendMode(Material);
-	return (bMaterialUsesAnisotropy && bIsNotTranslucent && Material.GetShadingModels().HasAnyShadingModel({ MSM_DefaultLit, MSM_ClearCoat }));
+	return bMaterialUsesAnisotropy && bIsNotTranslucent && (Material.GetShadingModels().HasAnyShadingModel({ MSM_DefaultLit, MSM_ClearCoat }) || (Material.GetShadingModels().HasAnyShadingModel({ MSM_Strata }) && bMaterialIsSubstrateToon));
 }
 
 void FAnisotropyMeshProcessor::AddMeshBatch(
@@ -183,7 +183,7 @@ bool FAnisotropyMeshProcessor::TryAddMeshBatch(
 	const FMaterial& Material)
 {
 	bool bResult = true;
-	if (ShouldDraw(Material, Material.MaterialUsesAnisotropy_RenderThread()))
+	if (ShouldDraw(Material, Material.MaterialUsesAnisotropy_RenderThread(), uint16(Material.MaterialGetSubstrateMaterialBsdfFeatures_RenderThread() & ESubstrateBsdfFeature::Toon) > 0))
 	{
 		const FMeshDrawingPolicyOverrideSettings OverrideSettings = ComputeMeshOverrideSettings(MeshBatch);
 		const ERasterizerFillMode MeshFillMode = ComputeMeshFillMode(Material, OverrideSettings);
@@ -248,7 +248,7 @@ bool FAnisotropyMeshProcessor::Process(
 
 void FAnisotropyMeshProcessor::CollectPSOInitializers(const FSceneTexturesConfig& SceneTexturesConfig, const FMaterial& Material, const FPSOPrecacheVertexFactoryData& VertexFactoryData, const FPSOPrecacheParams& PreCacheParams, TArray<FPSOPrecacheData>& PSOInitializers)
 {
-	if (ShouldDraw(Material, Material.MaterialUsesAnisotropy_GameThread()) && 
+	if (ShouldDraw(Material, Material.MaterialUsesAnisotropy_GameThread(), uint16(Material.MaterialGetSubstrateMaterialBsdfFeatures_GameThread() & ESubstrateBsdfFeature::Toon) > 0) &&
 		SupportsAnisotropicMaterialPass(FeatureLevel, GShaderPlatformForFeatureLevel[FeatureLevel]))
 	{
 		const FMeshDrawingPolicyOverrideSettings OverrideSettings = ComputeMeshOverrideSettings(PreCacheParams);
diff --git a/Engine/Source/Runtime/Renderer/Private/ShadingEnergyConservation.cpp b/Engine/Source/Runtime/Renderer/Private/ShadingEnergyConservation.cpp
--- a/Engine/Source/Runtime/Renderer/Private/ShadingEnergyConservation.cpp
+++ b/Engine/Source/Runtime/Renderer/Private/ShadingEnergyConservation.cpp
@@ -152,9 +152,9 @@ namespace ShadingEnergyConservationData
 
 		FTextureRHIRef Texture = OutTexture->GetRHI();
 
-		check(InFormat == PF_G16R16 || InFormat == PF_R8G8 || InFormat == PF_G16 || InFormat == PF_R8);
-		const uint32 ComponentCount = InFormat == PF_G16R16 || InFormat == PF_R8G8 ? 2u: 1u;
-		const bool b8bit = InFormat == PF_R8G8 || InFormat == PF_R8;
+		check(InFormat == PF_A16B16G16R16 || InFormat == PF_R8G8B8A8 || InFormat == PF_G16R16 || InFormat == PF_R8G8 || InFormat == PF_G16 || InFormat == PF_R8);
+		const uint32 ComponentCount = (InFormat == PF_A16B16G16R16 || InFormat == PF_R8G8B8A8) ? 4u : ((InFormat == PF_G16R16 || InFormat == PF_R8G8) ? 2u: 1u);
+		const bool b8bit = InFormat == PF_R8G8B8A8 || InFormat == PF_R8G8 || InFormat == PF_R8;
 		const uint32 DstBytesPerPixel = (b8bit? 1u : 2u) * ComponentCount;
 
 		// Write the contents of the texture with transcoding
@@ -175,6 +175,11 @@ namespace ShadingEnergyConservationData
 					{
 						Dest[1] = uint8(FMath::Clamp(float(Value.G) * 0xFF, 0u, 0xFF));
 					}
+					if (ComponentCount > 2)
+					{
+						Dest[2] = uint8(FMath::Clamp(float(Value.B) * 0xFF, 0u, 0xFF));
+						Dest[3] = uint8(FMath::Clamp(float(Value.A) * 0xFF, 0u, 0xFF));
+					}
 				}
 				else
 				{
@@ -186,6 +191,11 @@ namespace ShadingEnergyConservationData
 					{
 						Dest[1] = uint16(FMath::Clamp(float(Value.G) * 0xFFFF, 0u, 0xFFFF));
 					}
+					if (ComponentCount > 2)
+					{
+						Dest[2] = uint16(FMath::Clamp(float(Value.B) * 0xFFFF, 0u, 0xFFFF));
+						Dest[3] = uint16(FMath::Clamp(float(Value.A) * 0xFFFF, 0u, 0xFFFF));
+					}
 				}
 			}
 		}
@@ -431,7 +441,7 @@ void Init(FRDGBuilder& GraphBuilder, const FViewInfo& View)
 		const bool bR16Supported = GPixelFormats[PF_G16].Supported && UE::PixelFormat::HasCapabilities(PF_G16, EPixelFormatCapabilities::TextureFilterable);
 
 		const EPixelFormat SpecFormat = bRuntimeGeneration && CVarShadingEnergyConservation_TableFormat.GetValueOnRenderThread() > 0 ? PF_G32R32F : (bRG16Supported ? PF_G16R16 : PF_R8G8);
-		const EPixelFormat MaxSpecFormat = bRuntimeGeneration ? PF_R16F : (bRG16Supported ? PF_G16 : PF_R8);
+		const EPixelFormat MaxSpecFormat = bRuntimeGeneration ? PF_FloatRGBA : (bRG16Supported ? PF_A16B16G16R16 : PF_R8G8B8A8);
 		const EPixelFormat DiffFormat = bR16Supported ? PF_G16 : PF_R8;
 		const bool bBuildTable = 
 			bRuntimeGeneration ||
@@ -554,7 +564,7 @@ void Init(FRDGBuilder& GraphBuilder, const FViewInfo& View)
 				// Precomputed data are stored as float16
 				check(SpecFormat == PF_G16R16 || SpecFormat == PF_R8G8);
 				check(DiffFormat == PF_G16 || DiffFormat == PF_R8);
-				check(MaxSpecFormat == PF_G16 || MaxSpecFormat == PF_R8);
+				check(MaxSpecFormat == PF_A16B16G16R16 || MaxSpecFormat == PF_R8G8B8A8);
 
 				Out.GGXSpecEnergyTexture    = ShadingEnergyConservationData::CreateTexture2D(GraphBuilder.RHICmdList, GEngine->GGXReflectionEnergyTexture,   SpecFormat,    TEXT("Shading.GGXReflectionEnergy"));
 				Out.GGXGlassEnergyTexture   = ShadingEnergyConservationData::CreateTexture3D(GraphBuilder.RHICmdList, GEngine->GGXTransmissionEnergyTexture, SpecFormat,    TEXT("Shading.GGXTransmissionEnergy"));
diff --git a/Engine/Source/Runtime/Renderer/Private/Substrate/SubstrateVisualize.cpp b/Engine/Source/Runtime/Renderer/Private/Substrate/SubstrateVisualize.cpp
--- a/Engine/Source/Runtime/Renderer/Private/Substrate/SubstrateVisualize.cpp
+++ b/Engine/Source/Runtime/Renderer/Private/Substrate/SubstrateVisualize.cpp
@@ -921,6 +921,13 @@ void AddProcessAndPrintSubstrateMaterialPropertiesPasses(FRDGBuilder& GraphBuild
 						PrintFSmall(DrawPosX + DataOffset, DrawPosY, Closure.Roughness);
 						NewLine();
 
+						if (Closure.bHasAnisotropy)
+						{
+							PrintS(DrawPosX + 0.0f, DrawPosY, TEXT("Anisotropy"), Orange);
+							PrintFSmall(DrawPosX + DataOffset, DrawPosY, Closure.Anisotropy);
+							NewLine();
+						}
+
 						PrintS(DrawPosX + 0.0f, DrawPosY, TEXT("Pattern UVs"), Orange);
 						PrintF2DAdapt(DrawPosX + DataOffset, DrawPosY, Data.Closures[i].GlintUVDDXx, Data.Closures[i].GlintUVDDXy);
 						NewLine();
