Substrate - toon - fix for bad hard edge with spot light when using unified diffuse. (missing falloff term)

#tests PC editor with toon shading
#rnx
#Synced-CL 51842837

diff --git a/Engine/Shaders/Private/Substrate/SubstrateEvaluation.ush b/Engine/Shaders/Private/Substrate/SubstrateEvaluation.ush
--- a/Engine/Shaders/Private/Substrate/SubstrateEvaluation.ush
+++ b/Engine/Shaders/Private/Substrate/SubstrateEvaluation.ush
@@ -1359,7 +1359,7 @@ FSubstrateEvaluateResult SubstrateEvaluateBSDFCommon(FSubstrateBSDFContext BSDFC
 			const bool bApplyDiffuseRamp	= SUBSTRATE_LIGHTPASS_APPLIES_TOON_DIFFUSE_RAMP;
 			const bool bApplyDiffuseShadow	= SUBSTRATE_LIGHTPASS_APPLIES_TOON_DIFFUSE_RAMP || (!SUBSTRATE_LIGHTPASS_APPLIES_TOON_DIFFUSE_RAMP && !ToonEvalSettings.bUnifiedDiffuseIncludeShadow);
 
-			Sample.ToonData					= (BSDFContext.Context.NoL * 0.5 + 0.5) * (bApplyDiffuseShadow ? 1.0f : TransmittanceShadow);
+			Sample.ToonData					= (BSDFContext.Context.NoL * 0.5 + 0.5) * AreaLight.Falloff * (bApplyDiffuseShadow ? 1.0f : TransmittanceShadow);
 
 			// Diffuse
 			const float3 DiffuseRamp		= bApplyDiffuseRamp ? EvaluateToonProfileDiffuseRamp(ToonProfileId, BSDFContext.Context.NoL + ToonProfileDiffuseDither) : 1.0f;
