Mobile: Support WPO evaluation property
Adds a uniform branch to vertex shaders with an active WPO
#Synced-CL 52054491

diff --git a/Engine/Shaders/Private/MaterialTemplate.ush b/Engine/Shaders/Private/MaterialTemplate.ush
--- a/Engine/Shaders/Private/MaterialTemplate.ush
+++ b/Engine/Shaders/Private/MaterialTemplate.ush
@@ -3995,13 +3995,11 @@ void ApplyMaterialPreviousFirstPersonTransform(FMaterialVertexParameters Paramet
 bool ShouldEnableWorldPositionOffset(FMaterialVertexParameters Parameters)
 {
 #if USES_WORLD_POSITION_OFFSET
-	#if (FEATURE_LEVEL > FEATURE_LEVEL_ES3_1 || MOBILE_SUPPORTS_SM5_MATERIAL_NODES) // Non optional WPO on mobile not supporting sm5 material nodes
-		if (!Parameters.bEvaluateWorldPositionOffset ||
-			(GetPrimitiveData(Parameters).Flags & PRIMITIVE_SCENE_DATA_FLAG_EVALUATE_WORLD_POSITION_OFFSET) == 0)
-		{
-			return false;
-		}
-	#endif
+	if (!Parameters.bEvaluateWorldPositionOffset ||
+		(GetPrimitiveData(Parameters).Flags & PRIMITIVE_SCENE_DATA_FLAG_EVALUATE_WORLD_POSITION_OFFSET) == 0)
+	{
+		return false;
+	}
 	return true;
 #else
 	return false;
diff --git a/Engine/Source/Runtime/Engine/Private/ShaderCompiler/ShaderCompiler.cpp b/Engine/Source/Runtime/Engine/Private/ShaderCompiler/ShaderCompiler.cpp
--- a/Engine/Source/Runtime/Engine/Private/ShaderCompiler/ShaderCompiler.cpp
+++ b/Engine/Source/Runtime/Engine/Private/ShaderCompiler/ShaderCompiler.cpp
@@ -3507,7 +3507,6 @@ void GlobalBeginCompileShader(
 			const bool bGLESDeferredShading = (Target.Platform == SP_OPENGL_ES3_1_ANDROID && bAllowFramebufferFetch);
 			SET_SHADER_DEFINE_AND_COMPILE_ARGUMENT(Input.Environment, USE_GLES_FBF_DEFERRED, bGLESDeferredShading);
 			SET_SHADER_DEFINE(Input.Environment, MOBILE_EXTENDED_GBUFFER, MobileUsesExtendedGBuffer((EShaderPlatform)Target.Platform) ? 1 : 0);
-			SET_SHADER_DEFINE(Input.Environment, MOBILE_SUPPORTS_SM5_MATERIAL_NODES, MobileSupportsSM5MaterialNodes((EShaderPlatform)Target.Platform) ? 1 : 0);
 		}
 		else
 		{
