Fix minor shader warning

#rnx
#Synced-CL 52066376

diff --git a/Engine/Shaders/Private/ShadingEnergyConservationTable.usf b/Engine/Shaders/Private/ShadingEnergyConservationTable.usf
--- a/Engine/Shaders/Private/ShadingEnergyConservationTable.usf
+++ b/Engine/Shaders/Private/ShadingEnergyConservationTable.usf
@@ -94,7 +94,7 @@ void BuildEnergyTableCS(uint3 DispatchThreadId : SV_DispatchThreadID, uint3 Grou
 		}
 #else
 		float2 AlphaXY = 0;
-		GetAnisotropicRoughness(Alpha, Anisotropy, AlphaXY.x, AlphaXY.y);
+		GetAnisotropicRoughness(Alpha.x, Anisotropy, AlphaXY.x, AlphaXY.y);
 		const float3 H = ImportanceSampleVisibleGGX(RandSample.xy, AlphaXY, V, false).xyz;
 		const float3 L = reflect(-V, H);
 		if (L.z > 0)
