[Front Layer Translucency] Use early depth/stencil test + late write on FLT GBuffer pass.
- saves about 0.1ms on some areas of MegaLights demo when using front layer translucency, and will be particularly important for an upcoming change which will add a UAV write from this pixel shader.

#Synced-CL 52039087

#Synced-CL 52039087

diff --git a/Engine/Shaders/Private/FrontLayerTranslucency.usf b/Engine/Shaders/Private/FrontLayerTranslucency.usf
--- a/Engine/Shaders/Private/FrontLayerTranslucency.usf
+++ b/Engine/Shaders/Private/FrontLayerTranslucency.usf
@@ -4,6 +4,13 @@
 	FrontLayerTranslucency.usf
 =============================================================================*/
 
+#if defined(MainPS) && !OUTPUT_PIXEL_DEPTH_OFFSET
+	// Force early depth/stencil test unless the material uses depth offset to speedup the front layer GBuffer pass.
+	// Cant use EARLYDEPTHSTENCIL since that causes artifacts when combined with clip/discard.
+	// COMPILER_FORCE_DEPTHSTENCIL_EARLYTEST_LATEWRITE must be defined before including Platform.ush
+	#define COMPILER_FORCE_DEPTHSTENCIL_EARLYTEST_LATEWRITE 1
+#endif // defined(MainPS) && !OUTPUT_PIXEL_DEPTH_OFFSET
+
 #include "Common.ush"
 #include "GammaCorrectionCommon.ush"
 #include "BRDF.ush"
@@ -56,6 +63,16 @@ void MainVS(
 	Output.FactoryInterpolants = VertexFactoryGetInterpolantsVSToPS(Input, VFIntermediates, VertexParameters);
 }
 
+#if defined(MainPS) && !OUTPUT_PIXEL_DEPTH_OFFSET
+	// Force early depth/stencil test unless the material uses depth offset to speedup the front layer GBuffer pass.
+	// Cant use EARLYDEPTHSTENCIL since that causes artifacts when combined with clip/discard.
+	// DEPTHSTENCIL_EARLYTEST_LATEWRITE can only be defined after including Platform.ush
+	#define PIXELSHADER_EARLYDEPTHSTENCIL DEPTHSTENCIL_EARLYTEST_LATEWRITE
+#else // !defined(MainPS) || OUTPUT_PIXEL_DEPTH_OFFSET
+	#define PIXELSHADER_EARLYDEPTHSTENCIL
+#endif // !defined(MainPS) || OUTPUT_PIXEL_DEPTH_OFFSET
+
+PIXELSHADER_EARLYDEPTHSTENCIL
 void MainPS(
 	FVertexFactoryInterpolantsVSToPS Interpolants,
 	FFrontLayerTranslucencyGBufferInterpolantsVSToPS PassInterpolants,
