Fixed shader compilation error due to earlying out from a shader before a group barrier.

#Synced-CL 52064906

diff --git a/Engine/Shaders/Private/Skinning/AnimSequenceTransformProvider.usf b/Engine/Shaders/Private/Skinning/AnimSequenceTransformProvider.usf
--- a/Engine/Shaders/Private/Skinning/AnimSequenceTransformProvider.usf
+++ b/Engine/Shaders/Private/Skinning/AnimSequenceTransformProvider.usf
@@ -357,15 +357,15 @@ void EvaluateConcatenateScatterCS(uint3 GroupId : SV_GroupID, uint GroupIndex :
 	const uint HeaderIndex           = GroupId.x * InstancesPerGroup + InstanceInGroup;
 	const uint SharedTransformOffset = InstanceInGroup * SlotsPerInstance;
 
-	if (HeaderIndex >= NumTotalHeaders)
+	const bool bValidInstance = HeaderIndex < NumTotalHeaders;
+
+	FEvaluateConcatenateScatterBlockHeader BlockHeader = (FEvaluateConcatenateScatterBlockHeader)0;
+	if (bValidInstance)
 	{
-		return;
+		BlockHeader = EvaluateBlockHeaderBuffer.Load<FEvaluateConcatenateScatterBlockHeader>((HeaderBaseOffset + HeaderIndex) * HeaderSize);
 	}
 
-	const FEvaluateConcatenateScatterBlockHeader BlockHeader =
-		EvaluateBlockHeaderBuffer.Load<FEvaluateConcatenateScatterBlockHeader>((HeaderBaseOffset + HeaderIndex) * HeaderSize);
-
-	if (BoneIndex < BlockHeader.NumCompactBones)
+	if (bValidInstance && BoneIndex < BlockHeader.NumCompactBones)
 	{
 		SharedTransform[SharedTransformOffset + BoneIndex] = EvaluateAndRetargetBone(
 			BoneIndex,
@@ -380,7 +380,7 @@ void EvaluateConcatenateScatterCS(uint3 GroupId : SV_GroupID, uint GroupIndex :
 
 	GroupMemoryBarrierWithGroupSync();
 
-	if (BoneIndex < BlockHeader.NumCompactBones)
+	if (bValidInstance && BoneIndex < BlockHeader.NumCompactBones)
 	{
 		uint TransformIndex;
 		FBoneTransformWithUniformScale GlobalTransform = ConcatenateToGlobal(BoneIndex, SharedTransformOffset, BlockHeader.BoneMapBufferOffset, BlockHeader.InvGlobalRefPoseTransformBufferOffset, TransformIndex);
@@ -400,15 +400,15 @@ void EvaluateBlendConcatenateScatterCS(uint3 GroupId : SV_GroupID, uint GroupInd
 	const uint HeaderIndex           = GroupId.x * InstancesPerGroup + InstanceInGroup;
 	const uint SharedTransformOffset = InstanceInGroup * SlotsPerInstance;
 
-	if (HeaderIndex >= NumTotalHeaders)
+	const bool bValidInstance = HeaderIndex < NumTotalHeaders;
+
+	FEvaluateBlendConcatenateScatterBlockHeader BlockHeader = (FEvaluateBlendConcatenateScatterBlockHeader)0;
+	if (bValidInstance)
 	{
-		return;
+		BlockHeader = EvaluateBlockHeaderBuffer.Load<FEvaluateBlendConcatenateScatterBlockHeader>((HeaderBaseOffset + HeaderIndex) * HeaderSize);
 	}
 
-	const FEvaluateBlendConcatenateScatterBlockHeader BlockHeader =
-		EvaluateBlockHeaderBuffer.Load<FEvaluateBlendConcatenateScatterBlockHeader>((HeaderBaseOffset + HeaderIndex) * HeaderSize);
-
-	if (BoneIndex < BlockHeader.NumCompactBones)
+	if (bValidInstance && BoneIndex < BlockHeader.NumCompactBones)
 	{
 		FBoneTransformWithUniformScale SourceTransform = EvaluateAndRetargetBone(
 			BoneIndex,
@@ -435,7 +435,7 @@ void EvaluateBlendConcatenateScatterCS(uint3 GroupId : SV_GroupID, uint GroupInd
 
 	GroupMemoryBarrierWithGroupSync();
 
-	if (BoneIndex < BlockHeader.NumCompactBones)
+	if (bValidInstance && BoneIndex < BlockHeader.NumCompactBones)
 	{
 		uint TransformIndex;
 		FBoneTransformWithUniformScale GlobalTransform = ConcatenateToGlobal(BoneIndex, SharedTransformOffset, BlockHeader.BoneMapBufferOffset, BlockHeader.InvGlobalRefPoseTransformBufferOffset, TransformIndex);
