Path Tracing: Implement decals for PrimaryRays mode and add additional debug visualization modes

Refactor code to share as much as possible between the reference path tracer and these new modes. In PrimaryRays mode, the decals are only applied to camera rays for now to keep things as fast as possible.

Minor cleanup to how ApproximateCaustics is used to avoid branching.

#rnx
#Synced-CL 52066376

diff --git a/Engine/Shaders/Private/PathTracing/PathTracingCore.ush b/Engine/Shaders/Private/PathTracing/PathTracingCore.ush
--- a/Engine/Shaders/Private/PathTracing/PathTracingCore.ush
+++ b/Engine/Shaders/Private/PathTracing/PathTracingCore.ush
@@ -20,9 +20,10 @@
 
 #include "PathTracingCommon.ush"
 #include "../RayTracing/RayTracingHitGroupCommon.ush"
-#include "../RayTracing/RayTracingDecalGrid.ush"
+#include "PathTracingDecals.ush"
 
 #include "../ShadingModels.ush"
+
 #include "./Utilities/PathTracingRandomSequence.ush"
 #include "./Utilities/PathTracingRIS.ush"
 #include "./Light/PathTracingLightSampling.ush"
@@ -60,15 +61,11 @@ uint MISMode;
 // 1: Light sampling
 uint VolumeMISMode;
 
-uint ApproximateCaustics;
+float ApproximateCaustics; // Either 0 (off) or 1 (on) -- done this way to avoid conditionals
 uint SamplerType;
-uint EnableDBuffer;
 uint ApplyDiffuseSpecularOverrides;
 uint EnabledDirectLightingContributions;
 uint EnabledIndirectLightingContributions;
-float DecalRoughnessCutoff;
-float MeshDecalRoughnessCutoff;
-float MeshDecalBias;
 float MaxPathIntensity;
 float MaxNormalBias;
 float FilterWidth;
@@ -87,7 +84,6 @@ RWTexture2D<float4> AlbedoTexture;
 RWTexture2D<float4> NormalTexture;
 RWTexture2D<float>  DepthTexture;
 RaytracingAccelerationStructure TLAS;
-RaytracingAccelerationStructure DecalTLAS;
 uint SceneVisibleLightCount;
 
 Buffer<float> StartingMediumData;
@@ -367,135 +363,6 @@ struct FPathState
     }
 };
 
-bool HasDecalData(FDecalShaderPayload DecalPayload)
-{
-	return min3(DecalPayload.GetBaseColor().a, DecalPayload.GetWorldNormal().a, DecalPayload.GetMetallicSpecularRoughness().a) < 1.0;
-}
-
-void ApplyDecal(FDecalShaderPayload DecalPayload, inout FPackedPathTracingPayload PackedPayload, float3 RayDirection, float PathRoughness)
-{
-	{
-		float3 Radiance = PackedPayload.UnpackRadiance();
-		Radiance += DecalPayload.GetEmissive();
-		PackedPayload.PackRadiance(Radiance);
-	}
-
-	float4 DecalBaseColor = DecalPayload.GetBaseColor();
-	float4 DecalWorldNormal = DecalPayload.GetWorldNormal();
-	float4 DecalMetallicSpecularRoughness = DecalPayload.GetMetallicSpecularRoughness();
-
-	bool bHasColorResponse = PackedPayload.HasDecalResponseColor() && DecalBaseColor.a < 1.0;
-	bool bHasNormalResponse = PackedPayload.HasDecalResponseNormal() && DecalWorldNormal.a < 1.0;
-	bool bHasRoughnessResponse = PackedPayload.HasDecalResponseRoughness() && DecalMetallicSpecularRoughness.a < 1.0;
-
-	if (!bHasColorResponse && !bHasNormalResponse && !bHasRoughnessResponse)
-	{
-		// nothing to do
-		return;
-	}
-
-	float4 RoughnessData = PackedPayload.UnpackRoughnessAniso();
-	float Roughness = RoughnessData.x;
-	float3 DiffuseColor = PackedPayload.UnpackDiffuseColor();
-	float3 SpecularColor = PackedPayload.UnpackSpecularColor();
-	float Metallic = F0RGBToMetallic(SpecularColor);
-	float Specular = F0RGBToDielectricSpecular(SpecularColor);
-	float3 BaseColor = lerp(DiffuseColor, SpecularColor, Metallic);
-
-	if (bHasColorResponse)
-	{
-		BaseColor = BaseColor * DecalBaseColor.a + DecalBaseColor.rgb;
-	}
-
-	if (bHasNormalResponse)
-	{
-		// We normalize the normal to get smoother visual result (it helps to avoid having D_GGX explodes toward infinity, and matches ApplyDBufferData(...))
-		float3 WorldNormal = PackedPayload.UnpackWorldNormal();
-		WorldNormal = normalize(WorldNormal * DecalWorldNormal.a + DecalWorldNormal.rgb);
-		PackedPayload.PackWorldNormal(WorldNormal);
-	}
-
-	if (bHasRoughnessResponse)
-	{
-		Metallic  = Metallic  * DecalMetallicSpecularRoughness.a + DecalMetallicSpecularRoughness.x;
-		Specular  = Specular  * DecalMetallicSpecularRoughness.a + DecalMetallicSpecularRoughness.y;
-		Roughness = Roughness * DecalMetallicSpecularRoughness.a + DecalMetallicSpecularRoughness.z;
-
-		if (ApproximateCaustics)
-		{
-			Roughness = max(Roughness, PathRoughness);
-		}
-	}
-
-	if (bHasColorResponse || bHasRoughnessResponse)
-	{
-		// the encode-decode is potentially lossy, so only do it if we actually need to
-		PackedPayload.PackDiffuseColor(BaseColor - BaseColor * Metallic);
-		PackedPayload.PackSpecularColor(lerp(DielectricSpecularToF0(Specular), BaseColor, Metallic));
-		PackedPayload.PackRoughnessAniso(float4(Roughness, RoughnessData.yzw));
-	}
-}
-
-void EvaluateDecals(FPathState PathState, float HitT, inout FDecalShaderPayload DecalPayload)
-{
-	if (PathState.GetPathRoughness() < MeshDecalRoughnessCutoff)
-	{
-		FRayDesc DecalRay = (FRayDesc)0;
-		// go a bit further than the current hit, to handle decals partially inside the model
-		DecalRay.Direction = -PathState.Ray.Direction;
-		DecalRay.Origin = DecalPayload.TranslatedWorldPos - DecalRay.Direction * MeshDecalBias;
-		DecalRay.TMin = 0.0f;
-		DecalRay.TMax = HitT + MeshDecalBias;
-
-        const uint InstanceInclusionMask = PathState.GetInstanceInclusionMaskOpaque() | PathState.GetInstanceInclusionMaskTranslucent();
-		const uint MissShaderIndex = 0; // TODO
-
-		for (;;)
-		{
-			DecalPayload.SetMiss();
-			TraceRay(
-				DecalTLAS,
-#if DECAL_USE_REVERSE_CULLING
-				RAY_FLAG_CULL_FRONT_FACING_TRIANGLES,
-#else
-				RAY_FLAG_CULL_BACK_FACING_TRIANGLES,
-#endif
-				InstanceInclusionMask,
-				RAY_TRACING_SHADER_SLOT_MATERIAL,
-				RAY_TRACING_NUM_SHADER_SLOTS,
-				MissShaderIndex,
-				DecalRay.GetNativeDesc(),
-				DecalPayload);
-#if NEED_TMIN_WORKAROUND // extra safety - discard hit if not conforming
-			if (DecalPayload.HitT <= DecalRay.TMin)
-			{
-				DecalPayload.HitT = -1.0; 
-			}
-#endif
-			if (DecalPayload.IsMiss())
-			{
-				break;
-			}
-
-			// prepare next step around the loop
-			// retrace the exact same ray with TMin one ulp past the hit we just found
-			DecalRay.TMin = ComputeNewTMin(DecalRay.Origin, DecalRay.Direction, DecalPayload.HitT);
-		}
-	}
-
-#if PLATFORM_SUPPORTS_CALLABLE_SHADERS
-	if (PathState.GetPathRoughness() < DecalRoughnessCutoff)
-	{
-		FDecalLoopCount DecalLoopCount = DecalGridLookup(DecalPayload.TranslatedWorldPos);
-		for (uint Index = 0, Num = DecalLoopCount.NumDecals; Index < Num; ++Index)
-		{
-			uint DecalId = GetDecalId(Index, DecalLoopCount);
-			CallShader(DecalId, DecalPayload);
-		}
-	}
-#endif
-}
-
 FPackedPathTracingPayload TraceTransparentRay(inout FPathState PathState, int Bounce, inout FVolumeSegment VolumeSegment)
 {
 	const bool bLastBounce = Bounce == MaxBounces;
@@ -521,7 +388,7 @@ FPackedPathTracingPayload TraceTransparentRay(inout FPathState PathState, int Bo
 		return Payload;
 	}
 
-	FPackedPathTracingPayload OpaquePayload = InitPathTracingPayload(PathState.GetFirstScatterType(), ApproximateCaustics ? PathState.GetPathRoughness() : 0.0);
+	FPackedPathTracingPayload OpaquePayload = InitPathTracingPayload(PathState.GetFirstScatterType(), ApproximateCaustics * PathState.GetPathRoughness());
 	OpaquePayload.SetDBufferA(float4(0, 0, 0, 1));
 	OpaquePayload.SetDBufferB(float4(0, 0, 0, 1));
 	OpaquePayload.SetDBufferC(float4(0, 0, 0, 1));
@@ -557,45 +424,22 @@ FPackedPathTracingPayload TraceTransparentRay(inout FPathState PathState, int Bo
 		PathState.Ray.GetNativeDesc(),
 		OpaquePayload);
 #endif
-	if (OpaquePayload.IsDecalReceiver() && PathState.GetPathRoughness() < max(MeshDecalRoughnessCutoff, DecalRoughnessCutoff))
-	{
-		FDecalShaderPayload DecalPayload = InitDecalShaderPayload(PathState.Ray.Origin + OpaquePayload.HitT * PathState.Ray.Direction);
-		EvaluateDecals(PathState, OpaquePayload.HitT, DecalPayload);
-#if USE_DBUFFER
-		if (OpaquePayload.UsesDBufferLookup() && HasDecalData(DecalPayload))
-		{
-			if (EnableDBuffer)
-			{
-				// Retrace ray with DBuffer data
-				OpaquePayload = InitPathTracingPayload(PathState.GetFirstScatterType(), ApproximateCaustics ? PathState.GetPathRoughness() : 0.0);
-
-				OpaquePayload.SetDBufferA(DecalPayload.GetBaseColor());
-				OpaquePayload.SetDBufferB(DecalPayload.GetWorldNormal());
-				OpaquePayload.SetDBufferC(DecalPayload.GetMetallicSpecularRoughness());
-				OpaquePayload.SetStochasticSlabRand(RISRandSample.z); // Must use same random value to get a consistent result to the first call
-				OpaquePayload.SetFlag(PATH_TRACING_PAYLOAD_INPUT_FLAG_IGNORE_TRANSLUCENT);
-				TraceRay(
-					TLAS,
-					RayFlags,
-					InstanceInclusionMaskOpaque,
-					RAY_TRACING_SHADER_SLOT_MATERIAL,
-					RAY_TRACING_NUM_SHADER_SLOTS,
-					MissShaderIndex,
-					PathState.Ray.GetNativeDesc(),
-					OpaquePayload);
-			}
-			else
-			{
-				// The user asked for a dbuffer lookup, but we have disabled the double-trace.
-				// Set all the decal response flags so we at least pickup some kind of response
-				// even if it won't be exactly what the user had in mind.
-				OpaquePayload.SetFlag(PATH_TRACING_PAYLOAD_OUTPUT_FLAG_DECAL_RESPONSE_COLOR | PATH_TRACING_PAYLOAD_OUTPUT_FLAG_DECAL_RESPONSE_NORMAL | PATH_TRACING_PAYLOAD_OUTPUT_FLAG_DECAL_RESPONSE_ROUGHNESS);
-			}
-		}
-#endif
-		// Now that we know the real decal data, apply it
-		ApplyDecal(DecalPayload, OpaquePayload, PathState.Ray.Direction, PathState.GetPathRoughness());
-	}
+    if (OpaquePayload.IsDecalReceiver() && PathState.GetPathRoughness() < max(MeshDecalRoughnessCutoff, DecalRoughnessCutoff))
+    {
+        FDecalShaderPayload DecalPayload = InitDecalShaderPayload(PathState.Ray.Origin + OpaquePayload.HitT * PathState.Ray.Direction);
+        // Evaluate mesh decals and decal actors
+        EvaluateDecals(
+            PathState.Ray.Direction,
+            PathState.GetInstanceInclusionMaskOpaque() | PathState.GetInstanceInclusionMaskTranslucent(),
+            OpaquePayload.HitT,
+            PathState.GetPathRoughness(),
+            DecalPayload
+        );
+        // Handle d-buffer requests (will potentially re-trace the primary ray if needed)
+        PATH_TRACER_DECAL_PROCESS_DBUFFER(OpaquePayload, ApproximateCaustics, RISRandSample.z)
+        // Now that we know the real decal data, apply it
+        ApplyDecal(DecalPayload, OpaquePayload, ApproximateCaustics * PathState.GetPathRoughness());
+    }
 
 	if (OpaquePayload.IsHit())
 	{
@@ -604,7 +448,7 @@ FPackedPathTracingPayload TraceTransparentRay(inout FPathState PathState, int Bo
 
 	for (;;)
 	{
-		FPackedPathTracingPayload PackedPayload = InitPathTracingPayload(PathState.GetFirstScatterType(), ApproximateCaustics ? PathState.GetPathRoughness() : 0.0);
+		FPackedPathTracingPayload PackedPayload = InitPathTracingPayload(PathState.GetFirstScatterType(), ApproximateCaustics * PathState.GetPathRoughness());
 		PackedPayload.SetDBufferA(float4(0, 0, 0, 1));
 		PackedPayload.SetDBufferB(float4(0, 0, 0, 1));
 		PackedPayload.SetDBufferC(float4(0, 0, 0, 1));
@@ -1973,7 +1817,7 @@ bool PathTracingKernel(inout FPathState PathState, int Bounce)
 				LightRay.TMax = LightSample.Distance;
 				ApplyRayBias(LightRay.Origin, Payload.HitT, SignedPositionBias * Payload.WorldGeoNormal);
 
-				float AvgRoughness = ApproximateCaustics ? GetAverageRoughness(Payload) : 0.0;
+				float AvgRoughness = ApproximateCaustics * GetAverageRoughness(Payload);
 
 				if (bIsVolumeSample)
 				{
@@ -2164,7 +2008,7 @@ bool PathTracingKernel(inout FPathState PathState, int Bounce)
 				const bool bCloudCastShadows = bIsVolumeSample ? true : CastsCloudShadow(LightId);
 				const float CloudFactor = bIsVolumeSample ? Payload.GetCloudFactor() : 1.0;
 				const uint MissShaderIndex = GetLightMissShaderIndex(LightId);
-				const float ShadowRayRoughness = ApproximateCaustics ? PathState.GetPathRoughness() : 0.0;
+                const float ShadowRayRoughness = ApproximateCaustics * PathState.GetPathRoughness();
 				LightContrib *= TraceTransparentVisibilityRay(LightRay, ShadowRayRoughness, PreviousPathRoughness, MissShaderIndex, bCastShadows, bCloudCastShadows, CloudFactor, PathState.RandSequence);
 				PathState.AccumulateRadiance(LightContrib);
 			}
diff --git a/Engine/Shaders/Private/PathTracing/PathTracingDebug.usf b/Engine/Shaders/Private/PathTracing/PathTracingDebug.usf
--- a/Engine/Shaders/Private/PathTracing/PathTracingDebug.usf
+++ b/Engine/Shaders/Private/PathTracing/PathTracingDebug.usf
@@ -36,6 +36,7 @@
 #include "../RayTracing/RayTracingCommon.ush"
 #include "../RayTracing/RayTracingDecalGrid.ush"
 #include "PathTracingCommon.ush"
+#include "PathTracingDecals.ush"
 #include "Material/PathTracingMaterialSampling.ush"
 #include "Utilities/PathTracingRandomSequence.ush"
 #include "Utilities/PathTracingRIS.ush"
@@ -512,6 +513,27 @@ FPackedPathTracingPayload TraceTransparentRay(inout FPathState PathState, int Bo
 		PathState.Ray.GetNativeDesc(),
 		OpaquePackedPayload);
 #endif
+#if PATH_TRACER_KERNEL_TYPE == PATH_TRACER_KERNEL_CAMERA
+    if (OpaquePackedPayload.IsDecalReceiver() && max(MeshDecalRoughnessCutoff, DecalRoughnessCutoff) > 0.0)
+    {
+        FDecalShaderPayload DecalPayload = InitDecalShaderPayload(PathState.Ray.Origin + OpaquePackedPayload.HitT * PathState.Ray.Direction);
+        // Evaluate mesh decals and decal actors
+        EvaluateDecals(
+            PathState.Ray.Direction,
+            PATHTRACER_MASK_CAMERA_OPAQUE | PATHTRACER_MASK_CAMERA_TRANSLUCENT,
+            OpaquePackedPayload.HitT,
+            0.0,
+            DecalPayload
+        );
+
+        // Handle d-buffer requests (will potentially re-trace the primary ray if needed)
+        PATH_TRACER_DECAL_PROCESS_DBUFFER(OpaquePackedPayload, 1.0, RandSample.y)
+
+        // Now that we know the real decal data, apply it
+        ApplyDecal(DecalPayload, OpaquePackedPayload, 0.0);
+    }
+#endif
+
 	if (OpaquePackedPayload.IsHit())
 	{
 		// TODO: If hit, also process decals, etc ...
@@ -1229,6 +1251,13 @@ bool DebugModeRequiresFakeShading()
         case PATH_TRACER_DEBUG_VIZ_TRIANGLE_ID:
         case PATH_TRACER_DEBUG_VIZ_INSTANCE_ID:
         case PATH_TRACER_DEBUG_VIZ_CLUSTER_ID:
+        case PATH_TRACER_DEBUG_VIZ_DECAL_COVERAGE:
+        case PATH_TRACER_DEBUG_VIZ_DECAL_COLOR:
+        case PATH_TRACER_DEBUG_VIZ_DECAL_NORMAL:
+        case PATH_TRACER_DEBUG_VIZ_DECAL_RMS:
+        case PATH_TRACER_DEBUG_VIZ_DECAL_RADIANCE:
+        case PATH_TRACER_DEBUG_VIZ_DECAL_RESPONSE_MASK:
+        case PATH_TRACER_DEBUG_VIZ_DECAL_USE_DBUFFER:
             return true;
     }
     return false;
@@ -1270,6 +1299,31 @@ FPathTracingResult PathTracingDebugCore(FRayDesc Ray, uint2 LaunchIndex)
 		Ray.GetNativeDesc(),
 		OpaquePackedPayload);
 
+    FDecalShaderPayload DecalPayload = InitDecalShaderPayload(Ray.Origin + OpaquePackedPayload.HitT * Ray.Direction);
+    switch (DebugMode)
+    {
+        // Only evaluate the decals if the visualization needs to actually show decal data
+        case PATH_TRACER_DEBUG_VIZ_DECAL_COVERAGE:
+        case PATH_TRACER_DEBUG_VIZ_DECAL_COLOR:
+        case PATH_TRACER_DEBUG_VIZ_DECAL_NORMAL:
+        case PATH_TRACER_DEBUG_VIZ_DECAL_RMS:
+        case PATH_TRACER_DEBUG_VIZ_DECAL_RADIANCE:
+        {
+            if (OpaquePackedPayload.IsDecalReceiver() && max(MeshDecalRoughnessCutoff, DecalRoughnessCutoff) > 0.0)
+            {
+                // Evaluate mesh decals and decal actors
+                EvaluateDecals(
+                    Ray.Direction,
+                    PATHTRACER_MASK_CAMERA_OPAQUE | PATHTRACER_MASK_CAMERA_TRANSLUCENT,
+                    OpaquePackedPayload.HitT,
+                    0.0,
+                    DecalPayload
+                );
+            }
+            break;
+        }
+    }
+
 	if (OpaquePackedPayload.IsHit())
 	{
 		// TODO: If hit, also process decals, etc ...
@@ -1368,7 +1422,7 @@ FPathTracingResult PathTracingDebugCore(FRayDesc Ray, uint2 LaunchIndex)
 			const float3 TranslatedWorldPos = Ray.Origin + PackedPayload.HitT * Ray.Direction;
 			const float3 WorldNormal = HitPayload.WorldNormal;
 
-			float3 Color = 0;
+			float3 Color = 0, ColorAdd = 0;
 			float ExtraAlpha = 0;
 			switch (DebugMode)
 			{
@@ -1479,6 +1533,24 @@ FPathTracingResult PathTracingDebugCore(FRayDesc Ray, uint2 LaunchIndex)
                     // radiance is populated by the material with the appropriate color
                     Color = HitPayload.Radiance;
                     break;
+                case PATH_TRACER_DEBUG_VIZ_DECAL_COVERAGE: Color = lerp(1.0, 0.18, float3(DecalPayload.GetBaseColor().a, DecalPayload.GetWorldNormal().a, DecalPayload.GetMetallicSpecularRoughness().a)); break;
+                case PATH_TRACER_DEBUG_VIZ_DECAL_COLOR:    Color = 0.18 * DecalPayload.GetBaseColor().a + DecalPayload.GetBaseColor().xyz; break;
+                case PATH_TRACER_DEBUG_VIZ_DECAL_NORMAL:   Color = 0.18 * DecalPayload.GetWorldNormal().a + (DecalPayload.GetWorldNormal().xyz * 0.5 + 0.5 * (1 - DecalPayload.GetWorldNormal().a)); break;
+                case PATH_TRACER_DEBUG_VIZ_DECAL_RMS:      Color = 0.18 * DecalPayload.GetMetallicSpecularRoughness().a + DecalPayload.GetMetallicSpecularRoughness().zxy; break; 
+                case PATH_TRACER_DEBUG_VIZ_DECAL_RADIANCE: Color = 0.18; ColorAdd = DecalPayload.GetEmissive(); break; // color added after shading
+                case PATH_TRACER_DEBUG_VIZ_DECAL_RESPONSE_MASK:
+                {
+                    Color = lerp(0.18, 0.5, float3(
+                        PackedPayload.HasDecalResponseColor(),
+                        PackedPayload.HasDecalResponseNormal(),
+                        PackedPayload.HasDecalResponseRoughness()));
+                    break;
+                }
+                case PATH_TRACER_DEBUG_VIZ_DECAL_USE_DBUFFER:
+                {
+                    Color = PackedPayload.UsesDBufferLookup() ? 0.5 : 0.18;
+                    break;
+                }
 				default:
 				{
 					break;
@@ -1505,8 +1577,8 @@ FPathTracingResult PathTracingDebugCore(FRayDesc Ray, uint2 LaunchIndex)
             }
             Color *= PathThroughput * HitPayload.BSDFOpacity;
             Result.DiffuseColor += Color;
-            Result.WorldNormal += Luminance(PathThroughput) * HitPayload.BSDFOpacity * HitPayload.WorldNormal;
-            Result.Radiance.rgb += Color * FakeShading;
+            Result.WorldNormal += Luminance(PathThroughput) * HitPayload.BSDFOpacity * WorldNormal;
+            Result.Radiance.rgb += Color * FakeShading + ColorAdd;
 			Result.Radiance.a += Luminance(PathThroughput) * ExtraAlpha;
 		}
 
diff --git a/Engine/Shaders/Private/PathTracing/PathTracingDecals.ush b/Engine/Shaders/Private/PathTracing/PathTracingDecals.ush
new file mode 100644
--- /dev/null
+++ b/Engine/Shaders/Private/PathTracing/PathTracingDecals.ush
@@ -0,0 +1,190 @@
+// Copyright Epic Games, Inc. All Rights Reserved.
+
+#pragma once
+
+#include "PathTracingCommon.ush"
+#include "../RayTracing/RayTracingCommon.ush"
+#include "../RayTracing/RayTracingDecalGrid.ush"
+
+RaytracingAccelerationStructure DecalTLAS;
+float DecalRoughnessCutoff;
+float MeshDecalRoughnessCutoff;
+float MeshDecalBias;
+uint EnableDBuffer;
+
+void EvaluateDecals(float3 RayDirection, uint InstanceInclusionMask, float HitT, float PathRoughness, inout FDecalShaderPayload DecalPayload)
+{
+    if (PathRoughness < MeshDecalRoughnessCutoff)
+    {
+        FRayDesc DecalRay = (FRayDesc)0;
+        // go a bit further than the current hit, to handle decals partially inside the model
+    #if DECAL_USE_REVERSE_CULLING
+        DecalRay.Direction = -RayDirection;
+    #else
+        DecalRay.Direction = RayDirection;
+    #endif
+        DecalRay.Origin = DecalPayload.TranslatedWorldPos - DecalRay.Direction * MeshDecalBias;
+        DecalRay.TMin = 0.0f;
+        DecalRay.TMax = HitT + MeshDecalBias;
+
+        const uint MissShaderIndex = 0; // TODO: do we need anything here?
+        for (;;)
+        {
+            DecalPayload.SetMiss();
+            TraceRay(
+                DecalTLAS,
+    #if DECAL_USE_REVERSE_CULLING
+                RAY_FLAG_CULL_FRONT_FACING_TRIANGLES,
+    #else
+                RAY_FLAG_CULL_BACK_FACING_TRIANGLES,
+    #endif
+                InstanceInclusionMask,
+                RAY_TRACING_SHADER_SLOT_MATERIAL,
+                RAY_TRACING_NUM_SHADER_SLOTS,
+                MissShaderIndex,
+                DecalRay.GetNativeDesc(),
+                DecalPayload);
+    #if NEED_TMIN_WORKAROUND // extra safety - discard hit if not conforming
+            if (DecalPayload.HitT <= DecalRay.TMin)
+            {
+                DecalPayload.HitT = -1.0; 
+            }
+    #endif
+            if (DecalPayload.IsMiss())
+            {
+                break;
+            }
+
+            // prepare next step around the loop
+            // retrace the exact same ray with TMin one ulp past the hit we just found
+            DecalRay.TMin = ComputeNewTMin(DecalRay.Origin, DecalRay.Direction, DecalPayload.HitT);
+        }
+    }
+#if PLATFORM_SUPPORTS_CALLABLE_SHADERS
+    if (PathRoughness < DecalRoughnessCutoff)
+    {
+        FDecalLoopCount DecalLoopCount = DecalGridLookup(DecalPayload.TranslatedWorldPos);
+        for (uint Index = 0, Num = DecalLoopCount.NumDecals; Index < Num; ++Index)
+        {
+            uint DecalId = GetDecalId(Index, DecalLoopCount);
+            CallShader(DecalId, DecalPayload);
+        }
+    }
+#endif
+}
+
+bool HasDecalData(FDecalShaderPayload DecalPayload)
+{
+	return min3(DecalPayload.GetBaseColor().a, DecalPayload.GetWorldNormal().a, DecalPayload.GetMetallicSpecularRoughness().a) < 1.0;
+}
+
+void ApplyDecal(FDecalShaderPayload DecalPayload, inout FPackedPathTracingPayload PackedPayload, float PathRoughness)
+{
+	{
+		float3 Radiance = PackedPayload.UnpackRadiance();
+		Radiance += DecalPayload.GetEmissive();
+		PackedPayload.PackRadiance(Radiance);
+	}
+
+	float4 DecalBaseColor = DecalPayload.GetBaseColor();
+	float4 DecalWorldNormal = DecalPayload.GetWorldNormal();
+	float4 DecalMetallicSpecularRoughness = DecalPayload.GetMetallicSpecularRoughness();
+
+	bool bHasColorResponse = PackedPayload.HasDecalResponseColor() && DecalBaseColor.a < 1.0;
+	bool bHasNormalResponse = PackedPayload.HasDecalResponseNormal() && DecalWorldNormal.a < 1.0;
+	bool bHasRoughnessResponse = PackedPayload.HasDecalResponseRoughness() && DecalMetallicSpecularRoughness.a < 1.0;
+
+	if (!bHasColorResponse && !bHasNormalResponse && !bHasRoughnessResponse)
+	{
+		// nothing to do
+		return;
+	}
+
+	float4 RoughnessData = PackedPayload.UnpackRoughnessAniso();
+	float Roughness = RoughnessData.x;
+	float3 DiffuseColor = PackedPayload.UnpackDiffuseColor();
+	float3 SpecularColor = PackedPayload.UnpackSpecularColor();
+	float Metallic = F0RGBToMetallic(SpecularColor);
+	float Specular = F0RGBToDielectricSpecular(SpecularColor);
+	float3 BaseColor = lerp(DiffuseColor, SpecularColor, Metallic);
+
+	if (bHasColorResponse)
+	{
+		BaseColor = BaseColor * DecalBaseColor.a + DecalBaseColor.rgb;
+	}
+
+	if (bHasNormalResponse)
+	{
+		// We normalize the normal to get smoother visual result (it helps to avoid having D_GGX explodes toward infinity, and matches ApplyDBufferData(...))
+		float3 WorldNormal = PackedPayload.UnpackWorldNormal();
+		WorldNormal = normalize(WorldNormal * DecalWorldNormal.a + DecalWorldNormal.rgb);
+		PackedPayload.PackWorldNormal(WorldNormal);
+	}
+
+	if (bHasRoughnessResponse)
+	{
+		Metallic  = Metallic  * DecalMetallicSpecularRoughness.a + DecalMetallicSpecularRoughness.x;
+		Specular  = Specular  * DecalMetallicSpecularRoughness.a + DecalMetallicSpecularRoughness.y;
+		Roughness = Roughness * DecalMetallicSpecularRoughness.a + DecalMetallicSpecularRoughness.z;
+		Roughness = max(Roughness, PathRoughness);
+	}
+
+	if (bHasColorResponse || bHasRoughnessResponse)
+	{
+		// the encode-decode is potentially lossy, so only do it if we actually need to
+		PackedPayload.PackDiffuseColor(BaseColor - BaseColor * Metallic);
+		PackedPayload.PackSpecularColor(lerp(DielectricSpecularToF0(Specular), BaseColor, Metallic));
+		PackedPayload.PackRoughnessAniso(float4(Roughness, RoughnessData.yzw));
+	}
+}
+
+
+
+#if USE_DBUFFER
+// This bit of code is not easily expressed as a method, so inject it using #define instead. This assumes the surrounding variables will be named a certain way
+// TODO: Is there a more elegant way to refactor this?
+
+// SlabRand must be same random value as original opaque trace to get a consistent result to the first call
+#define PATH_TRACER_DECAL_PROCESS_DBUFFER(OpaquePayload, ApproximateCaustics, SlabRand)     \
+    if (OpaquePayload.UsesDBufferLookup() && HasDecalData(DecalPayload))                    \
+    {                                                                                       \
+        if (EnableDBuffer)                                                                  \
+        {                                                                                   \
+            /* Retrace ray with DBuffer data */                                             \
+            OpaquePayload = InitPathTracingPayload(                                         \
+                PathState.GetFirstScatterType(),                                            \
+                ApproximateCaustics * PathState.GetPathRoughness()                          \
+            );                                                                              \
+            OpaquePayload.SetDBufferA(DecalPayload.GetBaseColor());                         \
+            OpaquePayload.SetDBufferB(DecalPayload.GetWorldNormal());                       \
+            OpaquePayload.SetDBufferC(DecalPayload.GetMetallicSpecularRoughness());         \
+            OpaquePayload.SetStochasticSlabRand(SlabRand);                                  \
+            OpaquePayload.SetFlag(PATH_TRACING_PAYLOAD_INPUT_FLAG_IGNORE_TRANSLUCENT);      \
+            TraceRay(                                                                       \
+                TLAS,                                                                       \
+                RayFlags,                                                                   \
+                InstanceInclusionMaskOpaque,                                                \
+                RAY_TRACING_SHADER_SLOT_MATERIAL,                                           \
+                RAY_TRACING_NUM_SHADER_SLOTS,                                               \
+                MissShaderIndex,                                                            \
+                PathState.Ray.GetNativeDesc(),                                              \
+                OpaquePayload);                                                             \
+        }                                                                                   \
+        else                                                                                \
+        {                                                                                   \
+            /* The user asked for a dbuffer lookup, but the double-trace is disabled. */    \
+            /* Set all the decal response flags so we at least get some kind of response */ \
+            /* even if it won't be exactly what the user had in mind. */                    \
+            OpaquePayload.SetFlag(                                                          \
+                                PATH_TRACING_PAYLOAD_OUTPUT_FLAG_DECAL_RESPONSE_COLOR |     \
+                                PATH_TRACING_PAYLOAD_OUTPUT_FLAG_DECAL_RESPONSE_NORMAL |    \
+                                PATH_TRACING_PAYLOAD_OUTPUT_FLAG_DECAL_RESPONSE_ROUGHNESS); \
+        }                                                                                           \
+    }
+
+#else
+
+// If DBuffer support is disabled, do nothing.
+#define PATH_TRACER_DECAL_PROCESS_DBUFFER(OpaquePayload, ApproximateCaustics, SlabRand) 
+
+#endif
diff --git a/Engine/Shaders/Private/RayTracing/RayTracingDecalGrid.ush b/Engine/Shaders/Private/RayTracing/RayTracingDecalGrid.ush
--- a/Engine/Shaders/Private/RayTracing/RayTracingDecalGrid.ush
+++ b/Engine/Shaders/Private/RayTracing/RayTracingDecalGrid.ush
@@ -1,5 +1,7 @@
 // Copyright Epic Games, Inc. All Rights Reserved.
 
+#pragma once
+
 #ifndef USE_RAY_TRACING_DECAL_GRID
 #define USE_RAY_TRACING_DECAL_GRID 0
 #endif
diff --git a/Engine/Shaders/Shared/PathTracingDefinitions.h b/Engine/Shaders/Shared/PathTracingDefinitions.h
--- a/Engine/Shaders/Shared/PathTracingDefinitions.h
+++ b/Engine/Shaders/Shared/PathTracingDefinitions.h
@@ -134,3 +134,10 @@
 #define PATH_TRACER_DEBUG_VIZ_TRIANGLE_ID               47
 #define PATH_TRACER_DEBUG_VIZ_INSTANCE_ID               48
 #define PATH_TRACER_DEBUG_VIZ_CLUSTER_ID                49
+#define PATH_TRACER_DEBUG_VIZ_DECAL_COVERAGE            50
+#define PATH_TRACER_DEBUG_VIZ_DECAL_COLOR               51
+#define PATH_TRACER_DEBUG_VIZ_DECAL_NORMAL              52
+#define PATH_TRACER_DEBUG_VIZ_DECAL_RMS                 53
+#define PATH_TRACER_DEBUG_VIZ_DECAL_RADIANCE            54
+#define PATH_TRACER_DEBUG_VIZ_DECAL_RESPONSE_MASK       55
+#define PATH_TRACER_DEBUG_VIZ_DECAL_USE_DBUFFER         56
diff --git a/Engine/Source/Runtime/Renderer/Private/PathTracing.cpp b/Engine/Source/Runtime/Renderer/Private/PathTracing.cpp
--- a/Engine/Source/Runtime/Renderer/Private/PathTracing.cpp
+++ b/Engine/Source/Runtime/Renderer/Private/PathTracing.cpp
@@ -193,6 +193,14 @@ struct FPathTracingDebugModeInfo
         REGISTER_VIZ_TYPE(TEXT("TriangleID")            , PATH_TRACER_DEBUG_VIZ_TRIANGLE_ID             );
         REGISTER_VIZ_TYPE(TEXT("InstanceID")            , PATH_TRACER_DEBUG_VIZ_INSTANCE_ID             );
         REGISTER_VIZ_TYPE(TEXT("ClusterID")             , PATH_TRACER_DEBUG_VIZ_CLUSTER_ID              );
+        REGISTER_VIZ_TYPE(TEXT("DecalCoverage")         , PATH_TRACER_DEBUG_VIZ_DECAL_COVERAGE          );
+        REGISTER_VIZ_TYPE(TEXT("DecalColor")            , PATH_TRACER_DEBUG_VIZ_DECAL_COLOR             );
+        REGISTER_VIZ_TYPE(TEXT("DecalNormal")           , PATH_TRACER_DEBUG_VIZ_DECAL_NORMAL            );
+        REGISTER_VIZ_TYPE(TEXT("DecalRoughMetalSpec")   , PATH_TRACER_DEBUG_VIZ_DECAL_RMS               );
+        REGISTER_VIZ_TYPE(TEXT("DecalRadiance")         , PATH_TRACER_DEBUG_VIZ_DECAL_RADIANCE          );
+        REGISTER_VIZ_TYPE(TEXT("DecalResponseMask")     , PATH_TRACER_DEBUG_VIZ_DECAL_RESPONSE_MASK     );
+        REGISTER_VIZ_TYPE(TEXT("DecalUseDBuffer")       , PATH_TRACER_DEBUG_VIZ_DECAL_USE_DBUFFER       );
+
 #undef REGISTER_VIZ_TYPE
 	}
 
@@ -392,10 +400,10 @@ TAutoConsoleVariable<float> CVarPathTracingMaxPathIntensity(
 	ECVF_RenderThreadSafe
 );
 
-TAutoConsoleVariable<int32> CVarPathTracingApproximateCaustics(
+TAutoConsoleVariable<bool> CVarPathTracingApproximateCaustics(
 	TEXT("r.PathTracing.ApproximateCaustics"),
-	1,
-	TEXT("When non-zero, the path tracer will approximate caustic paths to reduce noise. This reduces speckles and noise from low-roughness glass and metals."),
+	true,
+	TEXT("When true, the path tracer will approximate caustic paths to reduce noise. This reduces speckles and noise from low-roughness glass and metals."),
 	ECVF_RenderThreadSafe
 );
 
@@ -580,7 +588,7 @@ TAutoConsoleVariable<int32> CVarPathTracingLightGridAxis(
 TAutoConsoleVariable<bool> CVarPathTracingUseDBuffer(
 	TEXT("r.PathTracing.UseDBuffer"),
 	true,
-	TEXT("Whether to support DBuffer functionality."),
+	TEXT("When enabled, materials are allowed to read from d-buffer nodes. This carries a cost as the primary ray must be traced again to re-evaluate the material after the decal information has been gathered."),
 	ECVF_RenderThreadSafe
 );
 
@@ -752,7 +760,7 @@ BEGIN_SHADER_PARAMETER_STRUCT(FPathTracingData, )
 	SHADER_PARAMETER(float , SSSGuidingRatio)
 	SHADER_PARAMETER(uint32, MISMode)
 	SHADER_PARAMETER(uint32, VolumeMISMode)
-	SHADER_PARAMETER(uint32, ApproximateCaustics)
+	SHADER_PARAMETER(float, ApproximateCaustics)
 	SHADER_PARAMETER(uint32, SamplerType)
 	SHADER_PARAMETER(uint32, EnableDBuffer)
 	SHADER_PARAMETER(uint32, VolumeFlags)
@@ -1025,7 +1033,7 @@ static void PreparePathTracingData(const FScene* Scene, const FViewInfo& View, F
 		PathTracingData.MaxPathIntensity = PPV.PathTracingMaxPathIntensity;
 	}
 	PathTracingData.MaxPathIntensity = FFloat16(PathTracingData.MaxPathIntensity).GetClampedNonNegativeAndFinite().GetFloat(); // Clip to half precision
-	PathTracingData.ApproximateCaustics = CVarPathTracingApproximateCaustics.GetValueOnRenderThread();
+	PathTracingData.ApproximateCaustics = CVarPathTracingApproximateCaustics.GetValueOnRenderThread() ? 1.0f : 0.0f;
 	PathTracingData.SamplerType = CVarPathTracingSamplerType.GetValueOnRenderThread();
 	PathTracingData.FilterWidth = CVarPathTracingFilterWidth.GetValueOnRenderThread();
 	PathTracingData.CameraFocusDistance = 0;
@@ -1616,7 +1624,12 @@ class FPathTracingDebugRG : public FGlobalShader
 
 	
 		// scene decals
-		SHADER_PARAMETER_RDG_UNIFORM_BUFFER(FRayTracingDecals, DecalParameters)
+        SHADER_PARAMETER_RDG_BUFFER_SRV(RaytracingAccelerationStructure, DecalTLAS)
+        SHADER_PARAMETER_RDG_UNIFORM_BUFFER(FRayTracingDecals, DecalParameters)
+        SHADER_PARAMETER(float, DecalRoughnessCutoff)
+        SHADER_PARAMETER(float, MeshDecalRoughnessCutoff)
+        SHADER_PARAMETER(float, MeshDecalBias)
+        SHADER_PARAMETER(uint32, EnableDBuffer)
 
 		// scene uniform buffer
 		SHADER_PARAMETER_RDG_UNIFORM_BUFFER(FSceneUniformParameters, Scene)
@@ -4156,7 +4169,14 @@ void FDeferredShadingSceneRenderer::RenderPathTracing(
 					PassParameters->RWSceneRoughness = GraphBuilder.CreateUAV(RoughnessTexture);
 					PassParameters->RWSceneLinearDepth = GraphBuilder.CreateUAV(LinearDepthTexture);
 
-					PassParameters->DecalParameters = View.RayTracingDecalUniformBuffer;
+                    // decal parameters
+                    PassParameters->DecalParameters = View.RayTracingDecalUniformBuffer;
+                    PassParameters->DecalTLAS = Scene->RayTracingScene.GetLayerView(ERayTracingSceneLayer::Decals, View.GetRayTracingSceneViewHandle());
+                    PassParameters->DecalRoughnessCutoff = PathTracing::UsesDecals(ViewFamily) && View.bHasRayTracingDecals ? CVarPathTracingDecalRoughnessCutoff.GetValueOnRenderThread() : -1;
+                    PassParameters->MeshDecalRoughnessCutoff = PathTracing::UsesDecals(ViewFamily) && Scene->RayTracingScene.GetNumNativeInstances(ERayTracingSceneLayer::Decals, View.GetRayTracingSceneViewHandle()) > 0 ? CVarPathTracingMeshDecalRoughnessCutoff.GetValueOnRenderThread() : -1;
+                    PassParameters->MeshDecalBias = CVarPathTracingMeshDecalBias.GetValueOnRenderThread();
+					PassParameters->EnableDBuffer = CVarPathTracingUseDBuffer.GetValueOnRenderThread();
+	
 					PassParameters->Scene = View.GetSceneUniforms().GetBuffer(GraphBuilder);
 					PassParameters->NaniteRayTracing = Nanite::GRayTracingManager.GetUniformBuffer();
 
