Substrate - Toon - Another unified diffuse improvement.
Diffuse luminance with NoL is accumulated together with light weight W accounting for al lfalloff, masks and diffuse ggrey luminance. NoL in [0,1] is also accumulated weighted by W.
When applying unified diffuse ramp, NoL is evaluated from weighted mean. Eveyrthin is softer and more coherent. But 3 MRT are use now when this is enable for the complex path.

#tests PC editor toon
#rnx
#Synced-CL 52087402

#Synced-CL 52097113

diff --git a/Engine/Shaders/Private/DeferredLightPixelShaders.usf b/Engine/Shaders/Private/DeferredLightPixelShaders.usf
--- a/Engine/Shaders/Private/DeferredLightPixelShaders.usf
+++ b/Engine/Shaders/Private/DeferredLightPixelShaders.usf
@@ -264,7 +264,8 @@ void DeferredLightPixelMain(
 	float4 SVPos			: SV_POSITION,
 	out float4 OutColor		: SV_Target0
 #if SUBSTRATE_OUTPUT_TOON_DATA
-	, out float4 OutToonData : SV_Target1
+	, out float3 OutToonData0 : SV_Target1
+	, out float2 OutToonData1 : SV_Target2
 #endif
 #if SUBSTRATE_OPAQUE_ROUGH_REFRACTION_ENABLED
 	, out float3 OutOpaqueRoughRefractionSceneColor : SV_Target1
@@ -279,7 +280,8 @@ void DeferredLightPixelMain(
 	const float2 PixelPos = SVPos.xy;
 	OutColor = 0;
 #if SUBSTRATE_OUTPUT_TOON_DATA
-	OutToonData = float4(0.0f, 0.0f, 0.0f, 0.0f);
+	OutToonData0 = float3(0.0f, 0.0f, 0.0f);
+	OutToonData1 = float2(0.0f, 0.0f);
 #endif
 #if SUBSTRATE_OPAQUE_ROUGH_REFRACTION_ENABLED
 	OutOpaqueRoughRefractionSceneColor = 0;
@@ -382,7 +384,11 @@ void DeferredLightPixelMain(
 		#if SUBSTRATE_OUTPUT_TOON_DATA
 			// We do not need to test for toon BSDF as it is a single BSDF and always will work correctly when needed.
 			{
-				OutToonData = SubstrateLighting.ToonData;
+				const float3 ExposedDiffuseLuminance = SubstrateLighting.ToonData.DiffuseLuminance * GetExposure();
+				const float FinalWeight = SubstrateLighting.ToonData.Weight * GetExposure(); // expose the weight which have HDR light contribution in it.
+
+				OutToonData0 = ExposedDiffuseLuminance;
+				OutToonData1 = float2(FinalWeight, FinalWeight * SubstrateLighting.ToonData.ShadowNoL);
 			}
 		#endif
 		#if SUBSTRATE_OPAQUE_ROUGH_REFRACTION_ENABLED
@@ -399,7 +405,8 @@ void DeferredLightPixelMain(
 	const float4 CostWithColorCode = ConvertEstimatedCostToColor(EstimatedCost);
 	OutColor = CostWithColorCode + CostWithColorCode;
 	#if SUBSTRATE_OUTPUT_TOON_DATA
-	OutToonData = float4(0.0f, 0.0f, 0.0f, -1.0f);
+	OutToonData0 = float3(0.0f, 0.0f, 0.0f);
+	OutToonData1 = float2(0.0f, 0.0f);
 	#endif
 	#if SUBSTRATE_OPAQUE_ROUGH_REFRACTION_ENABLED
 	OutOpaqueRoughRefractionSceneColor = 0.0f;
@@ -451,9 +458,6 @@ void DeferredLightPixelMain(
 	OutSeparatedSubsurfaceDiffuse = EncodeSeparatedSubsurfaceDiffuse(OutColor);
 #endif
 
-#if SUBSTRATE_OUTPUT_TOON_DATA
-	OutToonData.rgb *= GetExposure();
-#endif
 #if SUBSTRATE_OPAQUE_ROUGH_REFRACTION_ENABLED 
 	// Idem
 	OutOpaqueRoughRefractionSceneColor *= GetExposure();
@@ -659,7 +663,8 @@ void DeferredLightPixelMain(
 
 #ifdef DeferredLightApplyToonDiffusePS
 
-Texture2D<float4> ToonDataTexture;
+Texture2D<float3> ToonDataTexture0;
+Texture2D<float2> ToonDataTexture1;
 
 void DeferredLightApplyToonDiffusePS(
 	in float2 ScreenUV : TEXCOORD0,
@@ -705,14 +710,20 @@ void DeferredLightApplyToonDiffusePS(
 		UnpackToonCustomData(PackRGBA8(GBuffer.CustomData), ToonPatternUV, ToonProfileId);
 	#endif
 
-		float4 ToonData = ToonDataTexture[SVPos.xy];
+		const float3 ToonExposedDiffuse	= ToonDataTexture0[SVPos.xy];
+		const float2 ToonWeightNoL		= ToonDataTexture1[SVPos.xy];
+		const float ToonWeight			= ToonWeightNoL.x;
+		const float ToonWeightedNoL		= ToonWeightNoL.y;
 
 		FToonProfileEvaluateSettings ToonEvalSettings = GetToonProfileEvaluateSettings(ToonProfileId);
 		float ToonProfileDiffuseDither = EvaluateToonProfileDiffuseDither(ToonProfileId, ToonPatternUV, ToonEvalSettings);
 
-		const float NoL = ToonData.a * 2.0 - 1.0;
-		const float3 ToonDiffuse = ToonData.rgb * EvaluateToonProfileDiffuseRamp(ToonProfileId, NoL + ToonProfileDiffuseDither);
-		OutColor = float4(ToonDiffuse, Luminance(ToonDiffuse));
+		// Recoverd NoL from  (NoL * Weight) / Weight, to apply the diffuse profile, accounting for the weight many light contribution at once.
+		const float NoL = (ToonWeightedNoL / max(0.0001f, ToonWeight)) * 2.0 - 1.0;
+
+		const float3 ToonExposedDiffuseWithRamp = ToonExposedDiffuse * EvaluateToonProfileDiffuseRamp(ToonProfileId, NoL + ToonProfileDiffuseDither);
+
+		OutColor = float4(ToonExposedDiffuseWithRamp, Luminance(ToonExposedDiffuseWithRamp));
 	}
 	else
 	{
diff --git a/Engine/Shaders/Private/LightAccumulator.ush b/Engine/Shaders/Private/LightAccumulator.ush
--- a/Engine/Shaders/Private/LightAccumulator.ush
+++ b/Engine/Shaders/Private/LightAccumulator.ush
@@ -178,10 +178,17 @@ FDeferredLightingSplit LightAccumulator_GetResultSplit(FLightAccumulator In)
 	return Ret;
 }
 
+struct FDeferredLightingToonData
+{
+	float3 DiffuseLuminance;
+	float Weight;
+	float ShadowNoL;
+};
+
 struct FSubstrateDeferredLighting
 {
 	float4 SceneColor;						// Untouched scene luminance, alpha channel may contain scatterable light luminance.
-	float4 ToonData;
+	FDeferredLightingToonData ToonData;
 #if SUBSTRATE_OPAQUE_ROUGH_REFRACTION_ENABLED
 	float3 OpaqueRoughRefractionSceneColor;	// Luminance that is going to be blurier by the top layer roughness
 	float3 SubSurfaceSceneColor;			// Scene color that should go through the sub surface post process. There is not alpha: not needed because it is separated from specular already.
@@ -198,7 +205,6 @@ struct FSubstrateDeferredLighting
 FSubstrateDeferredLighting GetInitialisedSubstrateDeferredLighting()
 {
 	FSubstrateDeferredLighting Result = (FSubstrateDeferredLighting)0;
-	Result.ToonData = float4(0.0f, 0.0f, 0.0f, -1.0f);
 	return Result;
 }
 
diff --git a/Engine/Shaders/Private/Substrate/SubstrateDeferredLighting.ush b/Engine/Shaders/Private/Substrate/SubstrateDeferredLighting.ush
--- a/Engine/Shaders/Private/Substrate/SubstrateDeferredLighting.ush
+++ b/Engine/Shaders/Private/Substrate/SubstrateDeferredLighting.ush
@@ -204,11 +204,12 @@ FSubstrateDeferredLighting SubstrateDeferredLighting(
 			if (bIsToon)
 			{
 				const float3 ToonDiffuseLuminance = DiffuseLuminance * CommonMultiplier;
-				if (any(ToonDiffuseLuminance > 0))
-				{
-					SubstrateLighting.ToonData = float4(ToonDiffuseLuminance, BSDFEvaluate.ToonData) ;
-					DiffuseLuminance = 0.0f;
-				}
+
+				SubstrateLighting.ToonData.DiffuseLuminance = ToonDiffuseLuminance;											// The diffuse illuminance to apply the diffuse ramp onto later.
+				SubstrateLighting.ToonData.Weight			= BSDFEvaluate.ToonData.Weight * Luminance(CommonMultiplier);	// Accumulate more importance weight for current light (e.g. mask, intensity)
+				SubstrateLighting.ToonData.ShadowNoL		= BSDFEvaluate.ToonData.ShadowNoL;								// Forward NoL.
+
+				DiffuseLuminance = 0.0f;
 			}
 		#endif
 
diff --git a/Engine/Shaders/Private/Substrate/SubstrateEvaluation.ush b/Engine/Shaders/Private/Substrate/SubstrateEvaluation.ush
--- a/Engine/Shaders/Private/Substrate/SubstrateEvaluation.ush
+++ b/Engine/Shaders/Private/Substrate/SubstrateEvaluation.ush
@@ -391,6 +391,12 @@ float EvaluateGlintEnv(FSubstrateBSDFContext BSDFContext, float SafeRoughness, i
 ///////////////////////////////////////////////////////////////////////////////////////////////////
 // BSDF evaluation (punctual/area)
 
+struct FSubstrateToonData
+{
+	float Weight;
+	float ShadowNoL;
+};
+
 struct FSubstrateEvaluateResult
 {
 	float3 IntegratedDiffuseValue;
@@ -421,8 +427,7 @@ struct FSubstrateEvaluateResult
 	bool   bPostProcessSubsurface;		// True if we need to separate the subsurface light contribution for the screen space diffusion process.
 	bool   bApplyProjectedSolidAngle;	// True if the Saturate(NoL) factor should be applied or not over the value of the BSDF
 
-	// This is NoL maybe weighted by the shadow for toon shading model only.
-	float  ToonData;
+	FSubstrateToonData  ToonData;
 };
 
 // PUNCTUAL corresponds to an evaluation (from in & out directions).
@@ -1359,8 +1364,6 @@ FSubstrateEvaluateResult SubstrateEvaluateBSDFCommon(FSubstrateBSDFContext BSDFC
 			const bool bApplyDiffuseRamp	= SUBSTRATE_LIGHTPASS_APPLIES_TOON_DIFFUSE_RAMP;
 			const bool bApplyDiffuseShadow	= SUBSTRATE_LIGHTPASS_APPLIES_TOON_DIFFUSE_RAMP || (!SUBSTRATE_LIGHTPASS_APPLIES_TOON_DIFFUSE_RAMP && !ToonEvalSettings.bUnifiedDiffuseIncludeShadow);
 
-			Sample.ToonData					= (BSDFContext.Context.NoL * 0.5 + 0.5) * AreaLight.Falloff * (bApplyDiffuseShadow ? 1.0f : TransmittanceShadow);
-
 			// Diffuse
 			const float3 DiffuseRamp		= bApplyDiffuseRamp ? EvaluateToonProfileDiffuseRamp(ToonProfileId, BSDFContext.Context.NoL + ToonProfileDiffuseDither) : 1.0f;
 			Sample.DiffuseColor				= DiffuseColor * DiffuseRamp;
@@ -1371,6 +1374,9 @@ FSubstrateEvaluateResult SubstrateEvaluateBSDFCommon(FSubstrateBSDFContext BSDFC
 			Sample.SpecularPathValue = SpecularRamp;
 			Sample.IntegratedSpecularValue += ShadowTerms.SurfaceShadow * AreaLight.NoL * AreaLight.Falloff * AreaLight.FalloffColor * Sample.SpecularPathValue;
 
+			Sample.ToonData.Weight			= (bApplyDiffuseShadow ? 1.0f : TransmittanceShadow) * AreaLight.Falloff;						// Accumulate light weight. Albedo should not drive light importance so we do not need to include Lambert diffuse here.
+			Sample.ToonData.ShadowNoL		= (bApplyDiffuseShadow ? 1.0f : TransmittanceShadow) * (BSDFContext.Context.NoL * 0.5 + 0.5);	// NoL with or without shadow, remapped to [0,1] range.
+
 			// Others
 			Sample.DiffusePDF = BSDFContext.SatNoL / PI;
 			Sample.bPostProcessSubsurface = false;
diff --git a/Engine/Source/Runtime/Renderer/Private/LightRendering.cpp b/Engine/Source/Runtime/Renderer/Private/LightRendering.cpp
--- a/Engine/Source/Runtime/Renderer/Private/LightRendering.cpp
+++ b/Engine/Source/Runtime/Renderer/Private/LightRendering.cpp
@@ -1241,7 +1241,8 @@ class FDeferredLightApplyToonDiffusePS : public FGlobalShader
 		SHADER_PARAMETER_STRUCT_REF(FViewUniformShaderParameters, View)
 		SHADER_PARAMETER_RDG_UNIFORM_BUFFER(FSceneTextureUniformParameters, SceneTextures)
 		SHADER_PARAMETER_RDG_UNIFORM_BUFFER(FSubstrateGlobalUniformParameters, Substrate)
-		SHADER_PARAMETER_RDG_TEXTURE(Texture2D<float4>, ToonDataTexture)
+		SHADER_PARAMETER_RDG_TEXTURE(Texture2D<float4>, ToonDataTexture0)
+		SHADER_PARAMETER_RDG_TEXTURE(Texture2D<float4>, ToonDataTexture1)
 		RENDER_TARGET_BINDING_SLOTS()
 	END_SHADER_PARAMETER_STRUCT()
 
@@ -2515,7 +2516,8 @@ void FDeferredShadingSceneRenderer::RenderLights(
 				ToonPassParameters->Substrate = Substrate::BindSubstrateGlobalUniformParameters(View);
 
 				//	const FSubstrateSceneData& SubstrateSceneData = *View.GetSceneTextures().SubstrateSceneData;
-				ToonPassParameters->ToonDataTexture = SubstrateSceneData.SubstrateSceneToonData;
+				ToonPassParameters->ToonDataTexture0 = SubstrateSceneData.SubstrateSceneToonData0;
+				ToonPassParameters->ToonDataTexture1 = SubstrateSceneData.SubstrateSceneToonData1;
 
 				ToonPassParameters->RenderTargets[0] = FRenderTargetBinding(SceneTextures.Color.Target, ERenderTargetLoadAction::ELoad);
 				if (SceneTextures.Depth.Target)
@@ -2865,7 +2867,8 @@ uint32 SetupLightGraphicsPSOState(
 	{
 		GraphicsPSOInit.BlendState = TStaticBlendState<
 			CW_RGBA, BO_Add, BF_One, BF_One, BO_Add, BF_One, BF_One,
-			CW_RGBA, BO_Add, BF_One, BF_One, BO_Max, BF_One, BF_One>::GetRHI();
+			CW_RGB,  BO_Add, BF_One, BF_One, BO_Add, BF_One, BF_One,
+			CW_RG,   BO_Add, BF_One, BF_One, BO_Add, BF_One, BF_One>::GetRHI();
 
 		// As of today, OpaqueRoughRefraction cannot be used together with experimental toon features. This is reflected in the implementation of IsOpaqueRoughRefractionEnabled().
 		check(!Substrate::IsOpaqueRoughRefractionEnabled(ShaderPlatform));
@@ -3268,7 +3271,8 @@ static void RenderLight(
 							*PassParametersForToon = *PassParameters;
 							// And setup the needed toon render target
 							const FSubstrateSceneData& SubstrateSceneData = *View.GetSceneTextures().SubstrateSceneData;
-							PassParametersForToon->PS.RenderTargets[1] = FRenderTargetBinding(SubstrateSceneData.SubstrateSceneToonData, ERenderTargetLoadAction::ELoad);
+							PassParametersForToon->PS.RenderTargets[1] = FRenderTargetBinding(SubstrateSceneData.SubstrateSceneToonData0, ERenderTargetLoadAction::ELoad);
+							PassParametersForToon->PS.RenderTargets[2] = FRenderTargetBinding(SubstrateSceneData.SubstrateSceneToonData1, ERenderTargetLoadAction::ELoad);
 							// As of today, OpaqueRoughRefraction cannot be used together with experimental toon features. This is reflected in the implementation of IsOpaqueRoughRefractionEnabled().
 							check(!Substrate::IsOpaqueRoughRefractionEnabled(View.GetShaderPlatform()));
 
diff --git a/Engine/Source/Runtime/Renderer/Private/Substrate/Substrate.cpp b/Engine/Source/Runtime/Renderer/Private/Substrate/Substrate.cpp
--- a/Engine/Source/Runtime/Renderer/Private/Substrate/Substrate.cpp
+++ b/Engine/Source/Runtime/Renderer/Private/Substrate/Substrate.cpp
@@ -609,12 +609,20 @@ void InitialiseSubstrateFrameSceneData(FRDGBuilder& GraphBuilder, FScene& Scene,
 			}
 		}
 
-		Out.SubstrateSceneToonData = nullptr;
+		Out.SubstrateSceneToonData0 = nullptr;
+		Out.SubstrateSceneToonData1 = nullptr;
 		if (IsExperimentalToonUnifiedDiffuseEnabled())
 		{
-			Out.SubstrateSceneToonData = GraphBuilder.CreateTexture(
-				FRDGTextureDesc::Create2D(SceneTextureExtent, PF_FloatRGBA, FClearValueBinding(FLinearColor(0.0f, 0.0f, 0.0f, 0.0f)), TexCreate_ShaderResource | TexCreate_UAV | TexCreate_RenderTargetable), TEXT("Substrate.SubstrateSceneToonData"));
-			AddClearRenderTargetPass(GraphBuilder, Out.SubstrateSceneToonData, Out.SubstrateSceneToonData->Desc.ClearValue.GetClearColor());
+			// Prefer 32bits format but fallback more common format if necessary. However, those format must be HDR.
+			EPixelFormat ToonDataFormat0 = GPixelFormats[PF_FloatR11G11B10].Supported	? PF_FloatR11G11B10	: PF_FloatRGBA;
+			EPixelFormat ToonDataFormat1 = GPixelFormats[PF_G16R16F].Supported			? PF_G16R16F		: PF_FloatRGBA;
+
+			Out.SubstrateSceneToonData0 = GraphBuilder.CreateTexture(
+				FRDGTextureDesc::Create2D(SceneTextureExtent, ToonDataFormat0, FClearValueBinding(FLinearColor(0.0f, 0.0f, 0.0f, 0.0f)), TexCreate_ShaderResource | TexCreate_UAV | TexCreate_RenderTargetable), TEXT("Substrate.SubstrateSceneToonData0"));
+			AddClearRenderTargetPass(GraphBuilder, Out.SubstrateSceneToonData0, Out.SubstrateSceneToonData0->Desc.ClearValue.GetClearColor());
+			Out.SubstrateSceneToonData1 = GraphBuilder.CreateTexture(
+				FRDGTextureDesc::Create2D(SceneTextureExtent, ToonDataFormat1, FClearValueBinding(FLinearColor(0.0f, 0.0f, 0.0f, 0.0f)), TexCreate_ShaderResource | TexCreate_UAV | TexCreate_RenderTargetable), TEXT("Substrate.SubstrateSceneToonData1"));
+			AddClearRenderTargetPass(GraphBuilder, Out.SubstrateSceneToonData1, Out.SubstrateSceneToonData1->Desc.ClearValue.GetClearColor());
 		}
 
 		// Closure offsets
diff --git a/Engine/Source/Runtime/Renderer/Private/Substrate/Substrate.h b/Engine/Source/Runtime/Renderer/Private/Substrate/Substrate.h
--- a/Engine/Source/Runtime/Renderer/Private/Substrate/Substrate.h
+++ b/Engine/Source/Runtime/Renderer/Private/Substrate/Substrate.h
@@ -143,7 +143,8 @@ struct FSubstrateSceneData
 	FRDGTextureRef SeparatedOpaqueRoughRefractionSceneColor = nullptr;
 
 	// Used for Toon unified diffuse.
-	FRDGTextureRef SubstrateSceneToonData = nullptr;
+	FRDGTextureRef SubstrateSceneToonData0 = nullptr;
+	FRDGTextureRef SubstrateSceneToonData1 = nullptr;
 
 	//Public facing minimal uniform data.
 	TRDGUniformBufferRef<FSubstratePublicGlobalUniformParameters> SubstratePublicGlobalUniformParameters{};
