Substrate - Toon - Made the "diffuse include ramp" control genral and not limited to unified diffuse.

#rnx
#Synced-CL 52097327

diff --git a/Engine/Shaders/Private/DeferredLightPixelShaders.usf b/Engine/Shaders/Private/DeferredLightPixelShaders.usf
--- a/Engine/Shaders/Private/DeferredLightPixelShaders.usf
+++ b/Engine/Shaders/Private/DeferredLightPixelShaders.usf
@@ -719,7 +719,7 @@ void DeferredLightApplyToonDiffusePS(
 		float ToonProfileDiffuseDither = EvaluateToonProfileDiffuseDither(ToonProfileId, ToonPatternUV, ToonEvalSettings);
 
 		// Recoverd NoL from  (NoL * Weight) / Weight, to apply the diffuse profile, accounting for the weight many light contribution at once.
-		const float NoL = (ToonWeightedNoL / max(0.0001f, ToonWeight)) * 2.0 - 1.0;
+		const float NoL = (ToonWeightedNoL / max(0.0001f, ToonWeight));
 
 		const float3 ToonExposedDiffuseWithRamp = ToonExposedDiffuse * EvaluateToonProfileDiffuseRamp(ToonProfileId, NoL + ToonProfileDiffuseDither);
 
diff --git a/Engine/Shaders/Private/Substrate/SubstrateEvaluation.ush b/Engine/Shaders/Private/Substrate/SubstrateEvaluation.ush
--- a/Engine/Shaders/Private/Substrate/SubstrateEvaluation.ush
+++ b/Engine/Shaders/Private/Substrate/SubstrateEvaluation.ush
@@ -1361,22 +1361,28 @@ FSubstrateEvaluateResult SubstrateEvaluateBSDFCommon(FSubstrateBSDFContext BSDFC
 				// TODO H_PDF and PDF
 			}
 			
-			const bool bApplyDiffuseRamp	= SUBSTRATE_LIGHTPASS_APPLIES_TOON_DIFFUSE_RAMP;
-			const bool bApplyDiffuseShadow	= SUBSTRATE_LIGHTPASS_APPLIES_TOON_DIFFUSE_RAMP || (!SUBSTRATE_LIGHTPASS_APPLIES_TOON_DIFFUSE_RAMP && !ToonEvalSettings.bUnifiedDiffuseIncludeShadow);
+			const bool bThisPassAppliesDiffuseRamp	= SUBSTRATE_LIGHTPASS_APPLIES_TOON_DIFFUSE_RAMP;
 
 			// Diffuse
-			const float3 DiffuseRamp		= bApplyDiffuseRamp ? EvaluateToonProfileDiffuseRamp(ToonProfileId, BSDFContext.Context.NoL + ToonProfileDiffuseDither) : 1.0f;
+			float RampDiffuseNoL = BSDFContext.Context.NoL * 0.5 + 0.5;
+			if (ToonEvalSettings.bDiffuseRampIncludeShadow)
+			{
+				RampDiffuseNoL = RampDiffuseNoL * TransmittanceShadow;
+			}
+			const float3 DiffuseRamp		= bThisPassAppliesDiffuseRamp ? EvaluateToonProfileDiffuseRamp(ToonProfileId, RampDiffuseNoL + ToonProfileDiffuseDither) : 1.0f;
 			Sample.DiffuseColor				= DiffuseColor * DiffuseRamp;
 			Sample.DiffusePathValue			= Diffuse_Lambert(Sample.DiffuseColor);
-			Sample.IntegratedDiffuseValue	= (bApplyDiffuseShadow ? TransmittanceShadow : 1.0f) * AreaLight.Falloff * AreaLight.FalloffColor * Sample.DiffusePathValue; //  * AreaLight.NoL for smoother edges? But then transmittance is blocked. So this is left in the hand of the artist.
+			Sample.IntegratedDiffuseValue	= (ToonEvalSettings.bDiffuseRampIncludeShadow ? 1.0f : TransmittanceShadow) * AreaLight.Falloff * AreaLight.FalloffColor * Sample.DiffusePathValue; //  * AreaLight.NoL for smoother edges? But then transmittance is blocked. So this is left in the hand of the artist.
 
+			// Toon unified diffse data.
+			Sample.ToonData.Weight			= TransmittanceShadow * AreaLight.Falloff;	// Accumulate light weight. Albedo should not drive light importance so we do not need to include Lambert diffuse here.
+			Sample.ToonData.ShadowNoL		= RampDiffuseNoL;							// NoL with or without shadow, remapped to [0,1] range.
+
+			// Specular
 			const float3 SpecularRamp = SpecularColor * EvaluateToonProfileSpecularRamp(ToonProfileId, saturate(dot(Specular, 1.0f/3.0f) + ToonProfileSpecularDither));
 			Sample.SpecularPathValue = SpecularRamp;
 			Sample.IntegratedSpecularValue += ShadowTerms.SurfaceShadow * AreaLight.NoL * AreaLight.Falloff * AreaLight.FalloffColor * Sample.SpecularPathValue;
 
-			Sample.ToonData.Weight			= (bApplyDiffuseShadow ? 1.0f : TransmittanceShadow) * AreaLight.Falloff;						// Accumulate light weight. Albedo should not drive light importance so we do not need to include Lambert diffuse here.
-			Sample.ToonData.ShadowNoL		= (bApplyDiffuseShadow ? 1.0f : TransmittanceShadow) * (BSDFContext.Context.NoL * 0.5 + 0.5);	// NoL with or without shadow, remapped to [0,1] range.
-
 			// Others
 			Sample.DiffusePDF = BSDFContext.SatNoL / PI;
 			Sample.bPostProcessSubsurface = false;
diff --git a/Engine/Shaders/Private/ToonProfile.usf b/Engine/Shaders/Private/ToonProfile.usf
--- a/Engine/Shaders/Private/ToonProfile.usf
+++ b/Engine/Shaders/Private/ToonProfile.usf
@@ -22,7 +22,7 @@ float ShadowHatchingPatternStrength;
 
 float DiffuseIndirectScale;
 float SpecularIndirectScale;
-float bUnifiedDiffuseIncludeShadow;
+float bDiffuseRampIncludeShadow;
 float SpecularIndirectRampRepetition;
 
 RWTexture2DArray<float4>	ToonProfileTexture;
@@ -76,7 +76,7 @@ void ToonProfileUpdateAtlasTileCS(uint3 DispatchThreadId : SV_DispatchThreadID)
 			}
 			else if (PixelPos.x == TOON_PROFILE_EVAL_PARAMS2_X)
 			{
-				ToonProfileTexture[uint3(PixelPos, ToonProfileIndex)] = float4(bUnifiedDiffuseIncludeShadow, 0.0f, 0.0f, 0.0f);
+				ToonProfileTexture[uint3(PixelPos, ToonProfileIndex)] = float4(bDiffuseRampIncludeShadow, 0.0f, 0.0f, 0.0f);
 			}
 			else if (PixelPos.x == TOON_PROFILE_SHADOW_PARAMS_X)
 			{
diff --git a/Engine/Shaders/Private/ToonProfileCommon.ush b/Engine/Shaders/Private/ToonProfileCommon.ush
--- a/Engine/Shaders/Private/ToonProfileCommon.ush
+++ b/Engine/Shaders/Private/ToonProfileCommon.ush
@@ -16,9 +16,10 @@ float3 EvaluateToonProfileRamp(in float ToonProfileId, in float NoL, in float TO
 	return View.ToonProfileTexture.SampleLevel(View.ToonProfileSampler, ToonProfileCoord, 0).xyz;
 }
 
+// NoL input should be in [0,1], remapped from [-1, 1].
 float3 EvaluateToonProfileDiffuseRamp(in float ToonProfileId, in float NoL)
 {
-	return EvaluateToonProfileRamp(ToonProfileId, NoL * 0.5f + 0.5f, TOON_PROFILE_DIFFUSE_RAMP_Y);
+	return EvaluateToonProfileRamp(ToonProfileId, saturate(NoL), TOON_PROFILE_DIFFUSE_RAMP_Y);
 }
 
 float3 EvaluateToonProfileSpecularRamp(in float ToonProfileId, in float Specular)
@@ -44,7 +45,7 @@ struct FToonProfileEvaluateSettings
 	float DiffuseRampOffsetSize;
 	float SpecularRampOffsetStrength;
 	float SpecularRampOffsetSize;
-	float bUnifiedDiffuseIncludeShadow;
+	float bDiffuseRampIncludeShadow;
 };
 
 FToonProfileEvaluateSettings GetToonProfileEvaluateSettings(in float ToonProfileId)
@@ -67,7 +68,7 @@ FToonProfileEvaluateSettings GetToonProfileEvaluateSettings(in float ToonProfile
 		ToonProfileCoord.yz);
 	float4 Data2 = View.ToonProfileTexture.SampleLevel(View.ToonProfileSampler, ToonProfileCoord2, 0).xyzw;
 
-	Ret.bUnifiedDiffuseIncludeShadow = Data2.x;
+	Ret.bDiffuseRampIncludeShadow = Data2.x;
 
 	return Ret;
 }
diff --git a/Engine/Source/Runtime/Engine/Classes/Engine/ToonProfile.h b/Engine/Source/Runtime/Engine/Classes/Engine/ToonProfile.h
--- a/Engine/Source/Runtime/Engine/Classes/Engine/ToonProfile.h
+++ b/Engine/Source/Runtime/Engine/Classes/Engine/ToonProfile.h
@@ -132,10 +132,10 @@ struct FToonProfileStruct
 	float SpecularIndirectRampRepetition;
 
 	/**
-	 * Experimental. Include shadows as part of the unified diffuse ramp evaluation. Requires r.Substrate.Experimental.ToonUnifiedDiffuse=1.
+	 * Experimental. Include shadows as part of the unified diffuse ramp evaluation.
 	 */
-	UPROPERTY(Category = "Experimental", EditAnywhere, meta = (DisplayName = "Unified Diffuse Include Shadow"))
-	uint8 bUnifiedDiffuseIncludeShadow : 1;
+	UPROPERTY(Category = "Experimental", EditAnywhere, meta = (DisplayName = "Diffuse Ramp Includes Shadow"))
+	uint8 bDiffuseRampIncludeShadow : 1;
 
 	FToonProfileStruct();
 
diff --git a/Engine/Source/Runtime/Engine/Private/Rendering/ToonProfile.cpp b/Engine/Source/Runtime/Engine/Private/Rendering/ToonProfile.cpp
--- a/Engine/Source/Runtime/Engine/Private/Rendering/ToonProfile.cpp
+++ b/Engine/Source/Runtime/Engine/Private/Rendering/ToonProfile.cpp
@@ -82,7 +82,7 @@ public:
 		SHADER_PARAMETER(float, ShadowHatchingPatternStrength)
 		SHADER_PARAMETER(float, DiffuseIndirectScale)
 		SHADER_PARAMETER(float, SpecularIndirectScale)
-		SHADER_PARAMETER(float, bUnifiedDiffuseIncludeShadow)
+		SHADER_PARAMETER(float, bDiffuseRampIncludeShadow)
 		SHADER_PARAMETER(float, SpecularIndirectRampRepetition)
 
 		SHADER_PARAMETER(float, ToonProfileShadowExtinctionCoefficientInInvCentimeter)
@@ -416,7 +416,7 @@ IPooledRenderTarget* FToonProfileTextureManager::GetAtlasTexture(FRDGBuilder& Gr
 
 			PassParameters->DiffuseIndirectScale = Data.DiffuseIndirectScale;
 			PassParameters->SpecularIndirectScale = Data.SpecularIndirectScale;
-			PassParameters->bUnifiedDiffuseIncludeShadow = Data.bUnifiedDiffuseIncludeShadow ? 1.0f : 0.0f;
+			PassParameters->bDiffuseRampIncludeShadow = Data.bDiffuseRampIncludeShadow ? 1.0f : 0.0f;
 			PassParameters->SpecularIndirectRampRepetition = Data.SpecularIndirectRampRepetition;
 
 			auto CreateBufferFromColorCurve = [&](const FRuntimeCurveLinearColor& ColorCurve)
@@ -580,7 +580,7 @@ FToonProfileStruct::FToonProfileStruct()
 	SpecularIndirectScale = 1.0f;
 	SpecularIndirectRampRepetition = 0.0f;
 
-	bUnifiedDiffuseIncludeShadow = 0;
+	bDiffuseRampIncludeShadow = 0;
 
 	AddPointToFloatCurve(ShadowHatchingPatternDistributionRamp, 0.00, 0.0);
 	AddPointToFloatCurve(ShadowHatchingPatternDistributionRamp, 1.00, 1.0);
@@ -651,12 +651,6 @@ bool UToonProfile::CanEditChange(const FProperty* InProperty) const
 		static IConsoleVariable* CVar = IConsoleManager::Get().FindConsoleVariable(TEXT("r.Substrate.Experimental.ToonReflectionQuantizationEnabled"));
 		return CVar && CVar->GetInt() > 0;
 	}
-	
-	if (PropertyName == GET_MEMBER_NAME_STRING_CHECKED(FToonProfileStruct, bUnifiedDiffuseIncludeShadow))
-	{
-		static IConsoleVariable* CVar = IConsoleManager::Get().FindConsoleVariable(TEXT("r.Substrate.Experimental.ToonUnifiedDiffuse"));
-		return CVar && CVar->GetInt() > 0;
-	}
 
 	return true;
 }
