Fix an issue where fog could be applied incorrectly to volumetric clouds in splitscreen or stereo views, causing artifacts against the sky near the edges of objects.

Each view has its own volumetric textures and their contents are always top-left/(0,0) aligned, even in splitscreen where some views have a ViewRectMin != (0,0). As such, some portions of ComposeVolumetricRTOverScenePS need to use a ScreenUV which has been shifted to the top-left (e.g. U ranging from 0.0 to 0.5 for a splitscreen view on the right, rather than 0.5 to 1.0) in order to sample from the volumetric textures. This fixes an area where we weren't doing so, and updates comments/variable names to make this more clear.

#Synced-CL 52153539

diff --git a/Engine/Shaders/Private/VolumetricRenderTarget.usf b/Engine/Shaders/Private/VolumetricRenderTarget.usf
--- a/Engine/Shaders/Private/VolumetricRenderTarget.usf
+++ b/Engine/Shaders/Private/VolumetricRenderTarget.usf
@@ -759,22 +759,22 @@ float UvOffsetSampleAcceptanceWeight;
 float MinimumDistanceKmToDisableDisoclusion;
 float2 FullResolutionToVolumetricBufferResolutionScale;
 
-float4 ApplyFogToCloudFromUVsAndDepth(float2 UVWithOffset, float DepthKilometer, float4 CloudRGBT)
+float4 ApplyFogToCloudFromUVsAndDepth(float2 UV, float DepthKilometer, float4 CloudRGBT)
 {
 #if APPLY_FOG_LATE
-	float2 SVPosWithOffset = (UVWithOffset / (FullResolutionToVolumetricBufferResolutionScale.x *View.BufferSizeAndInvSize.xy * VolumetricTextureSizeAndInvSize.zw)) * View.BufferSizeAndInvSize.xy + View.ViewRectMin.xy;
-	CloudRGBT = ApplyFogToCloudSampleFromVPosAndDepth(SVPosWithOffset, SVPosWithOffset, DepthKilometer, CloudRGBT, View.ViewRectMin.xy);
+	float2 SVPos = (UV / (FullResolutionToVolumetricBufferResolutionScale.x *View.BufferSizeAndInvSize.xy * VolumetricTextureSizeAndInvSize.zw)) * View.BufferSizeAndInvSize.xy + View.ViewRectMin.xy;
+	CloudRGBT = ApplyFogToCloudSampleFromVPosAndDepth(SVPos, SVPos, DepthKilometer, CloudRGBT, View.ViewRectMin.xy);
 #endif // APPLY_FOG_LATE
 
 	return CloudRGBT;
 }
 
-float4 ApplyFogToCloudFromUVs(float2 UVWithOffset, float4 CloudRGBT)
+float4 ApplyFogToCloudFromUVs(float2 UV, float4 CloudRGBT)
 {
 #if APPLY_FOG_LATE
-	float2 SVPosWithOffset = (UVWithOffset / (FullResolutionToVolumetricBufferResolutionScale.x * View.BufferSizeAndInvSize.xy * VolumetricTextureSizeAndInvSize.zw)) * View.BufferSizeAndInvSize.xy + View.ViewRectMin.xy;
-	float VolumeFrontDepthKilometer = SafeSampleVolumetricDepthTexture(UVWithOffset).x;
-	CloudRGBT = ApplyFogToCloudSampleFromVPosAndDepth(SVPosWithOffset, SVPosWithOffset, VolumeFrontDepthKilometer, CloudRGBT, View.ViewRectMin.xy);
+	float2 SVPos = (UV / (FullResolutionToVolumetricBufferResolutionScale.x * View.BufferSizeAndInvSize.xy * VolumetricTextureSizeAndInvSize.zw)) * View.BufferSizeAndInvSize.xy + View.ViewRectMin.xy;
+	float VolumeFrontDepthKilometer = SafeSampleVolumetricDepthTexture(UV).x;
+	CloudRGBT = ApplyFogToCloudSampleFromVPosAndDepth(SVPos, SVPos, VolumeFrontDepthKilometer, CloudRGBT, View.ViewRectMin.xy);
 #endif // APPLY_FOG_LATE
 
 	return CloudRGBT;
@@ -791,6 +791,9 @@ void ComposeVolumetricRTOverScenePS(
 {
 	float2 CurResPixelCoord = float2(SVPos.xy);
 	float2 ScreenUV = CurResPixelCoord * View.BufferSizeAndInvSize.zw;
+
+	// Each view has its own volumetric textures, and their contents are always aligned to the top-left even when in splitscreen
+	// As such, ScreenUVNoOffset (and other UVs derived from it) should be used to sample from volumetric textures, and for the ApplyFogToCloud functions
 	float2 ScreenUVNoOffset = (CurResPixelCoord - View.ViewRectMin.xy) * View.BufferSizeAndInvSize.zw;
 	float2 VolumeUV = FullResolutionToVolumetricBufferResolutionScale.x * (ScreenUVNoOffset * View.BufferSizeAndInvSize.xy * VolumetricTextureSizeAndInvSize.zw);
 
@@ -1190,7 +1193,7 @@ void ComposeVolumetricRTOverScenePS(
 			DataAcc /= ValidSampleCount;
 			DepthForFog /= ValidSampleCount;
 			ValidSampleCount = 1.0;
-			DataAcc = ApplyFogToCloudFromUVsAndDepth(ScreenUV, DepthForFog, DataAcc);
+			DataAcc = ApplyFogToCloudFromUVsAndDepth(ScreenUVNoOffset, DepthForFog, DataAcc);
 		}
 	#endif
 	}
