Mobile: Remove optionality for morph and cloth previous position, make it always available when used (eg. velocity rendering)
#Synced-CL 52199710

diff --git a/Engine/Shaders/Private/GpuSkinVertexFactory.ush b/Engine/Shaders/Private/GpuSkinVertexFactory.ush
--- a/Engine/Shaders/Private/GpuSkinVertexFactory.ush
+++ b/Engine/Shaders/Private/GpuSkinVertexFactory.ush
@@ -16,9 +16,6 @@
 #ifndef GPUSKIN_APEX_CLOTH
 #define GPUSKIN_APEX_CLOTH 0
 #endif
-#ifndef GPUSKIN_APEX_CLOTH_PREVIOUS
-#define GPUSKIN_APEX_CLOTH_PREVIOUS 0
-#endif
 
 #ifndef GPUSKIN_LIMIT_2BONE_INFLUENCES
 #define GPUSKIN_LIMIT_2BONE_INFLUENCES	0 // default is 4 bones 
@@ -368,13 +365,11 @@ float GetClothSimulWeight(uint StaticAlpha)
 float3 GetClothSimulPosition(int Index, bool bPrevious)
 {
 	float3 Position;
-#if GPUSKIN_APEX_CLOTH_PREVIOUS
 	if(bPrevious)
 	{
 		return float3(GPUSkinAPEXClothVF.PreviousClothSimulVertsPositionsNormals[Index * 3], GPUSkinAPEXClothVF.PreviousClothSimulVertsPositionsNormals[Index * 3 + 1].x);
 	}
 	else
-#endif
 	{
 		return float3(GPUSkinAPEXClothVF.ClothSimulVertsPositionsNormals[Index * 3], GPUSkinAPEXClothVF.ClothSimulVertsPositionsNormals[Index * 3 + 1].x);
 	}
@@ -383,13 +378,11 @@ float3 GetClothSimulPosition(int Index, bool bPrevious)
 float3 GetClothSimulNormal(int Index, bool bPrevious)
 {
 	float3 Normal;
-#if GPUSKIN_APEX_CLOTH_PREVIOUS
 	if(bPrevious)
 	{
 		return float3(GPUSkinAPEXClothVF.PreviousClothSimulVertsPositionsNormals[Index * 3 + 1].y, GPUSkinAPEXClothVF.PreviousClothSimulVertsPositionsNormals[Index * 3 + 2]);
 	}
 	else
-#endif
 	{
 		return float3(GPUSkinAPEXClothVF.ClothSimulVertsPositionsNormals[Index * 3 + 1].y, GPUSkinAPEXClothVF.ClothSimulVertsPositionsNormals[Index * 3 + 2]);
 	}
@@ -489,7 +482,6 @@ float3 MorphPosition( FVertexFactoryInput Input, FVertexFactoryIntermediates Int
 
 float3 MorphPreviousPosition( FVertexFactoryInput Input, FVertexFactoryIntermediates Intermediates )
 {
-#if GPUSKIN_MORPH_USE_PREVIOUS
 	if (ResolvedView.FrameCounter == GPUSkinVF.MorphUpdatedFrameNumber && ResolvedView.WorldIsPaused == 0)
 	{
 		// 3 + 3 floats because of the format defined in FMorphGPUSkinVertex
@@ -502,7 +494,6 @@ float3 MorphPreviousPosition( FVertexFactoryInput Input, FVertexFactoryIntermedi
 
 		return Intermediates.UnpackedPosition + PreviousDeltaPosition;
 	}
-#endif
 	return Intermediates.UnpackedPosition + Input.DeltaPosition;
 }
 #endif // GPUSKIN_MORPH_BLEND
@@ -766,11 +757,7 @@ FVertexFactoryIntermediates GetVertexFactoryIntermediates(FVertexFactoryInput In
 		if( IsSimulatedVertex(Intermediates.ClothVertexInfluences[0]) )
 		{
 			Intermediates.SimulatedPosition = ClothingPosition(Intermediates.ClothVertexInfluences, false);
-#if GPUSKIN_APEX_CLOTH_PREVIOUS
 			Intermediates.PreviousSimulatedPosition = ClothingPosition(Intermediates.ClothVertexInfluences, true);
-#else
-			Intermediates.PreviousSimulatedPosition = Intermediates.SimulatedPosition;
-#endif
 		}
 		else
 		{
@@ -898,14 +885,12 @@ float3 VertexFactoryGetPreviousInstanceSpacePosition(FVertexFactoryInput Input,
 {
 	float3 PrevSkinnedPosition = SkinPreviousPosition(Input, Intermediates);
 #if GPUSKIN_APEX_CLOTH
-#if GPUSKIN_APEX_CLOTH_PREVIOUS
 	if (GPUSkinAPEXClothVF.bEnabled && IsSimulatedVertex(Intermediates.ClothVertexInfluences[0]))
 	{
 		// Transform simulated position from cloth space into local space and blend with skinned position
 		float4 PrevTransformedSimulatedPos = mul(float4(Intermediates.PreviousSimulatedPosition.xyz, 1), GPUSkinAPEXClothVF.PreviousClothToLocal);
 		PrevSkinnedPosition = lerp(PrevSkinnedPosition, PrevTransformedSimulatedPos.xyz, GPUSkinAPEXClothVF.ClothBlendWeight * Intermediates.PreviousSimulatedPosition.w);
 	}
-#endif // GPUSKIN_APEX_CLOTH_PREVIOUS
 #endif // GPUSKIN_APEX_CLOTH
 	return PrevSkinnedPosition;
 }
diff --git a/Engine/Source/Runtime/Engine/Private/GPUSkinVertexFactory.cpp b/Engine/Source/Runtime/Engine/Private/GPUSkinVertexFactory.cpp
--- a/Engine/Source/Runtime/Engine/Private/GPUSkinVertexFactory.cpp
+++ b/Engine/Source/Runtime/Engine/Private/GPUSkinVertexFactory.cpp
@@ -915,10 +915,6 @@ void TGPUSkinVertexFactory<BoneInfluenceType>::ModifyCompilationEnvironment(cons
 
 	OutEnvironment.SetDefine(TEXT("VF_SUPPORTS_PRIMITIVE_SCENE_DATA"), bSupportsPrimitiveIdStream && bUseGPUScene);
 
-	// Mobile doesn't support motion blur, don't use previous frame morph delta for mobile.
-	const bool bIsMobile = IsMobilePlatform(Parameters.Platform) && !MobileSupportsSM5MaterialNodes(Parameters.Platform);
-	OutEnvironment.SetDefine(TEXT("GPUSKIN_MORPH_USE_PREVIOUS"), !bIsMobile);
-
 	// Whether the material supports morph targets
 	OutEnvironment.SetDefine(TEXT("GPUSKIN_MORPH_BLEND"), Parameters.MaterialParameters.bIsUsedWithMorphTargets || Parameters.MaterialParameters.bIsSpecialEngineMaterial);
 
@@ -1370,10 +1366,6 @@ void TGPUSkinAPEXClothVertexFactory<BoneInfluenceType>::ModifyCompilationEnviron
 {
 	Super::ModifyCompilationEnvironment(Parameters, OutEnvironment);
 	OutEnvironment.SetDefine(TEXT("GPUSKIN_APEX_CLOTH"),TEXT("1"));
-	
-	// Mobile doesn't support motion blur, don't use previous frame data.
-	const bool bIsMobile = IsMobilePlatform(Parameters.Platform) && !MobileSupportsSM5MaterialNodes(Parameters.Platform);
-	OutEnvironment.SetDefine(TEXT("GPUSKIN_APEX_CLOTH_PREVIOUS"), !bIsMobile);
 }
 
 template <GPUSkinBoneInfluenceType BoneInfluenceType>
