Remove MobileLocalLightsBuffers
Drop the LOCAL_LIGHTS_DISABLED suffix from TMobileBasePassPS and use LocalLightsEnabled for when local lights are enabled on mobile.

#Synced-CL 52245232

diff --git a/Engine/Shaders/Private/DeferredLightingCommon.ush b/Engine/Shaders/Private/DeferredLightingCommon.ush
--- a/Engine/Shaders/Private/DeferredLightingCommon.ush
+++ b/Engine/Shaders/Private/DeferredLightingCommon.ush
@@ -259,7 +259,7 @@ float GetLocalLightAttenuation(
 	{
 		LightMask = Square( saturate( 1 - Square( DistanceSqr * Square(LightData.InvRadius) ) ) );
 		// This extra attenuation has been added for supporting existing translucent shading model and merged local lights on mobile
-		#if SHADING_PATH_MOBILE && (TRANSLUCENCY_SH_LIGHTING || MERGED_LOCAL_LIGHTS_MOBILE > 0)
+		#if SHADING_PATH_MOBILE && TRANSLUCENCY_SH_LIGHTING
 		if (!LightData.bRectLight)
 		{ 
 			LightMask *= 1.0f / (DistanceSqr + 1.0f);
diff --git a/Engine/Shaders/Private/MobileBasePassPixelShader.usf b/Engine/Shaders/Private/MobileBasePassPixelShader.usf
--- a/Engine/Shaders/Private/MobileBasePassPixelShader.usf
+++ b/Engine/Shaders/Private/MobileBasePassPixelShader.usf
@@ -1092,73 +1092,7 @@ void Main(
 #endif
 
 	// Local lights
-#if  MERGED_LOCAL_LIGHTS_MOBILE == 1
-	{
-		FDeferredLightData MergedLightData = (FDeferredLightData)0;
-		half3 LocalLightA = MobileSceneTextures.LocalLightTextureA.Load(int3(SvPosition.xy, 0)).xyz; 
-		half4 LocalLightB = MobileSceneTextures.LocalLightTextureB.Load(int3(SvPosition.xy, 0));
-
-		MergedLightData.SourceRadius	= 0;
-		MergedLightData.SourceLength	= 0;
-		MergedLightData.Tangent			= 0;
-		MergedLightData.SoftSourceRadius= 0;
-		MergedLightData.SpecularScale	= LocalLightB.a;
-		MergedLightData.DiffuseScale	= 1.f; // Diffuse scale is not handled with merged light
-		MergedLightData.bSpotLight		= false;
-		MergedLightData.bRectLight		= false;
-		MergedLightData.Color = LocalLightA.rgb;
-		MergedLightData.bRadialLight = false;
-
-		half4 LightAttenuation = half4(1, 1, 1, 1); 
-		float SurfaceShadow = 0;
-
-		MergedLightData.Direction = OctahedronToUnitVector(Pack888To1212(LocalLightB.rgb) * 2 - 1);
-
-		FLightAccumulator NewLighting = AccumulateDynamicLighting(MaterialParameters.WorldPosition_CamRelative, CameraVector, GBuffer, 1, MergedLightData, LightAttenuation, 0, uint2(0, 0), SurfaceShadow);
-		DirectLighting = LightAccumulator_Add(DirectLighting, NewLighting);
-	}
-#elif  MERGED_LOCAL_LIGHTS_MOBILE == 2
-	{
-		const FCulledLightsGridHeader CulledLightGridHeader = GetCulledLightsGridHeader(GridIndex);
-
-		uint LightingChannelMask = 0;
-#if !MOBILE_DEFERRED_LIGHTING
-		LightingChannelMask = GetPrimitive_LightingChannelMask(MaterialParameters.PrimitiveId);
-#endif
-
-		float3 MergedLightColor = uint3(0.f, 0.f , 0.f);
-		float3 MergedLightL = float3(0,0,0);
-		float MergedSpecularScale = 0.f;
-		float MergedTotalWeight = 0;
-		MergeLocalLights(CulledLightGridHeader, MaterialParameters.WorldPosition_CamRelative, EyeIndex, false, MergedLightColor, MergedLightL, MergedSpecularScale, MergedTotalWeight, GBuffer.WorldNormal, LightingChannelMask);
-		
-		if(MergedTotalWeight > 0.f)
-		{
-			MergedLightL = MergedLightL / MergedTotalWeight;
-			MergedSpecularScale = MergedSpecularScale / MergedTotalWeight;
-
-			FDeferredLightData MergedLightData = (FDeferredLightData)0;
-			MergedLightData.SourceRadius	= 0;
-			MergedLightData.SourceLength	= 0;
-			MergedLightData.Tangent			= 0;
-			MergedLightData.SoftSourceRadius= 0;
-			MergedLightData.SpecularScale	= MergedSpecularScale;
-			MergedLightData.DiffuseScale	= 1.f; // Diffuse scale is not handled with merged light
-			MergedLightData.bSpotLight		= false;
-			MergedLightData.bRectLight		= false;
-			MergedLightData.Color = MergedLightColor;
-			MergedLightData.bRadialLight = false;
-
-			half4 LightAttenuation = half4(1, 1, 1, 1); 
-			float SurfaceShadow = 0;
-
-			MergedLightData.Direction = MergedLightL;
-
-			FLightAccumulator NewLighting = AccumulateDynamicLighting(MaterialParameters.WorldPosition_CamRelative, CameraVector, GBuffer, 1, MergedLightData, LightAttenuation, 0, uint2(0, 0), SurfaceShadow);
-			DirectLighting = LightAccumulator_Add(DirectLighting, NewLighting);
-		}
-	}
-#elif  (!MATERIAL_SHADINGMODEL_SINGLELAYERWATER && ENABLE_CLUSTERED_LIGHTS)
+#if (!MATERIAL_SHADINGMODEL_SINGLELAYERWATER && ENABLE_CLUSTERED_LIGHTS)
 	{
 		const FCulledLightsGridHeader CulledLightGridHeader = GetCulledLightsGridHeader(GridIndex);
 
diff --git a/Engine/Shaders/Private/MobileLightingCommon.ush b/Engine/Shaders/Private/MobileLightingCommon.ush
--- a/Engine/Shaders/Private/MobileLightingCommon.ush
+++ b/Engine/Shaders/Private/MobileLightingCommon.ush
@@ -675,93 +675,6 @@ void AccumulateLightGridLocalLighting(const FCulledLightsGridHeader InLightGridH
 }
 #endif
 
-#if ENABLE_CLUSTERED_LIGHTS || MERGED_LOCAL_LIGHTS_MOBILE
-void MergeLocalLights(const FCulledLightsGridHeader InLightGridHeader
-	, float3 TranslatedWorldPosition
-	, uint EyeIndex
-	, float bSupportLightFunctions
-	, inout float3 MergedLightColor
-	, inout float3 MergedLightL
-	, inout float MergedSpecularScale
-	, inout float MergedTotalWeight
-#if ENABLE_NDOTL_INTEGRATION
-	, inout float3 N
-#endif
-	, uint LightingChannelMask
-)
-{
-	MergedLightColor = uint3(0.f, 0.f , 0.f);
-	MergedLightL = float3(0,0,0);
-	MergedTotalWeight = 0;
-	MergedSpecularScale = 0.f;
-	
-	uint NumLightsInGridCell = min(InLightGridHeader.NumLights, GetMaxLightsPerCell());
-
-	LOOP
-	for (uint GridLightListIndex = 0; GridLightListIndex < NumLightsInGridCell; GridLightListIndex++)
-	{
-		const FLocalLightData LocalLight = GetLocalLightDataFromGrid(InLightGridHeader.DataStartIndex + GridLightListIndex, EyeIndex);
-
-#if !MOBILE_DEFERRED_LIGHTING
-		// Check lighting channels for mobile forward
-		if ((UnpackLightingChannelMask(LocalLight) & LightingChannelMask) == 0)
-		{
-			continue;
-		}
-#endif
-
-		// extra-early out since we know light grid is sloppy and all lights in list are radial (have a range)
-		// appears useless
-		float InvLightRadiusSq = UnpackLightInvRadius(LocalLight) * UnpackLightInvRadius(LocalLight);
-		float DistLight = length2(TranslatedWorldPosition - UnpackLightTranslatedWorldPosition(LocalLight)) * InvLightRadiusSq;
-		
-		if (DistLight > 1.0f)
-		{
-			continue;
-		}
-
-		FDeferredLightData LightData = ConvertToDeferredLight(LocalLight);
-		LightData.SpecularScale	= UnpackLightSpecularScale(LocalLight);
-		LightData.DiffuseScale	= UnpackLightDiffuseScale(LocalLight);
-
-		float3 L = LightData.Direction;	// Already normalized
-		float3 ToLight = L;
-		float3 MaskedLightColor = LightData.Color;
-		float LightMask = 1;
-		if (LightData.bRadialLight)
-		{
-			LightMask = GetLocalLightAttenuation( TranslatedWorldPosition, LightData, ToLight, L );
-			MaskedLightColor *= LightMask;
-		}
-
-		float WeightLight = 1 - DistLight;
-		if (LightData.bSpotLight)
-		{
-			WeightLight *= SpotAttenuation(L, -LightData.Direction, LightData.SpotAngles);
-		}
-
-		#if ENABLE_NDOTL_INTEGRATION
-			half NoL = max(0, dot(N, L));
-			WeightLight = WeightLight * NoL;
-			MaskedLightColor = MaskedLightColor * NoL;
-		#endif
-
-		if( LightMask > 0)
-		{
-			MergedSpecularScale += LightData.SpecularScale * WeightLight;
-			MergedTotalWeight += WeightLight;
-			// For lights with light functions we only only use the compute the direction (OutColorB)
-			// The Color is added in additional draws using MainLightFunction()
-			if (!bSupportLightFunctions || !UnpackHasLightFunction(LocalLight))
-			{
-				MergedLightColor += MaskedLightColor;
-			}
-			MergedLightL += L *  WeightLight;
-		}
-	}
-}
-#endif
-
 /*------------------------------------------------------------------------------
 	Mobile Forward Decal Lighting
 ------------------------------------------------------------------------------*/
diff --git a/Engine/Shaders/Private/MobileLocalLightsBuffer.usf b/Engine/Shaders/Private/MobileLocalLightsBuffer.usf
deleted file mode 100644
diff --git a/Engine/Shaders/Private/Substrate/SubstrateMobileForwardLighting.ush b/Engine/Shaders/Private/Substrate/SubstrateMobileForwardLighting.ush
--- a/Engine/Shaders/Private/Substrate/SubstrateMobileForwardLighting.ush
+++ b/Engine/Shaders/Private/Substrate/SubstrateMobileForwardLighting.ush
@@ -261,147 +261,7 @@ float3 SubstrateMobileForwardLighting(
 				////
 				//// Evaluate local lights
 				////
-				#if MERGED_LOCAL_LIGHTS_MOBILE == 1
-				{
-					// Merged light data
-					FDeferredLightData LightData = (FDeferredLightData)0;
-					half3 LocalLightA = MobileSceneTextures.LocalLightTextureA.Load(int3(SvPosition.xy, 0)).xyz; 
-					half4 LocalLightB = MobileSceneTextures.LocalLightTextureB.Load(int3(SvPosition.xy, 0));
-
-					LightData.SourceRadius		= 0;
-					LightData.SourceLength		= 0;
-					LightData.Tangent			= 0;
-					LightData.SoftSourceRadius	= 0;
-					LightData.SpecularScale		= LocalLightB.a;
-					LightData.DiffuseScale		= 1.f;
-					LightData.bSpotLight		= false;
-					LightData.bRectLight		= false;
-					LightData.Color 			= LocalLightA.rgb;
-					LightData.bRadialLight 		= false;
-					LightData.Direction 		= OctahedronToUnitVector(Pack888To1212(LocalLightB.rgb) * 2 - 1);
-
-					const half4 LightAttenuation 	= half4(1, 1, 1, 1); 
-					const float SurfaceShadow = 0;
-					const uint  LocalLightChannelMask = LIGHTING_CHANNEL_MASK;
-
-					const FRectTexture LocalRectTexture = ConvertToRectTexture(LightData);
-
-					float3 L = LightData.Direction;	// Already normalized
-					float3 ToLight = L;
-					float LightMask = 1;
-					if (LightData.bRadialLight)
-					{
-						LightMask = GetLocalLightAttenuation(TranslatedWorldPosition, LightData, ToLight, L);
-					}
-
-					FTranslucencyLightingVector LightingVector = GetTranslucencyLightingVector(LightData, TranslatedWorldPosition);
-					#if SUBSTRATE_TRANSLUCENCY_LIGHTING_VOLUMETRIC_DIRECTIONAL
-					GetVolumeLightingDirectional(LightingVector.AmbientLightingVector, LightingVector.DirectionalLightingVector, SubstrateBSDFContext.N, BSDFEvaluate.DiffuseColor, GetMaterialTranslucencyDirectionalLightingIntensity(), InterpolatedLighting, VolumeLighting);
-					#elif SUBSTRATE_TRANSLUCENCY_LIGHTING_VOLUMETRIC_NONDIRECTIONAL
-					GetVolumeLightingNonDirectional(LightingVector.AmbientLightingVector, BSDFEvaluate.DiffuseColor, InterpolatedLighting, VolumeLighting);
-					#else
-					// Update the Substrate BSDF context accordgin to the new L
-					SubstrateBSDFContext.SubstrateUpdateBSDFContext(ToLight);
-
-					FSubstrateEvaluateResult BSDFEvaluate = (FSubstrateEvaluateResult)0;
-					Color += SubstrateForwardLightingCommon(
-						Dither,
-						Settings,
-						LightData,
-						ToLight,
-						LightMask,
-						LightAttenuation,
-						LocalRectTexture,
-						LocalLightChannelMask,
-						PrimitiveLightingChannelMask,
-						PrecomputedShadowFactors,
-						TranslatedWorldPosition,
-						SceneDepth,
-						LuminanceWeight(SubstrateBSDFContext, CurrentBSDF),
-						SubstratePixelHeader,
-						SubstrateBSDFContext,
-						BSDFEvaluate);
-					#endif
-				}
-				#elif  MERGED_LOCAL_LIGHTS_MOBILE == 2
-				{
-					const FCulledLightsGridHeader CulledLightsGridHeader = GetCulledLightsGridHeader(GridIndex);
-
-					float3 MergedLightColor = uint3(0.f, 0.f , 0.f);
-					float3 MergedLightL = float3(0,0,0);
-					float MergedSpecularScale = 0.f;
-					float MergedTotalWeight = 0;
-
-					uint LightingChannelMask = 0;
-					#if !MOBILE_DEFERRED_LIGHTING
-						LightingChannelMask = GetPrimitive_LightingChannelMask(MaterialParameters.PrimitiveId);
-					#endif
-
-					MergeLocalLights(CulledLightsGridHeader, MaterialParameters.WorldPosition_CamRelative, EyeIndex, false/*bSupportLightFunctions*/, MergedLightColor, MergedLightL, MergedSpecularScale, MergedTotalWeight, SubstrateBSDFContext.N, LightingChannelMask);
-
-					if (MergedTotalWeight > 0.f)
-					{
-						MergedLightL = MergedLightL / MergedTotalWeight;
-						MergedSpecularScale = MergedSpecularScale / MergedTotalWeight;
-
-						FDeferredLightData LightData = (FDeferredLightData)0;
-						LightData.SourceRadius		= 0;
-						LightData.SourceLength		= 0;
-						LightData.Tangent			= 0;
-						LightData.SoftSourceRadius	= 0;
-						LightData.SpecularScale		= MergedSpecularScale;
-						LightData.DiffuseScale		= 1.f;
-						LightData.bSpotLight		= false;
-						LightData.bRectLight		= false;
-						LightData.Color 			= MergedLightColor;
-						LightData.bRadialLight 		= false;
-						LightData.Direction 		= MergedLightL;
-
-						const half4 LightAttenuation = half4(1, 1, 1, 1); 
-						const float SurfaceShadow = 0;
-						const uint  LocalLightChannelMask = LIGHTING_CHANNEL_MASK;
-
-						const FRectTexture LocalRectTexture = ConvertToRectTexture(LightData);			
-
-						float3 L = LightData.Direction;	// Already normalized
-						float3 ToLight = L;
-						float LightMask = 1;
-						if (LightData.bRadialLight)
-						{
-							LightMask = GetLocalLightAttenuation(TranslatedWorldPosition, LightData, ToLight, L);
-						}
-
-						FTranslucencyLightingVector LightingVector = GetTranslucencyLightingVector(LightData, TranslatedWorldPosition);
-						#if SUBSTRATE_TRANSLUCENCY_LIGHTING_VOLUMETRIC_DIRECTIONAL
-						GetVolumeLightingDirectional(LightingVector.AmbientLightingVector, LightingVector.DirectionalLightingVector, SubstrateBSDFContext.N, BSDFEvaluate.DiffuseColor, GetMaterialTranslucencyDirectionalLightingIntensity(), InterpolatedLighting, VolumeLighting);
-						#elif SUBSTRATE_TRANSLUCENCY_LIGHTING_VOLUMETRIC_NONDIRECTIONAL
-						GetVolumeLightingNonDirectional(LightingVector.AmbientLightingVector, BSDFEvaluate.DiffuseColor, InterpolatedLighting, VolumeLighting);
-						#else
-						// Update the Substrate BSDF context accordgin to the new L
-						SubstrateBSDFContext.SubstrateUpdateBSDFContext(ToLight);
-
-						FSubstrateEvaluateResult BSDFEvaluate = (FSubstrateEvaluateResult)0;
-						Color += SubstrateForwardLightingCommon(
-							Dither,
-							Settings,
-							LightData,
-							ToLight,
-							LightMask,
-							LightAttenuation,
-							LocalRectTexture,
-							LocalLightChannelMask,
-							PrimitiveLightingChannelMask,
-							PrecomputedShadowFactors,
-							TranslatedWorldPosition,
-							SceneDepth,
-							LuminanceWeight(SubstrateBSDFContext, CurrentBSDF),
-							SubstratePixelHeader,
-							SubstrateBSDFContext,
-							BSDFEvaluate);
-						#endif
-					}
-				}
-				#elif ENABLE_CLUSTERED_LIGHTS && !DISABLE_FORWARD_LOCAL_LIGHTS && (FORWARD_SHADING || TRANSLUCENCY_LIGHTING_SURFACE_FORWARDSHADING) // Forbidden for Single Layer Water
+				#if ENABLE_CLUSTERED_LIGHTS && !DISABLE_FORWARD_LOCAL_LIGHTS && (FORWARD_SHADING || TRANSLUCENCY_LIGHTING_SURFACE_FORWARDSHADING) // Forbidden for Single Layer Water
 				{
 					// Basic implementation for FORWARD_LOCAL_LIGHTS where local lights are reloaded per BSDF layer. It is fast by default when layer==1.
 					const FCulledLightsGridHeader CulledLightsGridHeader = GetCulledLightsGridHeader(GridIndex);
diff --git a/Engine/Source/Editor/MaterialEditor/Private/MaterialStatsCommon.cpp b/Engine/Source/Editor/MaterialEditor/Private/MaterialStatsCommon.cpp
--- a/Engine/Source/Editor/MaterialEditor/Private/MaterialStatsCommon.cpp
+++ b/Engine/Source/Editor/MaterialEditor/Private/MaterialStatsCommon.cpp
@@ -386,7 +386,7 @@ void FMaterialStatsUtils::GetRepresentativeShaderTypesAndDescriptions(TMap<FName
 	else
 	{
 		static const FName TMobileBasePassVSFNoLightMapPolicyName = TEXT("TMobileBasePassVSFNoLightMapPolicy");
-		static const FName TMobileBasePassPSFNoLightMapPolicyName = TEXT("TMobileBasePassPSFNoLightMapPolicyLOCAL_LIGHTS_DISABLED");
+		static const FName TMobileBasePassPSFNoLightMapPolicyName = TEXT("TMobileBasePassPSFNoLightMapPolicy");
 		
 		if (TargetMaterial->GetShadingModels().IsUnlit())
 		{
@@ -404,7 +404,7 @@ void FMaterialStatsUtils::GetRepresentativeShaderTypesAndDescriptions(TMap<FName
 				{
 					// distance field shadows only shaders
 					static const FName TMobileBasePassPSFMobileDistanceFieldShadowsAndLQLightMapPolicyName =
-						TEXT("TMobileBasePassPSFMobileDistanceFieldShadowsAndLQLightMapPolicyLOCAL_LIGHTS_DISABLED");
+						TEXT("TMobileBasePassPSFMobileDistanceFieldShadowsAndLQLightMapPolicy");
 					
 					ShaderTypeNamesAndDescriptions.FindOrAdd(FLocalVertexFactoryName)
 							.Add(FRepresentativeShaderInfo(ERepresentativeShader::StationarySurface, TMobileBasePassPSFMobileDistanceFieldShadowsAndLQLightMapPolicyName,
@@ -413,7 +413,7 @@ void FMaterialStatsUtils::GetRepresentativeShaderTypesAndDescriptions(TMap<FName
 				else //no shadows & lightmapped
 				{
 					static const FName TMobileBasePassPSTLightMapPolicyLQName =
-						TEXT("TMobileBasePassPSTLightMapPolicyLQLOCAL_LIGHTS_DISABLED");
+						TEXT("TMobileBasePassPSTLightMapPolicyLQ");
 										
 					ShaderTypeNamesAndDescriptions.FindOrAdd(FLocalVertexFactoryName)
 						.Add(FRepresentativeShaderInfo(ERepresentativeShader::StationarySurface, TMobileBasePassPSTLightMapPolicyLQName,
diff --git a/Engine/Source/Runtime/Core/Private/HAL/ConsoleManager.cpp b/Engine/Source/Runtime/Core/Private/HAL/ConsoleManager.cpp
--- a/Engine/Source/Runtime/Core/Private/HAL/ConsoleManager.cpp
+++ b/Engine/Source/Runtime/Core/Private/HAL/ConsoleManager.cpp
@@ -4029,8 +4029,7 @@ static TAutoConsoleVariable<int32> CVarMobileForwardEnableLocalLights(
 	TEXT("r.Mobile.Forward.EnableLocalLights"),
 	1,
 	TEXT("0: Local Lights Disabled (default)\n"
-		"1: Local Lights Enabled\n"
-		"2: Local Lights Buffer Enabled\n"),
+		"1: Local Lights Enabled\n"),
 	ECVF_ReadOnly | ECVF_RenderThreadSafe);
 
 static TAutoConsoleVariable<int32> CVarMobileRenderRectLightsAsSpotLights(
diff --git a/Engine/Source/Runtime/Engine/Classes/Engine/EngineTypes.h b/Engine/Source/Runtime/Engine/Classes/Engine/EngineTypes.h
--- a/Engine/Source/Runtime/Engine/Classes/Engine/EngineTypes.h
+++ b/Engine/Source/Runtime/Engine/Classes/Engine/EngineTypes.h
@@ -676,7 +676,7 @@ enum EMobileLocalLightSetting : int
 {
 	LOCAL_LIGHTS_DISABLED UMETA(DisplayName = "Local Lights Disabled"),
 	LOCAL_LIGHTS_ENABLED UMETA(DisplayName = "Local Lights Enabled"),
-	LOCAL_LIGHTS_BUFFER UMETA(DisplayName = "Local Lights Buffer Enabled")
+	LOCAL_LIGHTS_BUFFER UE_DEPRECATED(5.8, "Local Lights Buffer has been deprecated. For local lights heavy games, please use mobile deferred r.Mobile.ShadingPath=1") UMETA(hidden)
 };
 
 /** Controls the way that the width scale property affects animation trails. */
diff --git a/Engine/Source/Runtime/Engine/Private/ShaderCompiler/ShaderCompiler.cpp b/Engine/Source/Runtime/Engine/Private/ShaderCompiler/ShaderCompiler.cpp
--- a/Engine/Source/Runtime/Engine/Private/ShaderCompiler/ShaderCompiler.cpp
+++ b/Engine/Source/Runtime/Engine/Private/ShaderCompiler/ShaderCompiler.cpp
@@ -3431,9 +3431,7 @@ void GlobalBeginCompileShader(
 			Input.DebugGroupName.ReplaceInline(TEXT("PROPAGATE_AO"), TEXT("AO"));
 			Input.DebugGroupName.ReplaceInline(TEXT("PROPAGATE_SECONDARY_OCCLUSION"), TEXT("SEC_OCC"));
 			Input.DebugGroupName.ReplaceInline(TEXT("PROPAGATE_MULTIPLE_BOUNCES"), TEXT("MULT_BOUNC"));
-			Input.DebugGroupName.ReplaceInline(TEXT("LOCAL_LIGHTS_DISABLED"), TEXT("NoLL"));
-			Input.DebugGroupName.ReplaceInline(TEXT("LOCAL_LIGHTS_ENABLED"), TEXT("LL"));
-			Input.DebugGroupName.ReplaceInline(TEXT("LOCAL_LIGHTS_PREPASS_ENABLED"), TEXT("LLPP"));
+			Input.DebugGroupName.ReplaceInline(TEXT("LocalLights"), TEXT("LL"));
 			Input.DebugGroupName.ReplaceInline(TEXT("PostProcess"), TEXT("Post"));
 			Input.DebugGroupName.ReplaceInline(TEXT("AntiAliasing"), TEXT("AA"));
 			Input.DebugGroupName.ReplaceInline(TEXT("Mobile"), TEXT("Mob"));
diff --git a/Engine/Source/Runtime/Landscape/Private/LandscapeRender.cpp b/Engine/Source/Runtime/Landscape/Private/LandscapeRender.cpp
--- a/Engine/Source/Runtime/Landscape/Private/LandscapeRender.cpp
+++ b/Engine/Source/Runtime/Landscape/Private/LandscapeRender.cpp
@@ -4049,9 +4049,9 @@ public:
 			FName("TLightMapDensityPSFNoLightMapPolicy"),
 
 			// Mobile
-			FName("TMobileBasePassPSFMobileDirectionalLightAndSHIndirectPolicyLOCAL_LIGHTS_DISABLED"),
+			FName("TMobileBasePassPSFMobileDirectionalLightAndSHIndirectPolicy"),
 			FName("TMobileBasePassVSFMobileDirectionalLightAndSHIndirectPolicy"),
-			FName("TMobileBasePassPSFNoLightMapPolicyLOCAL_LIGHTS_DISABLED"),
+			FName("TMobileBasePassPSFNoLightMapPolicy"),
 			FName("TMobileBasePassVSFNoLightMapPolicy"),
 
 			// Runtime virtual texture
@@ -4099,17 +4099,17 @@ public:
 			FName("FDebugViewModePS"),
 
 			// Mobile
-			FName("TMobileBasePassPSFMobileMovableDirectionalLightCSMWithLightmapPolicyLOCAL_LIGHTS_DISABLED"),
+			FName("TMobileBasePassPSFMobileMovableDirectionalLightCSMWithLightmapPolicy"),
 			FName("TMobileBasePassVSFMobileMovableDirectionalLightCSMWithLightmapPolicy"),
-			FName("TMobileBasePassPSFMobileMovableDirectionalLightWithLightmapPolicyLOCAL_LIGHTS_DISABLED"),
+			FName("TMobileBasePassPSFMobileMovableDirectionalLightWithLightmapPolicy"),
 			FName("TMobileBasePassVSFMobileMovableDirectionalLightWithLightmapPolicy"),
-			FName("TMobileBasePassPSFMobileDistanceFieldShadowsLightMapAndCSMLightingPolicyLOCAL_LIGHTS_DISABLED"),
+			FName("TMobileBasePassPSFMobileDistanceFieldShadowsLightMapAndCSMLightingPolicy"),
 			FName("TMobileBasePassVSFMobileDistanceFieldShadowsLightMapAndCSMLightingPolicy"),
-			FName("TMobileBasePassPSFMobileDistanceFieldShadowsAndLQLightMapPolicyLOCAL_LIGHTS_DISABLED"),
+			FName("TMobileBasePassPSFMobileDistanceFieldShadowsAndLQLightMapPolicy"),
 			FName("TMobileBasePassVSFMobileDistanceFieldShadowsAndLQLightMapPolicy"),
-			FName("TMobileBasePassPSTLightMapPolicyLQLOCAL_LIGHTS_DISABLED"),
+			FName("TMobileBasePassPSTLightMapPolicyLQ"),
 			FName("TMobileBasePassVSTLightMapPolicyLQ"),
-			FName("TMobileBasePassPSFMobileDirectionalLightAndCSMPolicyLOCAL_LIGHTS_DISABLED"),
+			FName("TMobileBasePassPSFMobileDirectionalLightAndCSMPolicy"),
 			FName("TMobileBasePassVSFMobileDirectionalLightAndCSMPolicy"),
 
 			FName("TBasePassVSFCachedVolumeIndirectLightingPolicy"),
diff --git a/Engine/Source/Runtime/Renderer/Internal/SceneTextures.h b/Engine/Source/Runtime/Renderer/Internal/SceneTextures.h
--- a/Engine/Source/Runtime/Renderer/Internal/SceneTextures.h
+++ b/Engine/Source/Runtime/Renderer/Internal/SceneTextures.h
@@ -143,10 +143,6 @@ struct FSceneTextures : public FMinimalSceneTextures
 
 	// Texture containing dynamic motion vectors. Can be bound by the base pass or its own velocity pass.
 	FRDGTextureRef Velocity{};
-	
-	// (Mobile Local Light Prepass) Textures containing LocalLight Direction and Color
-	FRDGTextureRef MobileLocalLightTextureA {};
-	FRDGTextureRef MobileLocalLightTextureB {};
 
 	// Texture containing the screen space ambient occlusion result.
 	FRDGTextureRef ScreenSpaceAO{};
diff --git a/Engine/Source/Runtime/Renderer/Private/LightRendering.cpp b/Engine/Source/Runtime/Renderer/Private/LightRendering.cpp
--- a/Engine/Source/Runtime/Renderer/Private/LightRendering.cpp
+++ b/Engine/Source/Runtime/Renderer/Private/LightRendering.cpp
@@ -510,7 +510,7 @@ bool FDeferredLightVS::ShouldCompilePermutation(const FGlobalShaderPermutationPa
 	FPermutationDomain PermutationVector(Parameters.PermutationId);
 	if (PermutationVector.Get<FRadialLight>())
 	{
-		return IsFeatureLevelSupported(Parameters.Platform, ERHIFeatureLevel::SM5) || IsMobileDeferredShadingEnabled(Parameters.Platform) || MobileLocalLightsBufferEnabled(Parameters.Platform);
+		return IsFeatureLevelSupported(Parameters.Platform, ERHIFeatureLevel::SM5) || IsMobileDeferredShadingEnabled(Parameters.Platform);
 	}
 	// used with FPrefilterPlanarReflectionPS on mobile
 	return true;
diff --git a/Engine/Source/Runtime/Renderer/Private/MobileBasePass.cpp b/Engine/Source/Runtime/Renderer/Private/MobileBasePass.cpp
--- a/Engine/Source/Runtime/Renderer/Private/MobileBasePass.cpp
+++ b/Engine/Source/Runtime/Renderer/Private/MobileBasePass.cpp
@@ -20,21 +20,6 @@
 #include "VariableRateShadingImageManager.h"
 #include "SceneCore.h"
 
-bool MobileLocalLightsBufferEnabled(const FStaticShaderPlatform Platform)
-{
-	return FReadOnlyCVARCache::MobileForwardLocalLights(Platform) == 2;
-}
-
-bool MobileMergeLocalLightsInPrepassEnabled(const FStaticShaderPlatform Platform)
-{
-	return MobileLocalLightsBufferEnabled(Platform) && MobileUsesFullDepthPrepass(Platform);
-}
-
-bool MobileMergeLocalLightsInBasepassEnabled(const FStaticShaderPlatform Platform)
-{
-	return MobileLocalLightsBufferEnabled(Platform) && !MobileUsesFullDepthPrepass(Platform);
-}
-
 int32 GMobileForwardLocalLightsSinglePermutation = 0;
 FAutoConsoleVariableRef CVarMobileForwardLocalLightsSinglePermutation(
 	TEXT("r.Mobile.Forward.LocalLightsSinglePermutation"),
@@ -51,20 +36,7 @@ bool MobileLocalLightsUseSinglePermutation(EShaderPlatform ShaderPlatform)
 EMobileLocalLightSetting GetMobileForwardLocalLightSetting(EShaderPlatform ShaderPlatform)
 {
 	const int32 MobileForwardLocalLightsIniValue = FReadOnlyCVARCache::MobileForwardLocalLights(ShaderPlatform);
-
-	if (MobileForwardLocalLightsIniValue > 0)
-	{
-		if (MobileForwardLocalLightsIniValue == 1)
-		{
-			return EMobileLocalLightSetting::LOCAL_LIGHTS_ENABLED;
-		}
-		else if (MobileForwardLocalLightsIniValue == 2)
-		{
-			return EMobileLocalLightSetting::LOCAL_LIGHTS_BUFFER;
-		}
-	}
-
-	return EMobileLocalLightSetting::LOCAL_LIGHTS_DISABLED;
+	return MobileForwardLocalLightsIniValue > 0 ? EMobileLocalLightSetting::LOCAL_LIGHTS_ENABLED: EMobileLocalLightSetting::LOCAL_LIGHTS_DISABLED;
 }
 
 uint8 GetMobileShadingModelStencilValue(FMaterialShadingModelField ShadingModel)
@@ -253,17 +225,6 @@ bool MobileBasePass::GetShaders(
 				PixelShader
 				);
 		}
-		case EMobileLocalLightSetting::LOCAL_LIGHTS_BUFFER:
-		{
-			return GetMobileBasePassShaders<EMobileLocalLightSetting::LOCAL_LIGHTS_BUFFER>(
-				LightMapPolicy,
-				MaterialResource,
-				VertexFactoryType,
-				ColoredTransmittanceFallback,
-				VertexShader,
-				PixelShader
-				);
-		}
 
 		default:
 			check(false);
diff --git a/Engine/Source/Runtime/Renderer/Private/MobileBasePassRendering.cpp b/Engine/Source/Runtime/Renderer/Private/MobileBasePassRendering.cpp
--- a/Engine/Source/Runtime/Renderer/Private/MobileBasePassRendering.cpp
+++ b/Engine/Source/Runtime/Renderer/Private/MobileBasePassRendering.cpp
@@ -156,17 +156,16 @@ EMobileTranslucentColorTransmittanceMode MobileActiveTranslucentColorTransmittan
 	typedef TMobileBasePassVS< LightMapPolicyType > TMobileBasePassVS##LightMapPolicyName; \
 	IMPLEMENT_MATERIAL_SHADER_TYPE(template<>, TMobileBasePassVS##LightMapPolicyName, TEXT("/Engine/Private/MobileBasePassVertexShader.usf"), TEXT("Main"), SF_Vertex); \
 
-#define IMPLEMENT_MOBILE_SHADING_BASEPASS_LIGHTMAPPED_PIXEL_SHADER_TYPE(LightMapPolicyType, LightMapPolicyName, LocalLightSetting) \
-	typedef TMobileBasePassPS< LightMapPolicyType, LocalLightSetting, EMobileTranslucentColorTransmittanceMode::DEFAULT > TMobileBasePassPS##LightMapPolicyName##LocalLightSetting; \
-	IMPLEMENT_MATERIAL_SHADER_TYPE(template<>, TMobileBasePassPS##LightMapPolicyName##LocalLightSetting, TEXT("/Engine/Private/MobileBasePassPixelShader.usf"), TEXT("Main"), SF_Pixel); \
-	typedef TMobileBasePassPS< LightMapPolicyType, LocalLightSetting, EMobileTranslucentColorTransmittanceMode::SINGLE_SRC_BLENDING > TMobileBasePassPS##LightMapPolicyName##LocalLightSetting##ThinTranslGrey; \
-	IMPLEMENT_MATERIAL_SHADER_TYPE(template<>, TMobileBasePassPS##LightMapPolicyName##LocalLightSetting##ThinTranslGrey, TEXT("/Engine/Private/MobileBasePassPixelShader.usf"), TEXT("Main"), SF_Pixel); \
+#define IMPLEMENT_MOBILE_SHADING_BASEPASS_LIGHTMAPPED_PIXEL_SHADER_TYPE(LightMapPolicyType, LightMapPolicyName, LocalLightSettingType, LocalLightSettingName) \
+	typedef TMobileBasePassPS< LightMapPolicyType, LocalLightSettingType, EMobileTranslucentColorTransmittanceMode::DEFAULT > TMobileBasePassPS##LightMapPolicyName##LocalLightSettingName; \
+	IMPLEMENT_MATERIAL_SHADER_TYPE(template<>, TMobileBasePassPS##LightMapPolicyName##LocalLightSettingName, TEXT("/Engine/Private/MobileBasePassPixelShader.usf"), TEXT("Main"), SF_Pixel); \
+	typedef TMobileBasePassPS< LightMapPolicyType, LocalLightSettingType, EMobileTranslucentColorTransmittanceMode::SINGLE_SRC_BLENDING > TMobileBasePassPS##LightMapPolicyName##LocalLightSettingName##ThinTranslGrey; \
+	IMPLEMENT_MATERIAL_SHADER_TYPE(template<>, TMobileBasePassPS##LightMapPolicyName##LocalLightSettingName##ThinTranslGrey, TEXT("/Engine/Private/MobileBasePassPixelShader.usf"), TEXT("Main"), SF_Pixel); \
 
 #define IMPLEMENT_MOBILE_SHADING_BASEPASS_LIGHTMAPPED_SHADER_TYPE(LightMapPolicyType, LightMapPolicyName) \
 	IMPLEMENT_MOBILE_SHADING_BASEPASS_LIGHTMAPPED_VERTEX_SHADER_TYPE(LightMapPolicyType, LightMapPolicyName) \
-	IMPLEMENT_MOBILE_SHADING_BASEPASS_LIGHTMAPPED_PIXEL_SHADER_TYPE(LightMapPolicyType, LightMapPolicyName, LOCAL_LIGHTS_DISABLED) \
-	IMPLEMENT_MOBILE_SHADING_BASEPASS_LIGHTMAPPED_PIXEL_SHADER_TYPE(LightMapPolicyType, LightMapPolicyName, LOCAL_LIGHTS_ENABLED) \
-	IMPLEMENT_MOBILE_SHADING_BASEPASS_LIGHTMAPPED_PIXEL_SHADER_TYPE(LightMapPolicyType, LightMapPolicyName, LOCAL_LIGHTS_BUFFER)
+	IMPLEMENT_MOBILE_SHADING_BASEPASS_LIGHTMAPPED_PIXEL_SHADER_TYPE(LightMapPolicyType, LightMapPolicyName, LOCAL_LIGHTS_DISABLED, ) \
+	IMPLEMENT_MOBILE_SHADING_BASEPASS_LIGHTMAPPED_PIXEL_SHADER_TYPE(LightMapPolicyType, LightMapPolicyName, LOCAL_LIGHTS_ENABLED, LocalLights)
 
 // Implement shader types per lightmap policy 
 IMPLEMENT_MOBILE_SHADING_BASEPASS_LIGHTMAPPED_SHADER_TYPE(TUniformLightMapPolicy<LMP_NO_LIGHTMAP>, FNoLightMapPolicy);
diff --git a/Engine/Source/Runtime/Renderer/Private/MobileBasePassRendering.h b/Engine/Source/Runtime/Renderer/Private/MobileBasePassRendering.h
--- a/Engine/Source/Runtime/Renderer/Private/MobileBasePassRendering.h
+++ b/Engine/Source/Runtime/Renderer/Private/MobileBasePassRendering.h
@@ -30,10 +30,6 @@
 #include "CompositionLighting/PostProcessDeferredDecals.h"
 #include "MobileSSR.h"
 
-bool MobileLocalLightsBufferEnabled(const FStaticShaderPlatform Platform);
-bool MobileMergeLocalLightsInPrepassEnabled(const FStaticShaderPlatform Platform);
-bool MobileMergeLocalLightsInBasepassEnabled(const FStaticShaderPlatform Platform);
-
 struct FMobileBasePassTextures
 {
 	FDBufferTextures DBufferTextures = {};
@@ -412,21 +408,6 @@ public:
 		
 		FForwardLightingParameters::ModifyCompilationEnvironment(Parameters.Platform, OutEnvironment);
 		OutEnvironment.SetDefine(TEXT("ENABLE_CLUSTERED_LIGHTS"), (LocalLightSetting == EMobileLocalLightSetting::LOCAL_LIGHTS_ENABLED) ? 1u : 0u);
-
-		// Translucent materials don't write to depth so cannot use prepass
-		uint32 MergedLocalLights = 0u;
-		if (LocalLightSetting == EMobileLocalLightSetting::LOCAL_LIGHTS_BUFFER)
-		{
-			if (MobileMergeLocalLightsInPrepassEnabled(Parameters.Platform) && !bTranslucentOrSLWMaterial)
-			{
-				MergedLocalLights = 1u;
-			}
-			else if (MobileMergeLocalLightsInBasepassEnabled(Parameters.Platform) || bTranslucentOrSLWMaterial)
-			{
-				MergedLocalLights = 2u;
-			}
-		}
-		OutEnvironment.SetDefine(TEXT("MERGED_LOCAL_LIGHTS_MOBILE"), MergedLocalLights);
 		OutEnvironment.SetDefine(TEXT("ENABLE_CLUSTERED_REFLECTION"), bEnableClusteredReflections ? 1u : 0u);
 		OutEnvironment.SetDefine(TEXT("ENABLE_SHADOWMASKTEXTURE"), bMobileUsesShadowMaskTexture && !bTranslucentOrSLWMaterial ? 1u : 0u);
 		OutEnvironment.SetDefine(TEXT("ENABLE_DBUFFER_TEXTURES"), Parameters.MaterialParameters.MaterialDomain == MD_Surface ? 1u : 0u);
diff --git a/Engine/Source/Runtime/Renderer/Private/MobileLocalLightsBuffer.cpp b/Engine/Source/Runtime/Renderer/Private/MobileLocalLightsBuffer.cpp
deleted file mode 100644
diff --git a/Engine/Source/Runtime/Renderer/Private/MobileShadingRenderer.cpp b/Engine/Source/Runtime/Renderer/Private/MobileShadingRenderer.cpp
--- a/Engine/Source/Runtime/Renderer/Private/MobileShadingRenderer.cpp
+++ b/Engine/Source/Runtime/Renderer/Private/MobileShadingRenderer.cpp
@@ -1723,12 +1723,6 @@ void FMobileSceneRenderer::Render(FRDGBuilder& GraphBuilder, const FSceneRenderU
 			}
 		}
 
-		// Local Light prepass
-		if (bRendererOutputFinalSceneColor)
-		{
-			RenderMobileLocalLightsBuffer(GraphBuilder, SceneTextures, SortedLightSet);
-		}
-
 		if (bRendererOutputFinalSceneColor)
 		{
 			if (bRequiresDBufferDecals)
diff --git a/Engine/Source/Runtime/Renderer/Private/SceneRendering.h b/Engine/Source/Runtime/Renderer/Private/SceneRendering.h
--- a/Engine/Source/Runtime/Renderer/Private/SceneRendering.h
+++ b/Engine/Source/Runtime/Renderer/Private/SceneRendering.h
@@ -2975,8 +2975,6 @@ protected:
 
 	void RenderVelocityPass(FRHICommandList& RHICmdList, const FViewInfo& View, const FInstanceCullingDrawParams* InstanceCullingDrawParams);
 
-	void RenderMobileLocalLightsBuffer(FRDGBuilder& GraphBuilder, FSceneTextures& SceneTextures, const FSortedLightSetSceneInfo& SortedLights);
-
 	void RenderCustomRenderPassBasePass(FRDGBuilder& GraphBuilder, TArrayView<FViewInfo> InViews, FRDGTextureRef ViewFamilyTexture, FSceneTextures& SceneTextures, FMobileBasePassTextures& BasePassTextures, FInstanceCullingManager& InstanceCullingManager, bool bIncludeTranslucent);
 
 	/** Renders the opaque base pass for mobile. */
diff --git a/Engine/Source/Runtime/Renderer/Private/SceneTextures.cpp b/Engine/Source/Runtime/Renderer/Private/SceneTextures.cpp
--- a/Engine/Source/Runtime/Renderer/Private/SceneTextures.cpp
+++ b/Engine/Source/Runtime/Renderer/Private/SceneTextures.cpp
@@ -823,16 +823,6 @@ void FSceneTextures::InitializeViewFamily(FRDGBuilder& GraphBuilder, FViewFamily
 	}
 #endif
 
-	extern bool MobileMergeLocalLightsInPrepassEnabled(const FStaticShaderPlatform Platform);
-	if(MobileMergeLocalLightsInPrepassEnabled(Config.ShaderPlatform))
-	{
-		FRDGTextureDesc MobileLocalLightTextureADesc = FRDGTextureDesc::Create2D(Config.Extent, PF_FloatR11G11B10, FClearValueBinding::Transparent, TexCreate_RenderTargetable | TexCreate_ShaderResource);
-		SceneTextures.MobileLocalLightTextureA = GraphBuilder.CreateTexture(MobileLocalLightTextureADesc, TEXT("MobileLocalLightTextureA"));
-
-		FRDGTextureDesc MobileLocalLightTextureBDesc = FRDGTextureDesc::Create2D(Config.Extent, PF_B8G8R8A8, FClearValueBinding({ 0.5f, 0.5f, 0.5f, 0.f }), TexCreate_RenderTargetable | TexCreate_ShaderResource);
-		SceneTextures.MobileLocalLightTextureB = GraphBuilder.CreateTexture(MobileLocalLightTextureBDesc, TEXT("MobileLocalLightTextureB"));
-	}
-
 	SceneTextures.DBufferTextures = MakePimpl<FDBufferTextures, EPimplPtrMode::DeepCopy>();
 	SceneTextures.SubstrateSceneData = MakePimpl<FSubstrateSceneData, EPimplPtrMode::DeepCopy>();
 
@@ -988,8 +978,6 @@ TArray<FRDGTextureRef> FSceneTextures::EnumerateSceneTextures() const
 	AddTextureIfNonNull(GBufferSGGX, Results);
 	AddTextureIfNonNull(DepthAux, Results);
 	AddTextureIfNonNull(Velocity, Results);
-	AddTextureIfNonNull(MobileLocalLightTextureA, Results);
-	AddTextureIfNonNull(MobileLocalLightTextureB, Results);
 	AddTextureIfNonNull(ScreenSpaceAO, Results);
 	AddTextureIfNonNull(DebugAux, Results);
 #if WITH_EDITOR
@@ -1335,16 +1323,6 @@ void SetupMobileSceneTextureUniformParameters(
 			}
 		}
 
-		if (HasBeenProduced(SceneTextures->MobileLocalLightTextureA))
-		{
-			SceneTextureParameters.LocalLightTextureA = SceneTextures->MobileLocalLightTextureA;
-		}
-
-		if (HasBeenProduced(SceneTextures->MobileLocalLightTextureB))
-		{
-			SceneTextureParameters.LocalLightTextureB = SceneTextures->MobileLocalLightTextureB;
-		}
-
 		if (EnumHasAnyFlags(SetupMode, EMobileSceneTextureSetupMode::CustomDepth))
 		{
 			const FCustomDepthTextures& CustomDepthTextures = SceneTextures->CustomDepth;
