Fix minor shader warning

#rnx
#Synced-CL 52269267

diff --git a/Engine/Shaders/Private/PostProcessCombineLUTsInner.usf b/Engine/Shaders/Private/PostProcessCombineLUTsInner.usf
--- a/Engine/Shaders/Private/PostProcessCombineLUTsInner.usf
+++ b/Engine/Shaders/Private/PostProcessCombineLUTsInner.usf
@@ -65,7 +65,7 @@ float3 ComputeOutputACES(
 	float bSDR)
 {
 #if USE_ACES_2
-	float3 PaperMaxLuminance = lerp(OutputMaxLuminance * InvHDRPaperWhiteRatio, OutputMaxLuminance, bSDR);
+	float PaperMaxLuminance = lerp(OutputMaxLuminance * InvHDRPaperWhiteRatio, OutputMaxLuminance, bSDR);
 	float3 ODTColor = ACESOutputTransform(
 		GradedColor * ACESSceneColorMultiplier,
 		(float3x3)WorkingColorSpace.ToAP0,
