[SingleLayerWater] Fixed striping artifacts in shallow water caused by texture gather sampling a different quad of texels than the shader thinks it does.

#Synced-CL 52295890

diff --git a/Engine/Shaders/Private/SingleLayerWaterCommon.ush b/Engine/Shaders/Private/SingleLayerWaterCommon.ush
--- a/Engine/Shaders/Private/SingleLayerWaterCommon.ush
+++ b/Engine/Shaders/Private/SingleLayerWaterCommon.ush
@@ -96,12 +96,15 @@ void UnpackWaterRefractionMask(float4 PackedRefractionMask, out bool4 bOutIsWate
 #undef SLW_MASK_IS_WATER_BITMASK
 #undef SLW_MASK_REFRACTION_BITMASK
 
-void SampleSceneDepthWithoutWaterBilateral(Texture2D Tex, SamplerState TexSampler, float2 UV, float2 TexSize, float WaterSurfaceSceneDepth,
+void SampleSceneDepthWithoutWaterBilateral(Texture2D Tex, SamplerState TexSampler, float2 UV, float2 TexSize, float2 InvTexSize, float WaterSurfaceSceneDepth,
 #if SINGLE_LAYER_WATER_REFRACTION_MASK_AVAILABLE
 	Texture2D RefractionMaskTex,
 #endif
 	out float OutDeviceZ, out float OutSceneDepth, out float4 OutWeights, out float4 OutRefractionMask)
 {
+	// Bias UV by ~1/512 of a texel to avoid sitting exactly on Gather block boundaries,
+	// where GPU fixed-point and shader float32 arithmetic can disagree on block selection.
+	UV += (0.5f / 256.0f) * InvTexSize;
 	// Gather deviceZ depth values and convert to scene depth
 	const float4 DeviceZ4 = Tex.Gather(TexSampler, UV);
 	const float4 SceneDepth4 = float4(ConvertFromDeviceZ(DeviceZ4.x), ConvertFromDeviceZ(DeviceZ4.y), ConvertFromDeviceZ(DeviceZ4.z), ConvertFromDeviceZ(DeviceZ4.w));
diff --git a/Engine/Shaders/Private/SingleLayerWaterComposite.usf b/Engine/Shaders/Private/SingleLayerWaterComposite.usf
--- a/Engine/Shaders/Private/SingleLayerWaterComposite.usf
+++ b/Engine/Shaders/Private/SingleLayerWaterComposite.usf
@@ -93,7 +93,7 @@ SingleLayerWaterCompositeOutput SingleLayerWaterComposite(float2 BufferUV, float
 		float OpaqueDepth;
 		float4 WeightsDummy;
 		float4 RefractionMaskDummy;
-		SampleSceneDepthWithoutWaterBilateral(SceneNoWaterDepthTexture, SceneNoWaterDepthSampler, SceneNoWaterUV, SceneNoWaterTextureSize, WaterDepth, OpaqueDeviceZ, OpaqueDepth, WeightsDummy, RefractionMaskDummy);
+		SampleSceneDepthWithoutWaterBilateral(SceneNoWaterDepthTexture, SceneNoWaterDepthSampler, SceneNoWaterUV, SceneNoWaterTextureSize, SceneNoWaterInvTextureSize, WaterDepth, OpaqueDeviceZ, OpaqueDepth, WeightsDummy, RefractionMaskDummy);
 
 		float DeltaDepth = CameraIsUnderWater ? WaterDepth : OpaqueDepth - WaterDepth;	// Inverted depth
 		if (DeltaDepth > 0.0)
@@ -180,7 +180,7 @@ SingleLayerWaterCompositeOutput SingleLayerWaterComposite(float2 BufferUV, float
 		float OpaqueDepth;
 		float4 WeightsDummy;
 		float4 RefractionMaskDummy;
-		SampleSceneDepthWithoutWaterBilateral(SceneNoWaterDepthTexture, SceneNoWaterDepthSampler, SceneNoWaterUV, SceneNoWaterTextureSize, WaterDepth, OpaqueDeviceZ, OpaqueDepth, WeightsDummy, RefractionMaskDummy);
+		SampleSceneDepthWithoutWaterBilateral(SceneNoWaterDepthTexture, SceneNoWaterDepthSampler, SceneNoWaterUV, SceneNoWaterTextureSize, SceneNoWaterInvTextureSize, WaterDepth, OpaqueDeviceZ, OpaqueDepth, WeightsDummy, RefractionMaskDummy);
 
 		float DeltaDepth = CameraIsUnderWater ? WaterDepth : OpaqueDepth - WaterDepth;	// Inverted depth
 		if (DeltaDepth > 0.0)
diff --git a/Engine/Shaders/Private/SingleLayerWaterShading.ush b/Engine/Shaders/Private/SingleLayerWaterShading.ush
--- a/Engine/Shaders/Private/SingleLayerWaterShading.ush
+++ b/Engine/Shaders/Private/SingleLayerWaterShading.ush
@@ -89,7 +89,7 @@ WaterVolumeLightingOutput EvaluateWaterVolumeLighting(
 		float BehindWaterSceneDepth;
 		float4 BehindWaterSceneColorWeights;
 		float4 DummyRefractionMask;
-		SampleSceneDepthWithoutWaterBilateral(SceneDepthWithoutSingleLayerWaterTexture, SceneDepthWithoutSingleLayerWaterSampler, ViewportUV, SceneDepthWithoutSingleLayerWaterTextureSize,
+		SampleSceneDepthWithoutWaterBilateral(SceneDepthWithoutSingleLayerWaterTexture, SceneDepthWithoutSingleLayerWaterSampler, ViewportUV, SceneDepthWithoutSingleLayerWaterTextureSize, SceneDepthWithoutSingleLayerWaterTextureInvSize,
 			WaterSurfaceSceneDepth,
 			BehindWaterSceneDeviceZ, BehindWaterSceneDepth, BehindWaterSceneColorWeights, DummyRefractionMask);
 #else
@@ -110,7 +110,8 @@ WaterVolumeLightingOutput EvaluateWaterVolumeLighting(
 		float PixelBehindWaterSceneDepth;
 		float4 PixelSceneColorWeights;
 		float4 PixelRefractionMask;
-		SampleSceneDepthWithoutWaterBilateral(SceneDepthWithoutSingleLayerWaterTexture, SceneDepthWithoutSingleLayerWaterSampler, ViewportUV, SceneDepthWithoutSingleLayerWaterTextureSize, WaterSurfaceSceneDepth,
+		SampleSceneDepthWithoutWaterBilateral(SceneDepthWithoutSingleLayerWaterTexture, SceneDepthWithoutSingleLayerWaterSampler, ViewportUV, SceneDepthWithoutSingleLayerWaterTextureSize, SceneDepthWithoutSingleLayerWaterTextureInvSize,
+			WaterSurfaceSceneDepth,
 		#if SINGLE_LAYER_WATER_REFRACTION_MASK_AVAILABLE
 			RefractionMaskTexture,
 		#endif
@@ -132,7 +133,8 @@ WaterVolumeLightingOutput EvaluateWaterVolumeLighting(
 		float DistortedSceneDepth;
 		float4 DistortedSceneColorWeights;
 		float4 DistortedRefractionMask;
-		SampleSceneDepthWithoutWaterBilateral(SceneDepthWithoutSingleLayerWaterTexture, SceneDepthWithoutSingleLayerWaterSampler, DistortedUV, SceneDepthWithoutSingleLayerWaterTextureSize, WaterSurfaceSceneDepth,
+		SampleSceneDepthWithoutWaterBilateral(SceneDepthWithoutSingleLayerWaterTexture, SceneDepthWithoutSingleLayerWaterSampler, DistortedUV, SceneDepthWithoutSingleLayerWaterTextureSize, SceneDepthWithoutSingleLayerWaterTextureInvSize,
+			WaterSurfaceSceneDepth,
 #if SINGLE_LAYER_WATER_REFRACTION_MASK_AVAILABLE
 			RefractionMaskTexture,
 #endif
