Fix error compiling MainVS in FrontLayerTranslucency.usf due to FRONT_LAYER_TILE_SIZE not being declared.
- split preprocessor branches for MainVS and MainPS.

#Synced-CL 52300829

diff --git a/Engine/Shaders/Private/FrontLayerTranslucency.usf b/Engine/Shaders/Private/FrontLayerTranslucency.usf
--- a/Engine/Shaders/Private/FrontLayerTranslucency.usf
+++ b/Engine/Shaders/Private/FrontLayerTranslucency.usf
@@ -35,6 +35,10 @@ struct FFrontLayerTranslucencyGBufferVSToPS
 	INVARIANT_OUTPUT float4 Position : SV_POSITION;
 };
 
+#endif // defined(MainVS) || defined(MainPS)
+
+#ifdef MainVS
+
 void MainVS(
 	FVertexFactoryInput Input,
 	out FFrontLayerTranslucencyGBufferVSToPS Output
@@ -63,14 +67,18 @@ void MainVS(
 	Output.FactoryInterpolants = VertexFactoryGetInterpolantsVSToPS(Input, VFIntermediates, VertexParameters);
 }
 
-#if defined(MainPS) && !OUTPUT_PIXEL_DEPTH_OFFSET
+#endif // MainVS
+
+#ifdef MainPS
+
+#if !OUTPUT_PIXEL_DEPTH_OFFSET
 	// Force early depth/stencil test unless the material uses depth offset to speedup the front layer GBuffer pass.
 	// Cant use EARLYDEPTHSTENCIL since that causes artifacts when combined with clip/discard.
 	// DEPTHSTENCIL_EARLYTEST_LATEWRITE can only be defined after including Platform.ush
 	#define PIXELSHADER_EARLYDEPTHSTENCIL DEPTHSTENCIL_EARLYTEST_LATEWRITE
-#else // !defined(MainPS) || OUTPUT_PIXEL_DEPTH_OFFSET
+#else // OUTPUT_PIXEL_DEPTH_OFFSET
 	#define PIXELSHADER_EARLYDEPTHSTENCIL
-#endif // !defined(MainPS) || OUTPUT_PIXEL_DEPTH_OFFSET
+#endif // OUTPUT_PIXEL_DEPTH_OFFSET
 
 PIXELSHADER_EARLYDEPTHSTENCIL
 void MainPS(
@@ -225,7 +233,7 @@ void MainPS(
 #endif // PROJECT_SUPPORTS_MEGALIGHTS
 }
 
-#endif // defined(MainVS) || defined(MainPS)
+#endif // MainPS
 
 #ifdef FrontLayerTranslucencyClearGBufferPS
 
