Substrate - Fix for alpha masked editor materials using alpha-to-coverage into MSAA

#tests PC editor, engine test and strata level.
#rnx

#Synced-CL 52376387

diff --git a/Engine/Shaders/Private/BasePassPixelShader.usf b/Engine/Shaders/Private/BasePassPixelShader.usf
--- a/Engine/Shaders/Private/BasePassPixelShader.usf
+++ b/Engine/Shaders/Private/BasePassPixelShader.usf
@@ -2717,14 +2717,15 @@ void FPixelShaderInOut_MainPS(
 //		- Pixel depth offset prevents the usage of early depth test because the depth is generated in shader. So we can only relay on LateZ in this case.
 //		- DEPTHSTENCIL_EARLYTEST_LATEWRITE does not help in this case due to UAV writes we need to skip.
 //		- In order to avoid UAV writes, we will run in shader manual depth tests.
+//		- We also disable early depth test when runnning editor primitive alpha-to-coverage (for instance masked LevelGridMaterial)
 #if IS_NANITE_PASS
 	// Even with masking and/or pixel depth offset, Nanite never uses actual clip() or SV_Depth
 	#define PIXELSHADER_EARLYDEPTHSTENCIL EARLYDEPTHSTENCIL
-#elif !OUTPUT_PIXEL_DEPTH_OFFSET
-	#define PIXELSHADER_EARLYDEPTHSTENCIL EARLYDEPTHSTENCIL
 #elif COMPILER_SUPPORTS_DEPTHSTENCIL_EARLYTEST_LATEWRITE
 	// If we support early depth test with late write behaviour then use it since we may be using discard, or modifying depth
 	#define PIXELSHADER_EARLYDEPTHSTENCIL DEPTHSTENCIL_EARLYTEST_LATEWRITE
+#elif !OUTPUT_PIXEL_DEPTH_OFFSET && !(EDITOR_ALPHA2COVERAGE && MATERIALBLENDING_MASKED)
+	#define PIXELSHADER_EARLYDEPTHSTENCIL EARLYDEPTHSTENCIL
 #endif
 
 #else // SUBSTRATE_OPAQUE_DEFERRED
