Transmittance LUT encoding

#tests PC editor, skyatmosphere engine test.
#rnx

#Synced-CL 52376387

diff --git a/Engine/Shaders/Private/SkyAtmosphere.usf b/Engine/Shaders/Private/SkyAtmosphere.usf
--- a/Engine/Shaders/Private/SkyAtmosphere.usf
+++ b/Engine/Shaders/Private/SkyAtmosphere.usf
@@ -395,6 +395,7 @@ float3 GetTransmittance(in float LightZenithCosAngle, in float PHeight)
 	float3 TransmittanceToLight = 1.0f;
 #else
 	float3 TransmittanceToLight = TransmittanceLutTexture.SampleLevel(TransmittanceLutTextureSampler, UV, 0).rgb;
+	TransmittanceToLight = DecodeTransmittance(TransmittanceToLight);
 #endif
 	return TransmittanceToLight;
 }
@@ -1196,6 +1197,8 @@ void RenderTransmittanceLutCS(uint3 ThreadId : SV_DispatchThreadID)
 
 	float3 transmittance = exp(-ss.OpticalDepth);
 
+	transmittance = EncodeTransmittance(transmittance);
+
 #if SHADING_PATH_MOBILE
 	TransmittanceLutUAV[int2(PixPos)] = float4(transmittance, 0.0f);
 #else
diff --git a/Engine/Shaders/Private/SkyAtmosphereCommon.ush b/Engine/Shaders/Private/SkyAtmosphereCommon.ush
--- a/Engine/Shaders/Private/SkyAtmosphereCommon.ush
+++ b/Engine/Shaders/Private/SkyAtmosphereCommon.ush
@@ -230,6 +230,19 @@ float3x3 GetSkyViewLutReferential(in float4x4 ViewSkyViewLutReferential)
 	return SkyViewLutReferential;
 }
 
+float3 EncodeTransmittance(float3 Transmittance)
+{
+	// We encode transmittance to avoid banding issue when storing it in RGB8.
+	// Encoding is simple square root which is a good fit to get more precision in the low value part,
+	// typically at horizon where it is important to get more precision for sunset transmittance on the sun disk.
+	return sqrt(Transmittance);
+}
+
+float3 DecodeTransmittance(float3 Transmittance)
+{
+	return Transmittance * Transmittance;
+}
+
 float3 GetAtmosphereTransmittance(
 	float3 PlanetCenterToWorldPos, float3 WorldDir, float BottomRadius, float TopRadius,
 	Texture2D<float4> TransmittanceLutTexture, SamplerState TransmittanceLutTextureSampler)
@@ -250,7 +263,7 @@ float3 GetAtmosphereTransmittance(
 	float2 TransmittanceLutUv;
 	getTransmittanceLutUvs(PHeight, LightZenithCosAngle, BottomRadius, TopRadius, TransmittanceLutUv);
 	const float3 TransmittanceToLight = Texture2DSampleLevel(TransmittanceLutTexture, TransmittanceLutTextureSampler, TransmittanceLutUv, 0.0f).rgb;
-	return TransmittanceToLight;
+	return DecodeTransmittance(TransmittanceToLight);
 }
 
 float3 GetLightDiskLuminance(
