Fix for GPU skin vertex factory missing instancing matrices.

#Synced-CL 52384345

diff --git a/Engine/Shaders/Private/GpuSkinVertexFactory.ush b/Engine/Shaders/Private/GpuSkinVertexFactory.ush
--- a/Engine/Shaders/Private/GpuSkinVertexFactory.ush
+++ b/Engine/Shaders/Private/GpuSkinVertexFactory.ush
@@ -318,8 +318,9 @@ FMaterialVertexParameters GetMaterialVertexParameters(
 
 	Result.PreSkinnedPosition = Intermediates.UnpackedPosition.xyz;
 	Result.PreSkinnedNormal = Input.TangentZ.xyz;
-
-	// Assumes no instacing
+	Result.InstanceLocalToWorld = Intermediates.SceneData.InstanceData.LocalToWorld;
+	Result.InstanceWorldToLocal = Intermediates.SceneData.InstanceData.WorldToLocal;
+	Result.PerInstanceRandom = Intermediates.SceneData.InstanceData.RandomID;
 	Result.PrevFrameLocalToWorld = Intermediates.PrevLocalToWorld;
 
 #if NUM_MATERIAL_TEXCOORDS_VERTEX
