Fix Substrate multi layared material with special complex slabs.

#Synced-CL 52375688

diff --git a/Engine/Shaders/Private/Substrate/Substrate.ush b/Engine/Shaders/Private/Substrate/Substrate.ush
--- a/Engine/Shaders/Private/Substrate/Substrate.ush
+++ b/Engine/Shaders/Private/Substrate/Substrate.ush
@@ -439,7 +439,7 @@ struct FSubstrateBSDF
 #define HEADER_BIT_OFFSET_COMPLEXPATH_MODE			(HEADER_BIT_OFFSET_SHAREDLOCALBASES_TYPE	+ HEADER_BIT_COUNT_SHAREDLOCALBASES_TYPE)
 
 // Sanity check
-#if (HEADER_BIT_OFFSET_COMPLEXPATH_SPECIAL + HEADER_BIT_COUNT_COMPLEXPATH_MODE) > 32
+#if (HEADER_BIT_OFFSET_COMPLEXPATH_MODE + HEADER_BIT_COUNT_COMPLEXPATH_MODE) > 32
 #error Substrate Header is larger than 32 bits
 #endif
 
diff --git a/Engine/Shaders/Private/Substrate/SubstrateExport.ush b/Engine/Shaders/Private/Substrate/SubstrateExport.ush
--- a/Engine/Shaders/Private/Substrate/SubstrateExport.ush
+++ b/Engine/Shaders/Private/Substrate/SubstrateExport.ush
@@ -241,7 +241,12 @@ void PackSubstrateOut(
 						const bool bSSSMask = BSDF.bIsBottom;	// SSS and Thin can only be on the bottom layer, so ignoring bIsSimpleVolume here.
 						SubstratePixelHeader.SetHasSubsurface(bSSSMask && bHasSSS && !bIsSimpleVolume);
 
-						SubstratePixelHeader.SetComplexPathMode(bForceComplexSpecialPath ? COMPLEXPATH_MODE_SLABSPECIAL : COMPLEXPATH_MODE_SLAB);
+						// Ensure the complex path mode is not overriden with a "Slab" if a previous BSDF set it to "SlabSpecial"
+						#if SUBSTRATE_COMPLEXSPECIALPATH
+						SubstratePixelHeader.SetComplexPathMode(bForceComplexSpecialPath ? COMPLEXPATH_MODE_SLABSPECIAL : SubstratePixelHeader.GetComplexPathMode());
+						#else
+						SubstratePixelHeader.SetComplexPathMode(COMPLEXPATH_MODE_SLAB);
+						#endif
 						break;
 					}
 
