Re-allow CompositeReflectionCapturesAndSkylightTWS to apply Reflection Catpures for projects that have r.AllowStaticLighting=0, it was being used for low end lighting
* Adds 139 instructions from Translucent Surface materials in projects with r.AllowStaticLighting=0, which are mostly branched around if no Reflection Captures are present

#Synced-CL 52184307

diff --git a/Engine/Shaders/Private/ReflectionEnvironmentComposite.ush b/Engine/Shaders/Private/ReflectionEnvironmentComposite.ush
--- a/Engine/Shaders/Private/ReflectionEnvironmentComposite.ush
+++ b/Engine/Shaders/Private/ReflectionEnvironmentComposite.ush
@@ -35,7 +35,6 @@ float3 CompositeReflectionCapturesAndSkylightTWS(
 	OutDiffuseIBL = half4(0.0f, 0.0f, 0.0f, 1.0f);
 #endif
 
-#if ALLOW_STATIC_LIGHTING || DIFFUSE_FROM_CAPTURE
 #if REFLECTION_COMPOSITE_USE_BLENDED_REFLECTION_CAPTURES
 	// Accumulate reflections from captures affecting this tile, applying largest captures first so that the smallest ones display on top
 	LOOP
@@ -195,8 +194,6 @@ float3 CompositeReflectionCapturesAndSkylightTWS(
 	ImageBasedReflections.rgb *= View.PrecomputedIndirectSpecularColorScale;
 	CompositedAverageBrightness.x *= Luminance( View.PrecomputedIndirectSpecularColorScale );
 
-#endif	// ALLOW_STATIC_LIGHTING || DIFFUSE_FROM_CAPTURE
-
 #if ENABLE_SKY_LIGHT
 #if DIFFUSE_FROM_CAPTURE
 	float SkylightReflectionScale = OutDiffuseIBL.a;
