Fix missing hair transmittance initialization causing RT lighting to be in correct in hair material.

#Synced-CL 52417206

diff --git a/Engine/Shaders/Private/RayTracing/RayTracingLightingCommon.ush b/Engine/Shaders/Private/RayTracing/RayTracingLightingCommon.ush
--- a/Engine/Shaders/Private/RayTracing/RayTracingLightingCommon.ush
+++ b/Engine/Shaders/Private/RayTracing/RayTracingLightingCommon.ush
@@ -73,6 +73,8 @@ FDeferredLightData GetRayTracingDeferredLightData(
 
 	LightData.IESAtlasIndex = RayTracingLightData.IESAtlasIndex;
 
+	LightData.HairTransmittance = InitHairTransmittanceData();
+
 	OutLightType = LightType;
 	OutLightMissShaderIndex = RayTracingLightData.LightMissShaderIndex;
 
