Fix incorrect diffuse response on legacy clear coat with Substrate.

#Synced-CL 52417206

diff --git a/Engine/Shaders/Private/Substrate/SubstrateEvaluation.ush b/Engine/Shaders/Private/Substrate/SubstrateEvaluation.ush
--- a/Engine/Shaders/Private/Substrate/SubstrateEvaluation.ush
+++ b/Engine/Shaders/Private/Substrate/SubstrateEvaluation.ush
@@ -497,7 +497,12 @@ FSubstrateEvaluateResult SubstrateEvaluateBSDFCommon(FSubstrateBSDFContext BSDFC
 			DiffuseColor = ApplyDiffuseColorBoost(DiffuseColor, LumenHardwareRayTracingUniformBuffer.DiffuseColorBoost);
 		#endif
 
+			// When using SimpleClearCoat with CLEAR_COAT_BOTTOM_NORMAL, the AreaLight is constructed with the clear coat's 
+			// normal. In order to get the correct diffuse response, AreaLight_NoL and AreaLight_Falloff are recomputed with 
+			// the bottom layer's normal.
 			bool bHazeAsSimpleClearCoat = false;
+			float AreaLight_NoL = AreaLight.NoL;
+			float AreaLight_Falloff = AreaLight.Falloff;
 		#if SUBSTRATE_FASTPATH==0 && CLEAR_COAT_BOTTOM_NORMAL 
 			float  bUseBottomNormal = false;
 			float3 TopNormal = BSDFContext.N;
@@ -519,6 +524,19 @@ FSubstrateEvaluateResult SubstrateEvaluateBSDFCommon(FSubstrateBSDFContext BSDFC
 					BSDFContext.TangentBasis[0] = BSDFContext.X;
 					BSDFContext.TangentBasis[1] = BSDFContext.Y;
 					BSDFContext.TangentBasis[2] = BottomNormal;
+
+					// Recompute AreaLight_Falloff and AreaLight_NoL with the correct diffuse normal
+					BRANCH
+					if (bIsRectLight)
+					{
+						RectIrradianceLambert(BottomNormal, AreaLight.Rect, AreaLight_Falloff, AreaLight_NoL);
+					}
+					else
+					{
+						// For point light, AreaLight.SphereSinAlpha=0, and SphereHorizonCosWrap() becomes a NOP
+						const float NoL = saturate(dot(BottomNormal, AreaLight.DiffuseL));
+						AreaLight_NoL = SphereHorizonCosWrap(NoL, Square(AreaLight.SphereSinAlpha));
+					}
 				}
 			}
 		#endif
@@ -539,14 +557,14 @@ FSubstrateEvaluateResult SubstrateEvaluateBSDFCommon(FSubstrateBSDFContext BSDFC
 				//   this transmittance, as the energy preservation code we remove it as well. However the visual effect was small 
 				//   and the cost was rather large (~8%). See the code below for where this is handled, this is no longer necessary
 				//   with the latest model.
-				Sample.DiffusePathValue = Diffuse_GGX_Rough(DiffuseColor, SafeRoughness, NoV, AreaLight.NoL, VoH, NoH, GetAreaLightDiffuseMicroReflWeight(AreaLight));
+				Sample.DiffusePathValue = Diffuse_GGX_Rough(DiffuseColor, SafeRoughness, NoV, AreaLight_NoL, VoH, NoH, GetAreaLightDiffuseMicroReflWeight(AreaLight));
 			}
 			else
 #endif
 			{
 				Sample.DiffusePathValue = Diffuse_Lambert(DiffuseColor);
 			}
-			Sample.IntegratedDiffuseValue += (ShadowTerms.SurfaceShadow * AreaLight.NoL * AreaLight.Falloff) * Sample.DiffusePathValue * AreaLight.FalloffColor;
+			Sample.IntegratedDiffuseValue += (ShadowTerms.SurfaceShadow * AreaLight_NoL * AreaLight_Falloff) * Sample.DiffusePathValue * AreaLight.FalloffColor;
 			Sample.DiffuseColor		= DiffuseColor;
 			Sample.DiffusePDF		= BSDFContext.SatNoL / PI;
 			Sample.bPostProcessSubsurface = false;
