[TSR] Improve high-intensity line aliasing in thin geometry detection

We extends depth-based thin edge detection to also detect bright sub-pixel lines using luminance intensity. Depth alone misses specular highlights and emissive edges that flicker under jitter. It's enabled by default with a new CVar r.TSR.ThinGeometryDetection.Coverage.MinKeepLineIntensity (default 0.30, set 0 to disable). The cost is added due to the need of the luma texture being recorded, and it is free if temporal stability analysis is on.

#tests thin geo nanite foliage test level, city sample, and TSR test all shows positive improvements on high intensity lines.
#Synced-CL 51948989

diff --git a/Engine/Shaders/Private/TemporalSuperResolution/TSRDetectThinGeometry.usf b/Engine/Shaders/Private/TemporalSuperResolution/TSRDetectThinGeometry.usf
--- a/Engine/Shaders/Private/TemporalSuperResolution/TSRDetectThinGeometry.usf
+++ b/Engine/Shaders/Private/TemporalSuperResolution/TSRDetectThinGeometry.usf
@@ -21,6 +21,13 @@
 #define EDGE_TYPE EDGE_HORIZONTAL_VERTICAL_DIAGNAL
 
 #define CONFIG_VARYING_RANGE_DETECTION (DIM_VARYING_RANGE_DETECTION)
+#define CONFIG_INTENSITY_LINE_DETECTION (DIM_INTENSITY_LINE_DETECTION)
+
+#if CONFIG_INTENSITY_LINE_DETECTION
+#define TLineEdgeType tsr_tensor_half2
+#else
+#define TLineEdgeType tsr_tensor_half
+#endif
 //------------------------------------------------------- CONSTANTS
 
 #if CONFIG_LDS_DESPILL
@@ -34,9 +41,13 @@ int TileOverscan;
 int ThinGeometryTextureIndex;
 float ErrorMultiplier;
 float MaxRelaxationWeight;
+float MinKeepLineIntensity;
 
 Texture2D<float> SceneDepthTexture;
 Texture2D<tsr_half> CurrentCoverageTexture;
+#if CONFIG_INTENSITY_LINE_DETECTION
+Texture2D<tsr_half> FlickeringLumaTexture;
+#endif
 Texture2D<tsr_half> ReprojectedHistoryCoverageTexture;
 Texture2D<tsr_half2> DecimateMaskTexture;
 
@@ -137,15 +148,28 @@ void FetchRegularPassInput(
 	tsr_tensor_short2 PixelPos
 	,out tsr_tensor_half SceneDepth
 	,out tsr_tensor_half DeviceZThinGeometryDetectionError
-	,out tsr_tensor_half ReprojectedHistoryCoverage)
+	,out tsr_tensor_half ReprojectedHistoryCoverage
+#if CONFIG_INTENSITY_LINE_DETECTION
+	,out tsr_tensor_half PrevHighIntensity
+	,out tsr_tensor_half FlickeringLuma
+#endif
+)
 {
 	tsr_tensor_float DepthCache;
-
+#if CONFIG_INTENSITY_LINE_DETECTION
+	tsr_tensor_half  RawCoverageCache;
+#endif
+	
 	UNROLL_N(SIMD_SIZE)
 	for (uint ElementIndex = 0; ElementIndex < SIMD_SIZE; ElementIndex++)
 	{
 		tsr_short2 ClampedFetchPixelPos = PixelPos.GetElement(ElementIndex);
+#if CONFIG_INTENSITY_LINE_DETECTION
+		FlickeringLuma.SetElement(ElementIndex, FlickeringLumaTexture[ClampedFetchPixelPos]);
+		RawCoverageCache.SetElement(ElementIndex, ReprojectedHistoryCoverageTexture[ClampedFetchPixelPos]);
+#else
 		ReprojectedHistoryCoverage.SetElement(ElementIndex, ReprojectedHistoryCoverageTexture[ClampedFetchPixelPos]);
+#endif
 		DepthCache.SetElement(ElementIndex,SceneDepthTexture[ClampedFetchPixelPos]);
 	}
 
@@ -163,6 +187,14 @@ void FetchRegularPassInput(
 		SceneDepth.SetElement(ElementIndex, tsr_half(Depth));
 		DeviceZThinGeometryDetectionError.SetElement(ElementIndex, tsr_half(DeviceZError));
 	}
+	
+#if CONFIG_INTENSITY_LINE_DETECTION
+	{
+		tsr_tensor_ushort Packed = tsr_tensor_ushort::CastFrom<tsr_half>(RawCoverageCache * tsr_half(255.0));
+		ReprojectedHistoryCoverage = tsr_tensor_half::CastFrom<tsr_ushort>(bit_shift_right(Packed, tsr_tensor_ushort::Const(tsr_ushort(0x1u)))) * rcp(tsr_half(127.0));
+		PrevHighIntensity = tsr_tensor_half::CastFrom<tsr_ushort>(bit_and(Packed, tsr_tensor_ushort::Const(tsr_ushort(0x1u))));
+	}
+#endif
 } // FetchRegularPassInput
 
 //------------------------------------------------------- FUNCTIONS
@@ -321,7 +353,12 @@ void DetectThinGeometryCS(
 	tsr_tensor_half HistoryRelaxationWeight = tsr_tensor_half::Const(0);
 	tsr_tensor_bool IsClusterHoleRegion = tsr_tensor_bool::Const(false);
 	tsr_tensor_half NewSampleWeight = tsr_tensor_half::Const(0.1);
-
+	#if CONFIG_INTENSITY_LINE_DETECTION
+	tsr_tensor_half FlickeringLuma = tsr_tensor_half::Const(0);
+	tsr_tensor_half PrevHighIntensity = tsr_tensor_half::Const(0);
+	tsr_tensor_half IntensityLineEdge = tsr_tensor_half::Const(0);
+	tsr_tensor_half CurrentLineEdge = tsr_tensor_half::Const(0);
+	#endif
 	// Thin geometry coverage estimation
 	if (bRegularPass)
 	{
@@ -329,12 +366,31 @@ void DetectThinGeometryCS(
 			DeviceZ, 
 			DeviceZThinGeometryDetectionError,
 			ReprojectedHistoryCoverage
+#if CONFIG_INTENSITY_LINE_DETECTION
+			,PrevHighIntensity
+			,FlickeringLuma
+#endif
 		);
 		
 		{
-			ThinGeometryEdge = HorizontalThinEdge3x3(DeviceZ, DeviceZThinGeometryDetectionError);
-			ThinGeometryEdge = ThinGeometryEdge + VerticalThinEdge3x3(DeviceZ, DeviceZThinGeometryDetectionError);
+			#if CONFIG_INTENSITY_LINE_DETECTION
+			TLineEdgeType LineDetectionValue = Concatenate(DeviceZ, FlickeringLuma);
+			TLineEdgeType LineDetectionThreshold = Concatenate(DeviceZThinGeometryDetectionError, tsr_tensor_half::Const(tsr_half(MinKeepLineIntensity)));
+			#else
+			TLineEdgeType LineDetectionValue = DeviceZ;
+			TLineEdgeType LineDetectionThreshold = DeviceZThinGeometryDetectionError;
+			#endif
+
+			TLineEdgeType LineEdge = HorizontalThinEdge3x3(LineDetectionValue, LineDetectionThreshold);
+			LineEdge = LineEdge + VerticalThinEdge3x3(LineDetectionValue, LineDetectionThreshold);
 			//Debug[0].SetComponent(2, tsr_tensor_half(ThinGeometryEdge));
+			#if CONFIG_INTENSITY_LINE_DETECTION
+			ThinGeometryEdge = LineEdge[0];
+			CurrentLineEdge = saturate(LineEdge[1] * tsr_tensor_half::Const(tsr_half(POSITIVE_INFINITY)));
+			//Debug[0].SetComponent(1, tsr_tensor_half(CurrentLineEdge));
+			#else
+			ThinGeometryEdge = LineEdge;
+			#endif
 		}
 
 		// Student t test of current observation and history to check if they are of the same distribution.
@@ -376,19 +432,65 @@ void DetectThinGeometryCS(
 			tsr_tensor_bool AcceptHistory = and(ValidCurrentThinGeometryRegion, and(bSameDistribution, and(NotAbsoluteMatch, not(bHasPixelAnimation))));
 			NewSampleWeight = select(AcceptHistory, NewSampleWeight, tsr_tensor_half::Const(1.0));
 
-			// Update reprojected history and use it for weight calculation.
+			//1. Update reprojected history and use it for weight calculation.
 			ReprojectedHistoryCoverage = lerp(ReprojectedHistoryCoverage,CurrentCoverage,NewSampleWeight);
-			Debug[0].SetComponent(1, NewSampleWeight);
+			//Debug[0].SetComponent(1, NewSampleWeight);
+
+			// Valid pixels for cluster hole stability improvement.
+			IsClusterHoleRegion = ReprojectedHistoryCoverage > tsr_tensor_half::Const(half(0.99f/127.0f));
+
+			#if CONFIG_INTENSITY_LINE_DETECTION
+			{
+				// 2. stochastic OR accumulation of intensity lines
+				// Keep it live through bright-dark alternation
+				// stochastically decay clears it after consecutive dark frames
+				tsr_tensor_half Noise;
+				UNROLL_N(SIMD_SIZE)
+				for (uint ElementIndex = 0; ElementIndex < SIMD_SIZE; ElementIndex++)
+				{
+					//Cannot use View.StateFrameIndexMod8 as it cannot trigger decay with a random initialization
+					#if 0
+						uint2 Random = Rand3DPCG16(int3(float2(PixelPos.GetElement(ElementIndex)), View.StateFrameIndex)).xy;
+						tsr_half TmpNoise = tsr_half(Hammersley16(0, 1, Random).x);
+					#else
+						tsr_half TmpNoise = tsr_half(InterleavedGradientNoise(float2(PixelPos.GetElement(ElementIndex)), View.StateFrameIndex));
+					#endif
+				
+					Noise.SetElement(ElementIndex, TmpNoise);
+				}
+
+				// 3% chance to decay if current is dark (and 90% if it is not inside thin geo region).  
+				tsr_tensor_half bDark = tsr_tensor_half::Const(tsr_half(1.0)) - CurrentLineEdge;
+				tsr_tensor_half DecayProbability = select(IsClusterHoleRegion, tsr_tensor_half::Const(tsr_half(0.03)), tsr_tensor_half::Const(tsr_half(0.9)));
+				tsr_tensor_half bDecay = bDark * saturate((DecayProbability - Noise) * tsr_tensor_half::Const(tsr_half(POSITIVE_INFINITY)));
+
+				//Debug[0].SetComponent(0, bDecay);
+				//Debug[0].SetComponent(3, PrevHighIntensity);
+
+				IntensityLineEdge = max(CurrentLineEdge, PrevHighIntensity) * (tsr_tensor_half::Const(1.0) - bDecay);
+			}
+#endif
 		}
 
-		// Update new history relaxation weight.
+		// Update new history relaxation weight (thin geometry cluster and Intensity weights).
 		{
 			// Calculate the weights.
 			{
 				HistoryRelaxationWeight = abs(ReprojectedHistoryCoverage - tsr_tensor_half::Const(0.5))+tsr_tensor_half::Const(0.02);
 				HistoryRelaxationWeight = HistoryRelaxationWeight * exp(HistoryRelaxationWeight * tsr_half(-20.0)) *  tsr_half(10.0/0.184); // delta = 1 maximum
 			}
-			tsr_tensor_half Min5x5HistoryCoverage = Min3x3(Min3x3(ReprojectedHistoryCoverage));
+			#if CONFIG_INTENSITY_LINE_DETECTION
+				// Min5x5(1 - x) == 1 - Max5x5(x), so this computes:
+				//   Min5x5HistoryCoverage = min5x5(ReprojectedHistoryCoverage)
+				//   IntensityWeight       = max5x5(IntensityLineEdge)  (after 1 - Min5x5(1 - x))
+				tsr_tensor_half Min5x5HistoryCoverage;
+				tsr_tensor_half IntensityWeight;
+				Deconcatenate(Min3x3(Min3x3(Concatenate(ReprojectedHistoryCoverage, tsr_tensor_half::Const(1.0) - IntensityLineEdge))), Min5x5HistoryCoverage, IntensityWeight);
+				IntensityWeight = tsr_tensor_half::Const(1.0) - IntensityWeight;
+			#else
+				tsr_tensor_half Min5x5HistoryCoverage = Min3x3(Min3x3(ReprojectedHistoryCoverage));
+			#endif
+			
 			tsr_tensor_bool bThinGeometryHistoryRegion = Min5x5HistoryCoverage > tsr_tensor_half::Const(0);
 			
 			// Reject samples that has min being background
@@ -419,14 +521,24 @@ void DetectThinGeometryCS(
 				}
 				#endif
 			}
+			
+			#if CONFIG_INTENSITY_LINE_DETECTION
+			// Update intensity weights
+			// Apply a uniform weight when intensity line antialiasing is required and HistoryRelaxationWeight is not set.
+			{
+				IntensityWeight =  IntensityWeight * (tsr_tensor_half::Const(tsr_half(1)) - CurrentLineEdge);
+				// cluster weights is more important than intensity weights.
+				HistoryRelaxationWeight = select(HistoryRelaxationWeight > tsr_tensor_half::Const(0), HistoryRelaxationWeight, IntensityWeight * tsr_tensor_half::Const(0.6));
+				//Debug[0].SetComponent(2,IntensityLineEdge);
+			}
+			#endif
+
 			// Not similar patches or has pixel animation should reset relaxation weight to 0
 			// to use more strict history clamping.
 			HistoryRelaxationWeight = select(and(bSameDistribution, not(bHasPixelAnimation)), HistoryRelaxationWeight, tsr_tensor_half::Const(0));
 			//Debug[0].SetComponent(2,HistoryRelaxationWeight);
 		}
 
-		// Valid pixels for cluster hole stability improvement.
-		IsClusterHoleRegion = ReprojectedHistoryCoverage > tsr_tensor_half::Const(half(0.99f/127.0f));
 	}
 	else
 	{	// Fast pass does not have any thin geometry indicators. Update coverage to current.
@@ -447,7 +559,17 @@ void DetectThinGeometryCS(
 		const tsr_short2 OutputPixelPos = ComputeElementOutputPixelPos(ElementIndex);
 		
 		R8Output[uint3(OutputPixelPos,ThinGeometryTextureIndex)] = EncodedThinGeometryEdge.GetElement(ElementIndex);//uint(ThinGeometryFlag * tsr_half(255.0));
-		HistoryCoverageOutput[uint3(OutputPixelPos,0)] = ReprojectedHistoryCoverage.GetElement(ElementIndex);
+		
+		#if CONFIG_INTENSITY_LINE_DETECTION
+			tsr_half CoverageClamped = saturate(ReprojectedHistoryCoverage.GetElement(ElementIndex));
+			tsr_ushort CoverageBits = tsr_ushort(CoverageClamped * tsr_half(127.0));
+			tsr_ushort IntensityBit = select(IntensityLineEdge.GetElement(ElementIndex) > tsr_half(0), tsr_ushort(1), tsr_ushort(0));
+			tsr_ushort PackedCoverage = (CoverageBits << tsr_ushort(1)) | IntensityBit;
+			HistoryCoverageOutput[uint3(OutputPixelPos, 0)] = tsr_half(PackedCoverage) / tsr_half(255.0);
+
+		#else
+			HistoryCoverageOutput[uint3(OutputPixelPos,0)] = ReprojectedHistoryCoverage.GetElement(ElementIndex);
+		#endif
 	}
 	
 	#if DEBUG_OUTPUT
diff --git a/Engine/Source/Runtime/Renderer/Private/PostProcess/TemporalSuperResolution.cpp b/Engine/Source/Runtime/Renderer/Private/PostProcess/TemporalSuperResolution.cpp
--- a/Engine/Source/Runtime/Renderer/Private/PostProcess/TemporalSuperResolution.cpp
+++ b/Engine/Source/Runtime/Renderer/Private/PostProcess/TemporalSuperResolution.cpp
@@ -353,6 +353,12 @@ TAutoConsoleVariable<float> CVarTSRThinGeometryCoverageMaxRelaxationWeight(
 	TEXT("The max history clamping box relaxation weight due to thin geometry detection (0 to 1)"),
 	ECVF_RenderThreadSafe);
 
+TAutoConsoleVariable<float> CVarTSRThinGeometryCoverageMinKeepLineIntensity(
+	TEXT("r.TSR.ThinGeometryDetection.Coverage.MinKeepLineIntensity"),
+	0.30,
+	TEXT("The value at which pixels are regarded as intensity lines after Normalization [0,1]. Do not reject on high intensity lines"),
+	ECVF_RenderThreadSafe);
+
 TAutoConsoleVariable<int32> CVarTSRThinGeometryDetectionWeightRelaxation(
 	TEXT("r.TSR.ThinGeometryDetection.WeightRelaxation"), 1,
 	TEXT("Adaptively trim the history relaxation to avoid ghosting."),
@@ -1055,7 +1061,9 @@ class FTSRDetectThinGeometryCS : public FTSRConvolutionNetworkShader
 	SHADER_USE_PARAMETER_STRUCT(FTSRDetectThinGeometryCS, FTSRConvolutionNetworkShader);
 
 	class FVaryingRangeDetectionDim : SHADER_PERMUTATION_BOOL("DIM_VARYING_RANGE_DETECTION");
-	using FPermutationDomain = TShaderPermutationDomain<FTSRConvolutionNetworkShader::FPermutationDomain, FVaryingRangeDetectionDim>;
+	class FIntensityLineDetectionDim : SHADER_PERMUTATION_BOOL("DIM_INTENSITY_LINE_DETECTION");
+	using FPermutationDomain = TShaderPermutationDomain<FTSRConvolutionNetworkShader::FPermutationDomain, FVaryingRangeDetectionDim, FIntensityLineDetectionDim>;
+
 
 	BEGIN_SHADER_PARAMETER_STRUCT(FParameters, )
 		SHADER_PARAMETER_STRUCT_INCLUDE(FTSRCommonParameters, CommonParameters)
@@ -1064,9 +1072,11 @@ class FTSRDetectThinGeometryCS : public FTSRConvolutionNetworkShader
 		SHADER_PARAMETER(int32, ThinGeometryTextureIndex)
 		SHADER_PARAMETER(float, ErrorMultiplier)
 		SHADER_PARAMETER(float, MaxRelaxationWeight)
+		SHADER_PARAMETER(float, MinKeepLineIntensity)
 
 		SHADER_PARAMETER_RDG_TEXTURE(Texture2D, SceneDepthTexture)
 		SHADER_PARAMETER_RDG_TEXTURE_SRV(Texture2D, CurrentCoverageTexture)
+		SHADER_PARAMETER_RDG_TEXTURE_SRV(Texture2D, FlickeringLumaTexture)
 		SHADER_PARAMETER_RDG_TEXTURE_SRV(Texture2D, ReprojectedHistoryCoverageTexture)
 		SHADER_PARAMETER_RDG_TEXTURE(Texture2D, DecimateMaskTexture)
 
@@ -1081,6 +1091,11 @@ class FTSRDetectThinGeometryCS : public FTSRConvolutionNetworkShader
 		PermutationVector.Set<FTSRConvolutionNetworkShader::FPermutationDomain>(
 			FTSRConvolutionNetworkShader::RemapPermutation(PermutationVector.Get<FTSRConvolutionNetworkShader::FPermutationDomain>()));
 
+		// If intensity line detection is on, we always need varying range. So disable it if not varying range.
+		if (PermutationVector.Get<FIntensityLineDetectionDim>() && !PermutationVector.Get<FVaryingRangeDetectionDim>())
+		{
+			PermutationVector.Set<FIntensityLineDetectionDim>(false);
+		}
 		return PermutationVector;
 	}
 
@@ -1664,14 +1679,15 @@ bool NeedTSRTranslucentClippedDepthPass(const FViewInfo& View)
 		View.Family->EngineShowFlags.PostProcessing;
 }
 
-bool ShouldAddTSRMainFlickeringLumaPass()
+bool ShouldAddTSRMainThinGeometryCoveragePass()
 {
-	return  (CVarTSRFlickeringEnable.GetValueOnRenderThread()!=0 && CVarTSRFlickeringPeriod.GetValueOnRenderThread() != 0.0f);
+	return ShouldEnableThinGeometryDetection();
 }
 
-bool ShouldAddTSRMainThinGeometryCoveragePass()
+bool ShouldAddTSRMainFlickeringLumaPass()
 {
-	return ShouldEnableThinGeometryDetection();
+	return ((CVarTSRFlickeringEnable.GetValueOnRenderThread() != 0 && CVarTSRFlickeringPeriod.GetValueOnRenderThread() != 0.0f)) ||
+		(CVarTSRThinGeometryCoverageMinKeepLineIntensity.GetValueOnRenderThread() > 0 && ShouldAddTSRMainThinGeometryCoveragePass());
 }
 
 int32 GetTSRMainFlickeringLumaTextureArraySize()
@@ -2696,7 +2712,9 @@ FDefaultTemporalUpscaler::FOutputs AddTemporalSuperResolutionPasses(
 	//Create the thin geometry mask to avoid over rejection
 	if (PassConfig.ThinGeometryDetectionEnable)
 	{
-		const bool EnableVaryingRangeDetection = (GetShadingRange() == EThinGeometryShadingRange::VaryingRange) && (FlickeringFramePeriod > 0.0f);
+		const float MinKeepLineIntensity = FMath::Clamp(CVarTSRThinGeometryCoverageMinKeepLineIntensity.GetValueOnRenderThread(), 0.0f, 1.0f);
+		const bool EnableVaryingRangeDetection = (GetShadingRange() == EThinGeometryShadingRange::VaryingRange)
+			&& ((FlickeringFramePeriod > 0.0f) || MinKeepLineIntensity != 0.0f);
 		{
 			FTSRConvolutionNetworkShader::FPermutationDomain ConvolutionNetworkPermutationVector;
 			ConvolutionNetworkPermutationVector.Set<FTSRConvolutionNetworkShader::FWaveSizeOps>(SelectWaveSize(View.GetShaderPlatform(), { 16, 32, 64 }));
@@ -2706,6 +2724,7 @@ FDefaultTemporalUpscaler::FOutputs AddTemporalSuperResolutionPasses(
 			FTSRDetectThinGeometryCS::FPermutationDomain PermutationVector;
 			PermutationVector.Set<FTSRDetectThinGeometryCS::FVaryingRangeDetectionDim>(EnableVaryingRangeDetection);
 			PermutationVector.Set<FTSRConvolutionNetworkShader::FPermutationDomain>(ConvolutionNetworkPermutationVector);
+			PermutationVector.Set<FTSRDetectThinGeometryCS::FIntensityLineDetectionDim>(MinKeepLineIntensity != 0.0f);
 			PermutationVector = FTSRDetectThinGeometryCS::RemapPermutation(PermutationVector);
 
 			const int32 GroupTileSize = 32;
@@ -2718,6 +2737,7 @@ FDefaultTemporalUpscaler::FOutputs AddTemporalSuperResolutionPasses(
 			PassParameters->ThinGeometryTextureIndex = ThinGeometryTextureIndex;
 			PassParameters->ErrorMultiplier = PassConfig.ThinGeometryErrorMultiplier;
 			PassParameters->MaxRelaxationWeight = FMath::Clamp(CVarTSRThinGeometryCoverageMaxRelaxationWeight.GetValueOnRenderThread(), 0.0f, 1.0f);
+			PassParameters->MinKeepLineIntensity = MinKeepLineIntensity;
 			PassParameters->SceneDepthTexture = PassInputs.SceneDepth.Texture;
 			PassParameters->DecimateMaskTexture = DecimateMaskTexture;
 
@@ -2727,6 +2747,7 @@ FDefaultTemporalUpscaler::FOutputs AddTemporalSuperResolutionPasses(
 				{
 					PassParameters->ReprojectedHistoryCoverageTexture = GraphBuilder.CreateSRV(BlackDummy);
 					PassParameters->CurrentCoverageTexture = GraphBuilder.CreateSRV(FRDGTextureSRVDesc(BlackDummy));
+					PassParameters->FlickeringLumaTexture = GraphBuilder.CreateSRV(FRDGTextureSRVDesc(WhiteDummy));
 				}
 				else
 				{
@@ -2739,10 +2760,20 @@ FDefaultTemporalUpscaler::FOutputs AddTemporalSuperResolutionPasses(
 						int32 ArrayIndex = GetTSRMainFlickeringLumaTextureArraySize() - 1;
 						PassParameters->CurrentCoverageTexture = GraphBuilder.CreateSRV(
 							FRDGTextureSRVDesc::CreateForSlice(PassInputs.FlickeringInputTexture.Texture,  /* SliceIndex = */ ArrayIndex));
+						if (MinKeepLineIntensity != 0.0f)
+						{
+							PassParameters->FlickeringLumaTexture = GraphBuilder.CreateSRV(
+								FRDGTextureSRVDesc::CreateForSlice(PassInputs.FlickeringInputTexture.Texture, /* SliceIndex = */ ArrayIndex - 1));
+						}
+						else
+						{
+							PassParameters->FlickeringLumaTexture = GraphBuilder.CreateSRV(FRDGTextureSRVDesc(WhiteDummy));
+						}
 					}
 					else
 					{
 						PassParameters->CurrentCoverageTexture = GraphBuilder.CreateSRV(FRDGTextureSRVDesc(BlackDummy));
+						PassParameters->FlickeringLumaTexture = GraphBuilder.CreateSRV(FRDGTextureSRVDesc(WhiteDummy));
 					}
 				}
 
