Scene Texture HLSL fixes.  Fix incorrect UV offset math -- needs to use texture extent / buffer size, not viewport size.  Support for fetch syntax on mobile.  Only affects scene texture fetches using custom HLSL shorthand syntax, which is why things aren't widely broken, as that feature was somewhat recently added (5.6).

#rnx

#Synced-CL 52460536

diff --git a/Engine/Shaders/Private/MaterialTemplate.ush b/Engine/Shaders/Private/MaterialTemplate.ush
--- a/Engine/Shaders/Private/MaterialTemplate.ush
+++ b/Engine/Shaders/Private/MaterialTemplate.ush
@@ -1240,6 +1240,32 @@ float4 GetSceneTextureViewSize(const uint SceneTextureId)
 	return ResolvedView.ViewSizeAndInvSize;
 }
 
+// Returns the buffer size and texel size (inverse buffer size) for a given scene texture.
+float4 GetSceneTextureBufferSize(const uint SceneTextureId)
+{
+#if POST_PROCESS_MATERIAL
+	if (IsPostProcessInputSceneTexture(SceneTextureId))
+	{
+		switch (SceneTextureId)
+		{
+		case PPI_PostProcessInput0:
+			return float4(PostProcessInput_0_Extent, PostProcessInput_0_ExtentInverse);
+		case PPI_PostProcessInput1:
+			return float4(PostProcessInput_1_Extent, PostProcessInput_1_ExtentInverse);
+		case PPI_PostProcessInput2:
+			return float4(PostProcessInput_2_Extent, PostProcessInput_2_ExtentInverse);
+		case PPI_PostProcessInput3:
+			return float4(PostProcessInput_3_Extent, PostProcessInput_3_ExtentInverse);
+		case PPI_PostProcessInput4:
+			return float4(PostProcessInput_4_Extent, PostProcessInput_4_ExtentInverse);
+		default:
+			return float4(0, 0, 0, 0);
+		}
+	}
+#endif
+	return ResolvedView.BufferSizeAndInvSize;
+}
+
 // Return the buffer UV min and max for a given scene texture id.
 float4 GetSceneTextureUVMinMax(const uint SceneTextureId)
 {
@@ -3405,10 +3431,10 @@ float4 SceneTextureLookup(float2 UV, int SceneTextureIndex, bool bFiltered)
 	return float4(0, 0, 0, 0);
 }
 
-#define SceneTextureFetch(SceneTextureIndex, PixelOffset) SceneTextureFetchFunc(Parameters, SceneTextureIndex, PixelOffset)
-
 #endif // SHADING_PATH_DEFERRED
 
+#define SceneTextureFetch(SceneTextureIndex, PixelOffset) SceneTextureFetchFunc(Parameters, SceneTextureIndex, PixelOffset)
+
 /**
  * Utility function and macro to streamline Custom HLSL fetching from SceneTextures or UserSceneTextures with a pixel offset from default UVs, including
  * multiplication by texture size and clamping.  Example that fetches at offset 1,1 from the default sample:
@@ -3420,7 +3446,7 @@ float4 SceneTextureFetchFunc(FMaterialPixelParameters Parameters, int SceneTextu
 #if SCENE_TEXTURES_DISABLED || !POST_PROCESS_MATERIAL
 	return float4(0,0,0,0);
 #else
-	float2 UV = ClampSceneTextureUV(GetDefaultSceneTextureUV(Parameters, SceneTextureIndex) + GetSceneTextureViewSize(SceneTextureIndex).zw * PixelOffset, SceneTextureIndex);
+	float2 UV = ClampSceneTextureUV(GetDefaultSceneTextureUV(Parameters, SceneTextureIndex) + GetSceneTextureBufferSize(SceneTextureIndex).zw * PixelOffset, SceneTextureIndex);
 #if SHADING_PATH_MOBILE
 	return MobileSceneTextureLookup(Parameters, SceneTextureIndex, UV);
 #else
