Substrate - Toon - mobile as BxDF
- Anisotropy is not supported yet.
- Shadow transmittance is not supported yet.
- requires extended gbuffer for now.

#tests PC editor, enginte test
#rnx
#Synced-CL 52508769

diff --git a/Engine/Shaders/Private/MobileLightingCommon.ush b/Engine/Shaders/Private/MobileLightingCommon.ush
--- a/Engine/Shaders/Private/MobileLightingCommon.ush
+++ b/Engine/Shaders/Private/MobileLightingCommon.ush
@@ -50,12 +50,20 @@
 #define ENABLE_PLANAR_REFLECTION 0
 #endif 
 
+#ifndef MOBILE_SHADINGMODEL_SUPPORT
+#define MOBILE_SHADINGMODEL_SUPPORT 0
+#endif
+
 #define REFLECTION_COMPOSITE_USE_BLENDED_REFLECTION_CAPTURES 1
 #define REFLECTION_COMPOSITE_SUPPORT_SKYLIGHT_BLEND 0 // Adds additional sampler
 #define REFLECTION_COMPOSITE_HAS_SPHERE_CAPTURES 1
 #define REFLECTION_COMPOSITE_HAS_BOX_CAPTURES 1
 #include "ReflectionEnvironmentComposite.ush"
 
+#if SUBSTRATE_ENABLED
+#include "Substrate/Substrate.ush"
+#endif
+
 half3 SafeGetOutColor(half3 OutColor)
 {
 	// Sky materials can result in high luminance values, e.g. the sun disk. 
@@ -132,7 +140,21 @@ return ShadowMap;
 #define MAX_MOBILE_SHADOWCASCADES 4u
 #endif 
 
-half MobileDirectionalLightCSM(float2 ScreenPosition, float SceneDepth, inout float ShadowPositionZ)
+struct FDirectionalShadowSettings
+{
+	bool bSubsurface;
+	float DensityMulConstant;
+};
+
+FDirectionalShadowSettings GetDefaultDirectionalShadowSettings()
+{
+	FDirectionalShadowSettings DirectionalShadowSettings;
+	DirectionalShadowSettings.bSubsurface = false;
+	DirectionalShadowSettings.DensityMulConstant = 0.0f;
+	return DirectionalShadowSettings;
+}
+
+half MobileDirectionalLightCSM(float2 ScreenPosition, float SceneDepth, inout float ShadowPositionZ, FDirectionalShadowSettings DirectionalShadowSettings)
 {
 	half ShadowMap = 1;
 #if ENABLE_MOBILE_CSM
@@ -142,9 +164,8 @@ half MobileDirectionalLightCSM(float2 ScreenPosition, float SceneDepth, inout fl
 	Settings.ShadowDepthTextureSampler = MobileDirectionalLight.DirectionalLightShadowSampler;
 	Settings.TransitionScale = MobileDirectionalLight.DirectionalLightDirectionAndShadowTransition.w;
 	Settings.ShadowBufferSize = MobileDirectionalLight.DirectionalLightShadowSize;
-	Settings.bSubsurface = false;
+	Settings.bSubsurface = DirectionalShadowSettings.bSubsurface;
 	Settings.bTreatMaxDepthUnshadowed = false;
-	Settings.DensityMulConstant = 0;
 	Settings.ProjectionDepthBiasParameters = 0;
 
 	float4 Count = float4(SceneDepth.xxxx >= MobileDirectionalLight.DirectionalLightShadowDistances);
@@ -164,6 +185,8 @@ half MobileDirectionalLightCSM(float2 ScreenPosition, float SceneDepth, inout fl
 		// Process CSM only when ShadowPosition is valid.
 		if (ShadowPosition.z > 0)
 		{
+			Settings.DensityMulConstant = DirectionalShadowSettings.DensityMulConstant * MobileDirectionalLight.DirectionalLightShadowDepthRange[CascadeIndex];
+
 			// Clamp pixel depth in light space for shadowing opaque, because areas of the shadow depth buffer that weren't rendered to will have been cleared to 1
 			// We want to force the shadow comparison to result in 'unshadowed' in that case, regardless of whether the pixel being shaded is in front or behind that plane
 
@@ -262,10 +285,10 @@ half4 SampleScreenSpaceShadowMask(float2 BufferUV, uint EyeIndex)
 	return 1.0;
 }
 
-FDeferredLightData GetDirectionalLightData(float4 ScreenPosition, inout half OutDynamicShadowing)
+FDeferredLightData GetDirectionalLightData(float4 ScreenPosition, inout half OutDynamicShadowing, FDirectionalShadowSettings DirectionalShadowSettings)
 {
 	float  ShadowPositionZ = 0;
-	half ShadowMap = MobileDirectionalLightCSM(ScreenPosition.xy, ScreenPosition.w, ShadowPositionZ);
+	half ShadowMap = MobileDirectionalLightCSM(ScreenPosition.xy, ScreenPosition.w, ShadowPositionZ, DirectionalShadowSettings);
 	OutDynamicShadowing = 1.0f;
 
 #if !MOBILE_DEFERRED_SHADING && DIRECTIONAL_LIGHT_CSM
@@ -295,8 +318,24 @@ FDeferredLightData GetDirectionalLightData(float4 ScreenPosition, inout half Out
 
 void AccumulateDirectionalLighting(FGBufferData GBuffer, float3 TranslatedWorldPosition, half3 CameraVector, float4 ScreenPosition, in half4 DynamicShadowFactors, in float Dither, inout float OutDirectionalLightShadow, inout FLightAccumulator DirectLighting)
 {
+	FDirectionalShadowSettings DirectionalShadowSettings = GetDefaultDirectionalShadowSettings();
+#if MOBILE_SHADINGMODEL_SUPPORT && SUBSTRATE_ENABLED
+	if (GBuffer.ShadingModelID == SHADINGMODELID_SUBSTRATE_TOON)
+	{
+		DirectionalShadowSettings.bSubsurface = true;
+
+		float ToonProfileId = 0.0f;
+		float2 PatternUVs = 0.0f;
+		UnpackToonCustomData(PackRGBA8(GBuffer.CustomData), PatternUVs, ToonProfileId);
+
+		// Density is extinction as 1/cm (not 1/m) to match with the transmittance shadow evaluation from ManualPCF below.
+		FToonProfileShadowSettings ToonProfileShadowSettings = GetToonProfileShadowSettings(ToonProfileId);
+		DirectionalShadowSettings.DensityMulConstant = ToonProfileShadowSettings.ShadowExtinctionCoefficientInInvCentimeter;
+	}
+#endif
+
 	half DynamicShadowing = 1.0f;
-	FDeferredLightData LightData = GetDirectionalLightData(ScreenPosition, DynamicShadowing);
+	FDeferredLightData LightData = GetDirectionalLightData(ScreenPosition, DynamicShadowing, DirectionalShadowSettings);
 
 #if ENABLE_SHADOWMASKTEXTURE 
 	half4 PreviewShadowMapChannelMask = half4(1.0f, 0.0f, 0.0f, 0.0f);
@@ -694,8 +733,9 @@ half3 ForwardDecalLighting(
 	half3 EmissiveColor, 
 	half  Opacity)
 {
+	FDirectionalShadowSettings DirectionalShadowSettings = GetDefaultDirectionalShadowSettings();
 	float ShadowPositionZ = 0;
-	half ShadowMap = MobileDirectionalLightCSM(ScreenPosition.xy, SceneDepth, ShadowPositionZ);
+	half ShadowMap = MobileDirectionalLightCSM(ScreenPosition.xy, SceneDepth, ShadowPositionZ, DirectionalShadowSettings);
 	half3 DirectionalLightDirection = MobileDirectionalLight.DirectionalLightDirectionAndShadowTransition.xyz;
 	half NoL = saturate(dot(half3(WorldNormal), DirectionalLightDirection));
 	
diff --git a/Engine/Shaders/Private/ShadingModels.ush b/Engine/Shaders/Private/ShadingModels.ush
--- a/Engine/Shaders/Private/ShadingModels.ush
+++ b/Engine/Shaders/Private/ShadingModels.ush
@@ -1061,6 +1061,78 @@ FDirectLighting PreintegratedSkinBxDF( FGBufferData GBuffer, half3 N, half3 V, F
 	return Lighting;
 }
 
+#if SUBSTRATE_ENABLED
+
+#include "Substrate/SubstrateToonBSDF.ush"
+
+// Mobile tech debt: adding support to legacy BxDF for now to not have to convert mobile to Substrate fully yet.
+FDirectLighting ToonBxDF(FGBufferData GBuffer, half3 N, half3 V, FAreaLight AreaLight, FShadowTerms ShadowTerms)
+{
+	const bool bThisPassAppliesDiffuseRamp = true;
+#if SUPPORTS_ANISOTROPIC_MATERIALS
+	bool bHasAnisotropy = HasAnisotropy(GBuffer.SelectiveOutputMask);
+#else
+	bool bHasAnisotropy = false;
+#endif
+	const bool bIsRect = IsRectLight(AreaLight);
+	//Context.VoH = bIsRect ? Context.NoV : Context.VoH;
+
+	BxDFContext Context;
+	float NoV, VoH, NoH;
+	BRANCH
+		if (bHasAnisotropy)
+		{
+			half3 X = GBuffer.WorldTangent;
+			half3 Y = normalize(cross(N, X));
+			Init(Context, N, X, Y, V, AreaLight.SpecularL);
+
+			NoV = Context.NoV;
+			VoH = Context.VoH;
+			NoH = Context.NoH;
+		}
+		else
+		{
+			Init(Context, N, V, AreaLight.SpecularL);
+
+			NoV = Context.NoV;
+			VoH = Context.VoH;
+			NoH = Context.NoH;
+
+			SphereMaxNoH(Context, AreaLight.SphereSinAlpha, true);
+		}
+
+	float3 DiffuseColor = GBuffer.DiffuseColor;
+	float3 SpecularColor = GBuffer.SpecularColor;
+	float  Roughness = GBuffer.Roughness;
+	float2 ToonPatternUVs = 0.0f;
+	float ToonProfileId = 0.0f;
+	UnpackToonCustomData(PackRGBA8(GBuffer.CustomData), ToonPatternUVs, ToonProfileId);
+
+	float	Anisotropy = GBuffer.Anisotropy;
+	float3	Emissive = 0;
+
+	FSubstrateEvaluateToonResult ToonSample = SubstrateToonEvaluateCommon(
+		DiffuseColor,
+		SpecularColor,
+		Roughness,
+		ToonProfileId,
+		ToonPatternUVs,
+		Anisotropy,
+		Emissive,
+		Context,
+		bHasAnisotropy,
+		bIsRect,
+		bThisPassAppliesDiffuseRamp,
+		N, V, AreaLight, ShadowTerms.SurfaceShadow, ShadowTerms.TransmittanceOrOpticalThickness);
+
+	FDirectLighting Lighting;
+	Lighting.Diffuse = ToonSample.IntegratedDiffuseValue;
+	Lighting.Specular = ToonSample.IntegratedSpecularValue;
+	Lighting.Transmission = 0;
+	return Lighting;
+}
+#endif // SUBSTRATE_ENABLED
+
 // UE_DEPRECATED 5.7 - Deprecated by Substrate - Substrate uses SubstrateEvaluateBSDFCommon() instead
 FDirectLighting IntegrateBxDF( FGBufferData GBuffer, half3 N, half3 V, FAreaLight AreaLight, FShadowTerms Shadow )
 {
@@ -1086,6 +1158,10 @@ FDirectLighting IntegrateBxDF( FGBufferData GBuffer, half3 N, half3 V, FAreaLigh
 			return ClothBxDF( GBuffer, N, V, AreaLight, Shadow );
 		case SHADINGMODELID_EYE:
 			return EyeBxDF( GBuffer, N, V, AreaLight, Shadow );
+	#if SUBSTRATE_ENABLED
+		case SHADINGMODELID_SUBSTRATE_TOON:
+			return ToonBxDF( GBuffer, N, V, AreaLight, Shadow );
+	#endif // SUBSTRATE_ENABLED
 		default:
 			return (FDirectLighting)0;
 	}
diff --git a/Engine/Shaders/Private/Substrate/Substrate.ush b/Engine/Shaders/Private/Substrate/Substrate.ush
--- a/Engine/Shaders/Private/Substrate/Substrate.ush
+++ b/Engine/Shaders/Private/Substrate/Substrate.ush
@@ -26,6 +26,7 @@
 #include "/Engine/Private/Substrate/SubstrateStatisticalOperators.ush"
 #include "/Engine/Private/Substrate/SubstrateSubsurface.ush"
 #include "../SpecularProfileCommon.ush"
+#include "SubstrateToonBSDF.ush"
 
 #if SUBSTRATE_USE_FULLYSIMPLIFIED_MATERIAL
 #define SUBSTRATE_CLAMPED_CLOSURE_COUNT 1
@@ -4245,28 +4246,6 @@ float3 UnpackR7G7B6Linear(uint rgb)
 	return Out;
 }
 
-uint PackToonCustomData(in float2 UVs, in float ProfileId, in float Dither)
-{
-	uint U12bits = saturate(UVs.x) * 4095.0f;
-	uint V12bits = saturate(UVs.y) * 4095.0f;
-
-	uint Packed	 = U12bits;
-	Packed		|= V12bits << 12;
-	Packed		|= (uint(ProfileId) & 0xFF) << 24;
-
-	return Packed;
-}
-
-void UnpackToonCustomData(in uint PackedToonCustomData, inout float2 UVs, inout float ProfileId)
-{
-	uint U12bits = float(BitFieldExtractU32(PackedToonCustomData, 12, 0));
-	uint V12bits = float(BitFieldExtractU32(PackedToonCustomData, 12, 12));
-
-	UVs = float2(U12bits, V12bits) * (1.0f / 4095.0f);
-
-	ProfileId	= float(BitFieldExtractU32(PackedToonCustomData, 8, 24)); // Stays in [0, 255]
-}
-
 uint PackR7(float Value, float Dither)
 {
 	return uint(Dither7bits(Value, Dither) * 127.0f);
diff --git a/Engine/Shaders/Private/Substrate/SubstrateEvaluation.ush b/Engine/Shaders/Private/Substrate/SubstrateEvaluation.ush
--- a/Engine/Shaders/Private/Substrate/SubstrateEvaluation.ush
+++ b/Engine/Shaders/Private/Substrate/SubstrateEvaluation.ush
@@ -1321,94 +1321,48 @@ FSubstrateEvaluateResult SubstrateEvaluateBSDFCommon(FSubstrateBSDFContext BSDFC
 			const float  Roughness			= TOON_ROUGHNESS(BSDFContext.BSDF);
 
 			const float ToonProfileId		= TOON_PROFILEID(BSDFContext.BSDF);
+			const float2 ToonPatternUVs		= TOON_PATTERNUVS(BSDFContext.BSDF);
 
-			FToonProfileEvaluateSettings ToonEvalSettings = GetToonProfileEvaluateSettings(ToonProfileId);
-
-			const float2 OffsetFraction = float2(0.0f, 0.0f);
-			const float2 ToonPatternUV = TOON_PATTERNUVS(BSDFContext.BSDF);
-			const float Anisotropy = TOON_ANISOTROPY(BSDFContext.BSDF);
-			float ToonProfileDiffuseDither  = EvaluateToonProfileDiffuseDither(ToonProfileId, ToonPatternUV, ToonEvalSettings);
-			float ToonProfileSpecularDither = EvaluateToonProfileSpecularDither(ToonProfileId, ToonPatternUV, ToonEvalSettings);
-
-			// TransmissionThickness in fact already stores transmittance.
-			float TransmittanceShadow = ShadowTerms.TransmittanceOrOpticalThickness;
-			{
-				// Apply shadow/transmittance hatching pattern
-				FToonProfileShadowSettings ToonProfileShadowSettings = GetToonProfileShadowSettings(ToonProfileId);
-				TransmittanceShadow = EvaluateToonProfileShadowHatching(ToonProfileId, ToonPatternUV, ToonProfileShadowSettings, TransmittanceShadow);
-			}
-
-			float3 Specular = 0.0f;
-			{
-				const float SafeRoughness = MakeRoughnessSafe(Roughness, SUBSTRATE_MIN_GGX_ROUGHNESS);
-				float Alpha2 = Pow4(SafeRoughness);
-				float Alpha2Spec = Alpha2;
-
-				const float3 F0 = SpecularColor;
-				const float3 F90 = 1.0f;
-
-				const bool bIsRectLight = IntegrationType == INTEGRATION_AREA_LIGHT_RECT;
-				if (bIsRectLight)
-				{
-					// In this case, we set D to 1 and Vis will contain the area light / GGX integration
-					float3 RectLightSpec = RectGGXApproxLTC(SafeRoughness, /*F0*/ 1.0, /*F90*/ 1.0, BSDFContext.N, BSDFContext.V, AreaLight.Rect, AreaLight.Texture);
-
-					Specular = RectLightSpec;
-				}
-				else
-				{
-					float D;
-					if (BSDF_GETHASANISOTROPY(BSDFContext.BSDF) && Anisotropy!=0)
-					{
-						float2 AlphaXY = 0;
-						float  Alpha = Square(SafeRoughness);
-						GetAnisotropicRoughness(Alpha, Anisotropy, AlphaXY.x, AlphaXY.y);
-						D = Substrate_D_GGX_Aniso(AlphaXY.x, AlphaXY.y, BSDFContext.Context.NoH, BSDFContext.Context.XoH, BSDFContext.Context.YoH);
-					}
-					else
-					{
-						D = Substrate_D_GGX(SafeRoughness, Alpha2Spec, BSDFContext.Context.NoH);
-					}
-					//const float3 FresnelTerm = Substrate_F_GGX(F0, F90, BSDFContext.Context.VoH);
-					//float Vis = Substrate_Vis_GGX(SafeRoughness, Alpha2Spec, BSDFContext.Context.NoV, AreaLight.NoL);
-
-					Specular =  D;// * Vis * FresnelTerm;
-					Specular *= GGXMaxEnergyLookup(SafeRoughness, Anisotropy, BSDFContext.Context.NoV);
-				}
+			const float Anisotropy			= TOON_ANISOTROPY(BSDFContext.BSDF);
+			const float3 Emissive			= BSDF_GETEMISSIVE(BSDFContext.BSDF);
 
-				// TODO H_PDF and PDF
-			}
-			
 			const bool bThisPassAppliesDiffuseRamp	= SUBSTRATE_LIGHTPASS_APPLIES_TOON_DIFFUSE_RAMP;
-
-			// Diffuse
-			float RampDiffuseNoL = BSDFContext.Context.NoL * 0.5 + 0.5;
-			if (ToonEvalSettings.bDiffuseRampIncludeShadow)
-			{
-				RampDiffuseNoL = RampDiffuseNoL * TransmittanceShadow;
-			}
-			const float3 DiffuseRamp		= bThisPassAppliesDiffuseRamp ? EvaluateToonProfileDiffuseRamp(ToonProfileId, RampDiffuseNoL + ToonProfileDiffuseDither) : 1.0f;
-			Sample.DiffuseColor				= DiffuseColor * DiffuseRamp;
-			Sample.DiffusePathValue			= Diffuse_Lambert(Sample.DiffuseColor);
-			Sample.IntegratedDiffuseValue	= (ToonEvalSettings.bDiffuseRampIncludeShadow ? 1.0f : TransmittanceShadow) * AreaLight.Falloff * AreaLight.FalloffColor * Sample.DiffusePathValue; //  * AreaLight.NoL for smoother edges? But then transmittance is blocked. So this is left in the hand of the artist.
-
-			// Toon unified diffse data.
-			Sample.ToonData.Weight			= TransmittanceShadow * AreaLight.Falloff;	// Accumulate light weight. Albedo should not drive light importance so we do not need to include Lambert diffuse here.
-			Sample.ToonData.ShadowNoL		= RampDiffuseNoL;							// NoL with or without shadow, remapped to [0,1] range.
-
-			// Specular
-			const float3 SpecularRamp = SpecularColor * EvaluateToonProfileSpecularRamp(ToonProfileId, saturate(dot(Specular, 1.0f/3.0f) + ToonProfileSpecularDither));
-			Sample.SpecularPathValue = SpecularRamp;
-			Sample.IntegratedSpecularValue += ShadowTerms.SurfaceShadow * AreaLight.NoL * AreaLight.Falloff * AreaLight.FalloffColor * Sample.SpecularPathValue;
-
-			// Others
-			Sample.DiffusePDF = BSDFContext.SatNoL / PI;
-			Sample.bPostProcessSubsurface = false;
-
-			Sample.EmissivePathValue		= BSDF_GETEMISSIVE(BSDFContext.BSDF);
-
-			Sample.ThroughputV				= OpaqueBSDFThroughput;
-			Sample.TransmittanceAlongN		= OpaqueBSDFThroughput;
+			const bool bHasAnisotropy				= BSDF_GETHASANISOTROPY(BSDFContext.BSDF) && Anisotropy != 0;
+			const bool bIsRectLight					= IntegrationType == INTEGRATION_AREA_LIGHT_RECT;
+
+			FSubstrateEvaluateToonResult ToonResult = SubstrateToonEvaluateCommon(
+				DiffuseColor,
+				SpecularColor,
+				Roughness,
+				ToonProfileId,
+				ToonPatternUVs,
+				Anisotropy,
+				Emissive,
+				BSDFContext.Context,
+				bHasAnisotropy,
+				bIsRectLight,
+				bThisPassAppliesDiffuseRamp,
+				BSDFContext.N, BSDFContext.V, AreaLight, ShadowTerms.SurfaceShadow, ShadowTerms.TransmittanceOrOpticalThickness);
+
+			// Forward to FSubstrateEvaluateResult
+			Sample.IntegratedDiffuseValue	= ToonResult.IntegratedDiffuseValue;
+			Sample.IntegratedSpecularValue	= ToonResult.IntegratedSpecularValue;
+
+			Sample.DiffuseColor				= ToonResult.DiffuseColor;
+
+			Sample.EmissivePathValue		= ToonResult.EmissivePathValue;
+
+			Sample.DiffusePathValue			= ToonResult.DiffusePathValue;
+			Sample.SpecularPathValue		= ToonResult.SpecularPathValue;
+
+			Sample.DiffusePDF				= ToonResult.DiffusePDF;
+
+			Sample.ThroughputV				= ToonResult.ThroughputV;
+			Sample.TransmittanceAlongN		= ToonResult.TransmittanceAlongN;
+			Sample.bPostProcessSubsurface	= ToonResult.bPostProcessSubsurface;
+
+			Sample.ToonData.Weight			= ToonResult.ToonDataWeight;
+			Sample.ToonData.ShadowNoL		= ToonResult.ToonDataShadowNoL;
 		}
 		break;
 
diff --git a/Engine/Shaders/Private/Substrate/SubstrateToonBSDF.ush b/Engine/Shaders/Private/Substrate/SubstrateToonBSDF.ush
new file mode 100644
--- /dev/null
+++ b/Engine/Shaders/Private/Substrate/SubstrateToonBSDF.ush
@@ -0,0 +1,169 @@
+// Copyright Epic Games, Inc. All Rights Reserved.
+
+#pragma once
+
+#if SUBSTRATE_ENABLED
+
+#include "../Common.ush"
+#include "../BRDF.ush"
+#include "../AreaLightCommon.ush"
+
+#include "../ToonProfileCommon.ush"
+
+#include "/Engine/Shared/SubstrateDefinitions.h"
+
+
+uint PackToonCustomData(in float2 UVs, in float ProfileId, in float Dither)
+{
+	uint U12bits = saturate(UVs.x) * 4095.0f;
+	uint V12bits = saturate(UVs.y) * 4095.0f;
+
+	uint Packed = U12bits;
+	Packed |= V12bits << 12;
+	Packed |= (uint(ProfileId) & 0xFF) << 24;
+
+	return Packed;
+}
+
+
+void UnpackToonCustomData(in uint PackedToonCustomData, inout float2 UVs, inout float ProfileId)
+{
+	uint U12bits = float(BitFieldExtractU32(PackedToonCustomData, 12, 0));
+	uint V12bits = float(BitFieldExtractU32(PackedToonCustomData, 12, 12));
+
+	UVs = float2(U12bits, V12bits) * (1.0f / 4095.0f);
+
+	ProfileId = float(BitFieldExtractU32(PackedToonCustomData, 8, 24)); // Stays in [0, 255]
+}
+
+
+// Data to forward to FSubstrateEvaluateResult or FDirectLighting
+struct FSubstrateEvaluateToonResult
+{
+	float3 IntegratedDiffuseValue;
+	float3 IntegratedSpecularValue;
+
+	float3 DiffuseColor;
+
+	float3 EmissivePathValue;
+
+	float3 DiffusePathValue;
+	float3 SpecularPathValue;
+
+	float  DiffusePDF;
+
+	float3 ThroughputV;
+	float3 TransmittanceAlongN;
+	bool   bPostProcessSubsurface;
+
+	float ToonDataWeight;
+	float ToonDataShadowNoL;
+};
+
+
+FSubstrateEvaluateToonResult SubstrateToonEvaluateCommon(
+	float3	DiffuseColor,
+	float3	SpecularColor,
+	float	Roughness,
+	float	ToonProfileId,
+	float2	ToonPatternUVs,
+	float	Anisotropy,
+	float3	Emissive,
+	BxDFContext Context,
+	bool bHasAnisotropy,
+	bool bIsRectLight,
+	bool bThisPassAppliesDiffuseRamp,
+	half3 N, half3 V, FAreaLight AreaLight, 
+	float ShadowTermsSurfaceShadow,		// ShadowTerms.SurfaceShadow
+	float ShadowTermsTransmittance)		// ShadowTerms.TransmittanceOrOpticalThickness, storing transmittance in the case of toon.
+{
+	FSubstrateEvaluateToonResult Sample = (FSubstrateEvaluateToonResult)0;
+
+	FToonProfileEvaluateSettings ToonEvalSettings = GetToonProfileEvaluateSettings(ToonProfileId);
+
+	float ToonProfileDiffuseDither = EvaluateToonProfileDiffuseDither(ToonProfileId, ToonPatternUVs, ToonEvalSettings);
+	float ToonProfileSpecularDither = EvaluateToonProfileSpecularDither(ToonProfileId, ToonPatternUVs, ToonEvalSettings);
+
+	float TransmittanceShadow = ShadowTermsTransmittance;
+	{
+		// Apply shadow/transmittance hatching pattern
+		FToonProfileShadowSettings ToonProfileShadowSettings = GetToonProfileShadowSettings(ToonProfileId);
+		TransmittanceShadow = EvaluateToonProfileShadowHatching(ToonProfileId, ToonPatternUVs, ToonProfileShadowSettings, TransmittanceShadow);
+	}
+
+	float3 Specular = 0.0f;
+	{
+		const float SafeRoughness = MakeRoughnessSafe(Roughness, SUBSTRATE_MIN_GGX_ROUGHNESS);
+		float Alpha2 = Pow4(SafeRoughness);
+		float Alpha2Spec = Alpha2;
+
+		const float3 F0 = SpecularColor;
+		const float3 F90 = 1.0f;
+
+		if (bIsRectLight)
+		{
+			// In this case, we set D to 1 and Vis will contain the area light / GGX integration
+			float3 RectLightSpec = RectGGXApproxLTC(SafeRoughness, /*F0*/ SpecularColor, /*F90*/ 1.0f, N, V, AreaLight.Rect, AreaLight.Texture);
+
+			Specular = RectLightSpec;
+		}
+		else
+		{
+			float D;
+			if (bHasAnisotropy && Anisotropy != 0)
+			{
+				float2 AlphaXY = 0;
+				float  Alpha = Square(SafeRoughness);
+				GetAnisotropicRoughness(Alpha, Anisotropy, AlphaXY.x, AlphaXY.y);
+				D = D_GGXaniso(AlphaXY.x, AlphaXY.y, Context.NoH, Context.XoH, Context.YoH);
+			}
+			else
+			{
+				D = D_GGX(Alpha2Spec, Context.NoH);
+			}
+			//const float3 FresnelTerm = Substrate_F_GGX(F0, F90, Context.VoH);
+			//float Vis = Substrate_Vis_GGX(SafeRoughness, Alpha2Spec, Context.NoV, AreaLight.NoL);
+
+			Specular = D;// * Vis * FresnelTerm;
+			Specular *= GGXMaxEnergyLookup(SafeRoughness, Anisotropy, Context.NoV);
+		}
+
+		// TODO H_PDF and PDF
+	}
+
+	// Diffuse
+	float RampDiffuseNoL = Context.NoL * 0.5 + 0.5;
+	if (ToonEvalSettings.bDiffuseRampIncludeShadow)
+	{
+		RampDiffuseNoL = RampDiffuseNoL * TransmittanceShadow;
+	}
+	const float3 DiffuseRamp = bThisPassAppliesDiffuseRamp ? EvaluateToonProfileDiffuseRamp(ToonProfileId, RampDiffuseNoL + ToonProfileDiffuseDither) : 1.0f;
+	Sample.DiffuseColor = DiffuseColor * DiffuseRamp;
+	Sample.DiffusePathValue = Diffuse_Lambert(Sample.DiffuseColor);
+	Sample.IntegratedDiffuseValue = (ToonEvalSettings.bDiffuseRampIncludeShadow ? 1.0f : TransmittanceShadow) * AreaLight.Falloff * AreaLight.FalloffColor * Sample.DiffusePathValue; //  * AreaLight.NoL for smoother edges? But then transmittance is blocked. So this is left in the hand of the artist.
+
+	// Toon unified diffse data.
+	Sample.ToonDataWeight = TransmittanceShadow * AreaLight.Falloff;	// Accumulate light weight. Albedo should not drive light importance so we do not need to include Lambert diffuse here.
+	Sample.ToonDataShadowNoL = RampDiffuseNoL;							// NoL with or without shadow, remapped to [0,1] range.
+
+	// Specular
+	const float3 SpecularRamp = SpecularColor * EvaluateToonProfileSpecularRamp(ToonProfileId, saturate(dot(Specular, 1.0f / 3.0f) + ToonProfileSpecularDither));
+	Sample.SpecularPathValue = SpecularRamp;
+	Sample.IntegratedSpecularValue += ShadowTermsSurfaceShadow * AreaLight.NoL * AreaLight.Falloff * AreaLight.FalloffColor * Sample.SpecularPathValue;
+
+	// Others
+	Sample.DiffusePDF = saturate(Context.NoL) / PI;
+	Sample.bPostProcessSubsurface = false;
+
+	Sample.EmissivePathValue = Emissive;
+
+	const float OpaqueBSDFThroughput = 0.0f;
+	Sample.ThroughputV = OpaqueBSDFThroughput;
+	Sample.TransmittanceAlongN = OpaqueBSDFThroughput;
+
+	return Sample;
+
+}
+
+#endif // SUBSTRATE_ENABLED
+
diff --git a/Engine/Source/Runtime/Engine/Private/Materials/MaterialInterface.cpp b/Engine/Source/Runtime/Engine/Private/Materials/MaterialInterface.cpp
--- a/Engine/Source/Runtime/Engine/Private/Materials/MaterialInterface.cpp
+++ b/Engine/Source/Runtime/Engine/Private/Materials/MaterialInterface.cpp
@@ -693,6 +693,14 @@ FMaterialRelevance UMaterialInterface::GetRelevance_Internal(const UMaterial* Ma
 	FMaterialRelevance MaterialRelevance;
 
 	MaterialRelevance.ShadingModelMask = GetShadingModels().GetShadingModelField();
+
+	if ((uint16(MaterialResource->MaterialGetSubstrateMaterialBsdfFeatures_GameThread()) & uint16(ESubstrateBsdfFeature::Toon)) > 0)
+	{
+		// Add the Substrate shading model to the ShadingModelMask so that we detect on mobile that is it not only DefaultLit material in the scene.
+		// So that the base pass write the correct stencil values for the toon shading model to be schedule during the "all shading models" pass.
+		MaterialRelevance.ShadingModelMask |= EMaterialShadingModel::MSM_Strata;
+	}
+
 	MaterialRelevance.CustomDepthStencilUsageMask = MaterialResource->GetCustomDepthStencilUsageMask_GameThread();
 	MaterialRelevance.bDecal = bDecal;
 
diff --git a/Engine/Source/Runtime/Engine/Public/SceneView.h b/Engine/Source/Runtime/Engine/Public/SceneView.h
--- a/Engine/Source/Runtime/Engine/Public/SceneView.h
+++ b/Engine/Source/Runtime/Engine/Public/SceneView.h
@@ -852,6 +852,7 @@ BEGIN_GLOBAL_SHADER_PARAMETER_STRUCT_WITH_CONSTRUCTOR(FMobileDirectionalLightSha
 	SHADER_PARAMETER_EX(FVector4f, DirectionalLightShadowSize, EShaderPrecisionModifier::Half)
 	SHADER_PARAMETER_EX(FVector4f, DirectionalLightDistanceFadeMADAndSpecularScale, EShaderPrecisionModifier::Half) // .z is used for SpecularScale, .w is for DiffuseScale
 	SHADER_PARAMETER_EX(FVector4f, DirectionalLightShadowDistances, EShaderPrecisionModifier::Half)
+	SHADER_PARAMETER_EX(FVector4f, DirectionalLightShadowDepthRange, EShaderPrecisionModifier::Half)
 	SHADER_PARAMETER_ARRAY(FMatrix44f, DirectionalLightScreenToShadow, [MAX_MOBILE_SHADOWCASCADES])
 	SHADER_PARAMETER(uint32, DirectionalLightNumCascades)
 	SHADER_PARAMETER(uint32, DirectionalLightShadowMapChannelMask)
diff --git a/Engine/Source/Runtime/Renderer/Private/MobileBasePass.cpp b/Engine/Source/Runtime/Renderer/Private/MobileBasePass.cpp
--- a/Engine/Source/Runtime/Renderer/Private/MobileBasePass.cpp
+++ b/Engine/Source/Runtime/Renderer/Private/MobileBasePass.cpp
@@ -799,6 +799,11 @@ bool FMobileBasePassMeshProcessor::TryAddMeshBatch(const FMeshBatch& RESTRICT Me
 	if (ShouldDraw(Material))
 	{
 		FMaterialShadingModelField ShadingModels = Material.GetShadingModels();
+		if ((uint16(Material.MaterialGetSubstrateMaterialBsdfFeatures_RenderThread()) & uint16(ESubstrateBsdfFeature::Toon)) > 0)
+		{
+			ShadingModels.AddShadingModel(EMaterialShadingModel::MSM_Strata); // This model is not only a DefaultLit model and will be schedule in the "all shading model" lighting pass.
+		}
+
 		const EShaderPlatform ShaderPlatform = GetFeatureLevelShaderPlatform(FeatureLevel);
 #if WITH_EDITOR
 		// Non-Editor builds filter out shading models on a material load
diff --git a/Engine/Source/Runtime/Renderer/Private/MobileBasePassRendering.cpp b/Engine/Source/Runtime/Renderer/Private/MobileBasePassRendering.cpp
--- a/Engine/Source/Runtime/Renderer/Private/MobileBasePassRendering.cpp
+++ b/Engine/Source/Runtime/Renderer/Private/MobileBasePassRendering.cpp
@@ -472,6 +472,7 @@ void SetupMobileDirectionalLightUniformParameters(
 					}
 					Params.DirectionalLightScreenToShadow[OutShadowIndex] = FMatrix44f(ShadowInfo->GetScreenToShadowMatrix(SceneView));		// LWC_TODO: Precision loss?
 					Params.DirectionalLightShadowDistances[OutShadowIndex] = ShadowInfo->CascadeSettings.SplitFar;
+					Params.DirectionalLightShadowDepthRange[OutShadowIndex] = ShadowInfo->MaxSubjectZ;										// TODO: Depth range is not MaxSubjectZ - MinSubjectZ as on the non-mobile it seems?
 					Params.DirectionalLightNumCascades++;
 					OutShadowIndex++;
 				}
