Fixed GPU skin VF to propagate instancing parameters and use the correct macros.

#Synced-CL 52259874

diff --git a/Engine/Shaders/Private/GpuSkinVertexFactory.ush b/Engine/Shaders/Private/GpuSkinVertexFactory.ush
--- a/Engine/Shaders/Private/GpuSkinVertexFactory.ush
+++ b/Engine/Shaders/Private/GpuSkinVertexFactory.ush
@@ -233,6 +233,21 @@ FMaterialPixelParameters GetMaterialPixelParameters(FVertexFactoryInterpolantsVS
 	Result.PrimitiveId = Interpolants.PrimitiveId;
 #endif
 
+#if USE_INSTANCE_CULLING
+	FInstanceSceneData InstanceData = GetInstanceSceneData(Interpolants.PrimitiveId);
+	Result.PerInstanceRandom = InstanceData.RandomID;
+#if USES_PER_INSTANCE_CUSTOM_DATA
+	Result.CustomDataOffset = InstanceData.CustomDataOffset;
+	Result.CustomDataCount = InstanceData.CustomDataCount;
+#endif
+#if HAS_INSTANCE_LOCAL_TO_WORLD_PS
+	Result.InstanceLocalToWorld = InstanceData.LocalToWorld;
+#endif
+#if HAS_INSTANCE_WORLD_TO_LOCAL_PS
+	Result.InstanceWorldToLocal = InstanceData.WorldToLocal;
+#endif
+#endif
+
 	return Result;
 }
 
@@ -318,8 +333,10 @@ FMaterialVertexParameters GetMaterialVertexParameters(
 
 	Result.PreSkinnedPosition = Intermediates.UnpackedPosition.xyz;
 	Result.PreSkinnedNormal = Input.TangentZ.xyz;
+#if USE_INSTANCE_CULLING
 	Result.InstanceLocalToWorld = Intermediates.SceneData.InstanceData.LocalToWorld;
 	Result.InstanceWorldToLocal = Intermediates.SceneData.InstanceData.WorldToLocal;
+#endif
 	Result.PerInstanceRandom = Intermediates.SceneData.InstanceData.RandomID;
 	Result.PrevFrameLocalToWorld = Intermediates.PrevLocalToWorld;
 
