Enabling support for fp16 shader compilation on some target platforms

#Synced-CL 52398231

diff --git a/Engine/Shaders/Private/TemporalSuperResolution/TSRCommon.ush b/Engine/Shaders/Private/TemporalSuperResolution/TSRCommon.ush
--- a/Engine/Shaders/Private/TemporalSuperResolution/TSRCommon.ush
+++ b/Engine/Shaders/Private/TemporalSuperResolution/TSRCommon.ush
@@ -133,6 +133,27 @@
 
 #endif
 
+// On platforms that compile via SPIRV (PLATFORM_USES_SPIRV_RESOURCE_TYPES), Vulkan SPIRV requires 32-bit sampled types for OpTypeImage
+#if CONFIG_COMPILE_FP16 && PLATFORM_USES_SPIRV_RESOURCE_TYPES
+	#define tsr_half_resource   float
+	#define tsr_half2_resource  float2
+	#define tsr_half3_resource  float3
+	#define tsr_half4_resource  float4
+	#define tsr_ushort_resource uint
+#else
+	#define tsr_half_resource   tsr_half
+	#define tsr_half2_resource  tsr_half2
+	#define tsr_half3_resource  tsr_half3
+	#define tsr_half4_resource  tsr_half4
+	#define tsr_ushort_resource tsr_ushort
+#endif
+
+#if CONFIG_SCENE_COLOR_ALPHA
+	#define tsr_halfC_resource tsr_half4_resource
+#else
+	#define tsr_halfC_resource tsr_half3_resource
+#endif
+
 
 // Largest encodable normal number in a half used on console.
 static const tsr_half LargestNormalNumber = tsr_half(65504.0);
diff --git a/Engine/Shaders/Private/TemporalSuperResolution/TSRRejectShading.usf b/Engine/Shaders/Private/TemporalSuperResolution/TSRRejectShading.usf
--- a/Engine/Shaders/Private/TemporalSuperResolution/TSRRejectShading.usf
+++ b/Engine/Shaders/Private/TemporalSuperResolution/TSRRejectShading.usf
@@ -48,28 +48,28 @@ uint bEnableResurrection;
 uint bEnableFlickeringHeuristic;
 
 
-Texture2D<tsr_halfC> InputTexture;
-Texture2D<tsr_half>  InputMoireLumaTexture;
-Texture2D<tsr_half4> InputSceneTranslucencyTexture;
-Texture2D<tsr_half4> ReprojectedHistoryGuideTexture;
-Texture2D<tsr_half>  ReprojectedHistoryGuideMetadataTexture;
-Texture2D<tsr_half4> ReprojectedHistoryMoireTexture;
+Texture2D<tsr_halfC_resource> InputTexture;
+Texture2D<tsr_half_resource>  InputMoireLumaTexture;
+Texture2D<tsr_half4_resource> InputSceneTranslucencyTexture;
+Texture2D<tsr_half4_resource> ReprojectedHistoryGuideTexture;
+Texture2D<tsr_half_resource>  ReprojectedHistoryGuideMetadataTexture;
+Texture2D<tsr_half4_resource> ReprojectedHistoryMoireTexture;
 #if CONFIG_HISTORY_RESURRECTION
-	Texture2D<tsr_half4> ResurrectedHistoryGuideTexture;
-	Texture2D<tsr_half>  ResurrectedHistoryGuideMetadataTexture;
+	Texture2D<tsr_half4_resource> ResurrectedHistoryGuideTexture;
+	Texture2D<tsr_half_resource>  ResurrectedHistoryGuideMetadataTexture;
 #endif
-Texture2D<tsr_half2> DecimateMaskTexture;
-Texture2D<tsr_ushort> IsMovingMaskTexture;
-Texture2D<tsr_ushort> ThinGeometryTexture;
+Texture2D<tsr_half2_resource> DecimateMaskTexture;
+Texture2D<tsr_ushort_resource> IsMovingMaskTexture;
+Texture2D<tsr_ushort_resource> ThinGeometryTexture;
 Texture2D<float2> ClosestDepthTexture;
 
-RWTexture2DArray<tsr_half4> HistoryGuideOutput;
-RWTexture2DArray<tsr_half4> HistoryMoireOutput;
-RWTexture2D<tsr_half4>  HistoryRejectionOutput;
+RWTexture2DArray<tsr_half4_resource> HistoryGuideOutput;
+RWTexture2DArray<tsr_half4_resource> HistoryMoireOutput;
+RWTexture2D<tsr_half4_resource>  HistoryRejectionOutput;
 RWTexture2DArray<uint>  ReprojectionFieldOutput;
-RWTexture2D<tsr_halfC>  InputSceneColorOutput;
-RWTexture2D<tsr_half>   InputSceneColorLdrLumaOutput;
-RWTexture2D<tsr_ushort> AntiAliasMaskOutput;
+RWTexture2D<tsr_halfC_resource>  InputSceneColorOutput;
+RWTexture2D<tsr_half_resource>   InputSceneColorLdrLumaOutput;
+RWTexture2D<tsr_ushort_resource> AntiAliasMaskOutput;
 
 
 //------------------------------------------------------- GLOBALS
@@ -129,7 +129,7 @@ void FetchThinGeometryEdge(
 	for (uint ElementIndex = 0; ElementIndex < SIMD_SIZE; ElementIndex++)
 	{
 		tsr_short2 ClampedFetchPixelPos = PixelPos.GetElement(ElementIndex);
-		tsr_short EncodedThinGeometryWeight = ThinGeometryTexture[ClampedFetchPixelPos];
+		tsr_short EncodedThinGeometryWeight = tsr_short(ThinGeometryTexture[ClampedFetchPixelPos]);
 		
 		tsr_half Weight = tsr_half(EncodedThinGeometryWeight >> tsr_ushort(1)) * rcp(tsr_half(127.0));
 		bool UseMultiplyOperation = (EncodedThinGeometryWeight & tsr_ushort(0x1))!= tsr_ushort(0x0);
@@ -151,7 +151,7 @@ void FetchDecimateMask(
 	{
 		tsr_short2 ClampedFetchPixelPos = PixelPos.GetElement(ElementIndex);
 		
-		tsr_half2 DecimateMask = DecimateMaskTexture[ClampedFetchPixelPos];
+		tsr_half2 DecimateMask = tsr_half2(DecimateMaskTexture[ClampedFetchPixelPos]);
 
 		tsr_ushort BitMask = tsr_ushort(DecimateMask.r * tsr_half(255.0));
 
@@ -175,14 +175,14 @@ void FetchSceneColorAndTranslucency(
 		float2 TranslucencyTextureUV = ApplyScreenTransform(float2(ClampedFetchPixelPos), InputPixelPosToTranslucencyTextureUV);
 		TranslucencyTextureUV = fastClamp(TranslucencyTextureUV, TranslucencyTextureUVMin, TranslucencyTextureUVMax);
 		
-		OriginalOpaqueInput.SetElement(ElementIndex, InputTexture[ClampedFetchPixelPos]);
-		OriginalTranslucencyInput.SetElement(ElementIndex, InputSceneTranslucencyTexture.SampleLevel(GlobalBilinearClampedSampler, TranslucencyTextureUV, 0));
+		OriginalOpaqueInput.SetElement(ElementIndex, tsr_halfC(InputTexture[ClampedFetchPixelPos]));
+		OriginalTranslucencyInput.SetElement(ElementIndex, tsr_half4(InputSceneTranslucencyTexture.SampleLevel(GlobalBilinearClampedSampler, TranslucencyTextureUV, 0)));
 	}
 } // FetchSceneColorAndTranslucency()
 
 void FetchHistoryGuide(
-	Texture2D<tsr_half4> HistoryGuideTexture,
-	Texture2D<tsr_half>  HistoryGuideMetadataTexture,
+	Texture2D<tsr_half4_resource> HistoryGuideTexture,
+	Texture2D<tsr_half_resource>  HistoryGuideMetadataTexture,
 	tsr_tensor_short2 PixelPos,
 	out tsr_tensor_halfC History,
 	out tsr_tensor_half  HistoryUncertainty)
@@ -193,14 +193,14 @@ void FetchHistoryGuide(
 	{
 		tsr_short2 ClampedFetchPixelPos = PixelPos.GetElement(ElementIndex);
 
-		History.SetElement(ElementIndex, HistoryGuideTexture[ClampedFetchPixelPos]);
-		HistoryUncertainty.SetElement(ElementIndex, HistoryGuideMetadataTexture[ClampedFetchPixelPos]);
+		History.SetElement(ElementIndex, tsr_half4(HistoryGuideTexture[ClampedFetchPixelPos]));
+		HistoryUncertainty.SetElement(ElementIndex, tsr_half(HistoryGuideMetadataTexture[ClampedFetchPixelPos]));
 	}
 	#else
 	{
 		tsr_short2 ClampedFetchPixelPos = PixelPos.GetElement(ElementIndex);
 
-		tsr_half4 ReprojectedGuide = HistoryGuideTexture[ClampedFetchPixelPos];
+		tsr_half4 ReprojectedGuide = tsr_half4(HistoryGuideTexture[ClampedFetchPixelPos]);
 		History.SetElement(ElementIndex, ReprojectedGuide.rgb);
 		HistoryUncertainty.SetElement(ElementIndex, ReprojectedGuide.a);
 	}
@@ -234,8 +234,8 @@ void FetchMoireInput(
 	{
 		tsr_short2 ClampedFetchPixelPos = PixelPos.GetElement(ElementIndex);
 		
-		MoireInput.SetElement(ElementIndex, InputMoireLumaTexture[ClampedFetchPixelPos]);
-		PrevMoireHistory.SetElement(ElementIndex, ReprojectedHistoryMoireTexture[ClampedFetchPixelPos]);
+		MoireInput.SetElement(ElementIndex, tsr_half(InputMoireLumaTexture[ClampedFetchPixelPos]));
+		PrevMoireHistory.SetElement(ElementIndex, tsr_half4(ReprojectedHistoryMoireTexture[ClampedFetchPixelPos]));
 		bIsStatic.SetElement(ElementIndex, tsr_half(1.0) - tsr_half(IsMovingMaskTexture[ClampedFetchPixelPos]) * rcp(tsr_half(255.0)));
 	}
 } // FetchMoireInput()
@@ -900,11 +900,15 @@ void RejectShadingCS(
 			#else
 				HistoryGuideOutput[tsr_short3(OutputHistoryPixelPos, 0)] = tsr_half4(FinalGuideColor.xyz, InputUncertainty.GetElement(ElementIndex));
 			#endif
-			HistoryRejectionOutput[OutputPixelPos] = tsr_half4(
-				RejectionBlendFinal.GetElement(ElementIndex) - tsr_half(0.5 * rcp(256.0)),
-				DisableHistoryClamp.GetElement(ElementIndex) - tsr_half(0.5 * rcp(256.0)),
-				tsr_half(255.5 * rcp(256.0)) - IncreaseValidityMultiplier.GetElement(ElementIndex),
-				BitMask);
+			{
+				const tsr_half kHalfOverTwoFiftySix = tsr_half(0.5) * tsr_half(rcp(256.0));
+				const tsr_half k255HalfOverTwoFiftySix = tsr_half(255.5) * tsr_half(rcp(256.0));
+				HistoryRejectionOutput[OutputPixelPos] = tsr_half4(
+					RejectionBlendFinal.GetElement(ElementIndex) - kHalfOverTwoFiftySix,
+					DisableHistoryClamp.GetElement(ElementIndex) - kHalfOverTwoFiftySix,
+					k255HalfOverTwoFiftySix - IncreaseValidityMultiplier.GetElement(ElementIndex),
+					BitMask);
+			}
 			AntiAliasMaskOutput[OutputPixelPos] = select(bSpatialAntiAlias.GetElement(ElementIndex), tsr_ushort(1), tsr_ushort(0));
 		}
 	}
