Substrate - Toon - added support for non extended gbuffer.

#rnx
#tests pc editor

#Synced-CL 52512026

diff --git a/Engine/Shaders/Private/DeferredShadingCommon.ush b/Engine/Shaders/Private/DeferredShadingCommon.ush
--- a/Engine/Shaders/Private/DeferredShadingCommon.ush
+++ b/Engine/Shaders/Private/DeferredShadingCommon.ush
@@ -695,6 +695,15 @@ void MobileEncodeGBuffer(
 			OutGBufferC.a = GBuffer.CustomData.y; // IrisNormal_Oct.x
 			OutGBufferB.g = GBuffer.CustomData.z; // IrisNormal_Oct.y
 		}
+		else if (GBuffer.ShadingModelID == SHADINGMODELID_SUBSTRATE_TOON)
+		{
+			// This is the toon profile index.
+			OutGBufferC.a = GBuffer.CustomData.a;
+			// UVs are ignored and ramp offset or shadow hatching won't be used on mobile.
+
+			OutGBufferA.b = MobileEncodeIdAndColorChannel(GBuffer.ShadingModelID, 0, true);
+			// OutGBufferA.b and OutGBufferB.a could be used to store more data.
+		}
 		else
 		{
 			// encode 4 bits for shading model
@@ -810,6 +819,17 @@ FGBufferData MobileDecodeGBuffer(half4 InGBufferA, half4 InGBufferB, half4 InGBu
 			GBuffer.GBufferAO = 1;	// no space for AO
 			GBuffer.Specular = 0.5f;	// no space for Specular
 		}
+		else if (GBuffer.ShadingModelID == SHADINGMODELID_SUBSTRATE_TOON)
+		{
+			// This is the toon profile index.
+			GBuffer.CustomData.a = InGBufferC.a;
+			// UVs are ignored and ramp offset or shadow hatching won't be used on mobile.
+			GBuffer.CustomData.rgb = 0.0f.xxx;
+
+			// OutGBufferA.b and OutGBufferB.a could be used to store more data.
+
+			GBuffer.GBufferAO = 1;	// no space for AO
+		}
 	#endif
 
 	if (GBuffer.ShadingModelID == SHADINGMODELID_EYE)
diff --git a/Engine/Shaders/Private/Substrate/SubstrateToonBSDF.ush b/Engine/Shaders/Private/Substrate/SubstrateToonBSDF.ush
--- a/Engine/Shaders/Private/Substrate/SubstrateToonBSDF.ush
+++ b/Engine/Shaders/Private/Substrate/SubstrateToonBSDF.ush
@@ -60,6 +60,9 @@ struct FSubstrateEvaluateToonResult
 	float ToonDataShadowNoL;
 };
 
+// PatternsUVs can only be used on mobile when the extended GBuffer is enabled. Always available on non mobile renderer.
+// When PatternsUVs cannot be used, ramp uvs and shadow hatching cannot be renderered and are skipped.
+#define TOON_CAN_USE_PATTERNS_UVS (!SHADING_PATH_MOBILE || (!MOBILE_DEFERRED_SHADING || MOBILE_EXTENDED_GBUFFER))
 
 FSubstrateEvaluateToonResult SubstrateToonEvaluateCommon(
 	float3	DiffuseColor,
@@ -81,15 +84,17 @@ FSubstrateEvaluateToonResult SubstrateToonEvaluateCommon(
 
 	FToonProfileEvaluateSettings ToonEvalSettings = GetToonProfileEvaluateSettings(ToonProfileId);
 
-	float ToonProfileDiffuseDither = EvaluateToonProfileDiffuseDither(ToonProfileId, ToonPatternUVs, ToonEvalSettings);
-	float ToonProfileSpecularDither = EvaluateToonProfileSpecularDither(ToonProfileId, ToonPatternUVs, ToonEvalSettings);
-
+	float ToonProfileDiffuseDither = 0.0f;
+	float ToonProfileSpecularDither = 0.0f;
 	float TransmittanceShadow = ShadowTermsTransmittance;
-	{
-		// Apply shadow/transmittance hatching pattern
-		FToonProfileShadowSettings ToonProfileShadowSettings = GetToonProfileShadowSettings(ToonProfileId);
-		TransmittanceShadow = EvaluateToonProfileShadowHatching(ToonProfileId, ToonPatternUVs, ToonProfileShadowSettings, TransmittanceShadow);
-	}
+#if TOON_CAN_USE_PATTERNS_UVS
+	ToonProfileDiffuseDither = EvaluateToonProfileDiffuseDither(ToonProfileId, ToonPatternUVs, ToonEvalSettings);
+	ToonProfileSpecularDither = EvaluateToonProfileSpecularDither(ToonProfileId, ToonPatternUVs, ToonEvalSettings);
+
+	// Apply shadow/transmittance hatching pattern
+	FToonProfileShadowSettings ToonProfileShadowSettings = GetToonProfileShadowSettings(ToonProfileId);
+	TransmittanceShadow = EvaluateToonProfileShadowHatching(ToonProfileId, ToonPatternUVs, ToonProfileShadowSettings, TransmittanceShadow);
+#endif
 
 	float3 Specular = 0.0f;
 	{
