Fix RescaleVariableRateShading accessing depth buffer out of view bounds.

#Synced-CL 52512372

diff --git a/Engine/Shaders/Private/VariableRateShading/VRSShadingRateReproject.usf b/Engine/Shaders/Private/VariableRateShading/VRSShadingRateReproject.usf
--- a/Engine/Shaders/Private/VariableRateShading/VRSShadingRateReproject.usf
+++ b/Engine/Shaders/Private/VariableRateShading/VRSShadingRateReproject.usf
@@ -45,7 +45,8 @@ float2 ReprojectBufferUV(uint3 DispatchThreadId, uint3 GroupThreadId, float2 Buf
 		// Center the UV in the middle of the 2x2 block
 		float2 UVOffset = ScaledUVOffset;
 		float2 DepthUV = UVOffset + (float2(DispatchThreadId.xy) + float2(1, 1)) * pixelUVSize;
-		CachedDepth[CachedDepthIndex.x][CachedDepthIndex.y] = LookupDeviceZ(DepthUV);  
+		DepthUV = clamp(DepthUV, View.BufferBilinearUVMinMax.xy, View.BufferBilinearUVMinMax.zw);
+		CachedDepth[CachedDepthIndex.x][CachedDepthIndex.y] = LookupDeviceZ(DepthUV);
     }
 	GroupMemoryBarrierWithGroupSync();
 
