Fix [Switch Only] Texture Data Override Normals are flipped/inverted on mirrored actors
- Fix DeterminantSign always returning +1.0 in UBO non-instance-culling path.
- In GetSceneDataIntermediatesUBO, the else branch (USE_INSTANCE_CULLING=0) was reading DeterminantSign from InstanceData.Flags which was never copied from Primitive.Flags (zero-initialized).
- This caused mirrored actors on lower-end platforms using the UBO path (e.g. Switch) to have incorrect tangent basis, resulting in flipped TDO normals.

#rnx

#Synced-CL 52556927

diff --git a/Engine/Shaders/Private/SceneData.ush b/Engine/Shaders/Private/SceneData.ush
--- a/Engine/Shaders/Private/SceneData.ush
+++ b/Engine/Shaders/Private/SceneData.ush
@@ -1305,7 +1305,7 @@ FSceneDataIntermediates GetSceneDataIntermediatesUBO(uint DrawInstanceId)
 	Intermediates.InstanceData.WorldToLocal = Intermediates.Primitive.WorldToLocal;
 	Intermediates.InstanceData.NonUniformScale = Intermediates.Primitive.NonUniformScale;
 	Intermediates.InstanceData.InvNonUniformScale = Intermediates.Primitive.InvNonUniformScale;
-	Intermediates.InstanceData.DeterminantSign = GetPrimitive_DeterminantSign_FromFlags(Intermediates.InstanceData.Flags);
+	Intermediates.InstanceData.DeterminantSign = GetPrimitive_DeterminantSign_FromFlags(Intermediates.Primitive.Flags);
 #endif
 	Intermediates.InstanceData.LocalBoundsCenter = Intermediates.Primitive.InstanceLocalBoundsCenter;
 	Intermediates.InstanceData.LocalBoundsExtent = Intermediates.Primitive.InstanceLocalBoundsExtent;
