[DeferredLighting] Multiply Lighting.Transmission by LightData.DiffuseScale in order to preserve overall light energy. Without it certainly shading models like eye or hair were too dark for high diffuse scale values.

#Synced-CL 52454962

diff --git a/Engine/Shaders/Private/DeferredLightingCommon.ush b/Engine/Shaders/Private/DeferredLightingCommon.ush
--- a/Engine/Shaders/Private/DeferredLightingCommon.ush
+++ b/Engine/Shaders/Private/DeferredLightingCommon.ush
@@ -407,6 +407,7 @@ FLightAccumulator AccumulateDynamicLighting(
 
 			Lighting.Specular *= LightData.SpecularScale;
 			Lighting.Diffuse  *= LightData.DiffuseScale;
+			Lighting.Transmission *= LightData.DiffuseScale;
 
 		#if USE_LIGHT_FUNCTION_ATLAS
 			MaskedLightColor *= GetLocalLightFunctionCommon(TranslatedWorldPosition, LightData.LightFunctionAtlasLightIndex);
