Fix for UE-373164

#Synced-CL 52633472

diff --git a/Engine/Shaders/Private/Substrate/SubstrateMobileForwardLighting.ush b/Engine/Shaders/Private/Substrate/SubstrateMobileForwardLighting.ush
--- a/Engine/Shaders/Private/Substrate/SubstrateMobileForwardLighting.ush
+++ b/Engine/Shaders/Private/Substrate/SubstrateMobileForwardLighting.ush
@@ -213,7 +213,8 @@ float3 SubstrateMobileForwardLighting(
 				if (bHasDirectionalLight)
 				{
 					half DynamicShadowing = 1.0f;
-					FDeferredLightData DirLightData = GetDirectionalLightData(MaterialParameters.ScreenPosition, DynamicShadowing);
+					FDirectionalShadowSettings DirectionalShadowSettings = GetDefaultDirectionalShadowSettings();
+					FDeferredLightData DirLightData = GetDirectionalLightData(MaterialParameters.ScreenPosition, DynamicShadowing, DirectionalShadowSettings);
 					half4 DirLightAttenuation = half4(1, 1, DynamicShadowing, DynamicShadowing);
 					uint DirLightingChannelMask = LIGHTING_CHANNEL_MASK;
 					float3 ToDirLight = DirLightData.Direction;	// Already normalized
