FSSS
- Improved filtering, it is now properly centered
- Remove caps on UI parameters for artist to play with.

#rnx

#Synced-CL 52638104

diff --git a/Engine/Shaders/Private/FogSeparateComposition.usf b/Engine/Shaders/Private/FogSeparateComposition.usf
--- a/Engine/Shaders/Private/FogSeparateComposition.usf
+++ b/Engine/Shaders/Private/FogSeparateComposition.usf
@@ -34,15 +34,20 @@ void DownsampleColorCS(uint3 ThreadId : SV_DispatchThreadID)
 	}
 
 	float4 Accumulator = 0.0f;
-	float2 InvSrcMipResolution = 1.0f / float2(SrcMipResolution);
-	int2 SourceCoord = int2(ThreadId.xy) * 2;
+	float2 UVScale = 1.0f / (DstMipResolution * 2);
+	float2 SourceCoord = (float2(ThreadId.xy) + 0.5f) * 2.0f;
 
+#if 0
+	// Debug 2x2 downsample
+	Accumulator += SourceTexture.SampleLevel(SourceSampler, SourceCoord * UVScale, 0); 
+#else
 	// 4x4 box filter with destination texel at the center with only 4 filtered taps.
-	Accumulator += SourceTexture.SampleLevel(SourceSampler, (float2(SourceCoord) + float2(0, 0)) * InvSrcMipResolution, 0);
-	Accumulator += SourceTexture.SampleLevel(SourceSampler, (float2(SourceCoord) + float2(0, 2)) * InvSrcMipResolution, 0);
-	Accumulator += SourceTexture.SampleLevel(SourceSampler, (float2(SourceCoord) + float2(2, 2)) * InvSrcMipResolution, 0);
-	Accumulator += SourceTexture.SampleLevel(SourceSampler, (float2(SourceCoord) + float2(2, 0)) * InvSrcMipResolution, 0);
+	Accumulator += SourceTexture.SampleLevel(SourceSampler, (SourceCoord + float2(-1.0f,-1.0f)) * UVScale, 0);
+	Accumulator += SourceTexture.SampleLevel(SourceSampler, (SourceCoord + float2(-1.0f, 1.0f)) * UVScale, 0);
+	Accumulator += SourceTexture.SampleLevel(SourceSampler, (SourceCoord + float2( 1.0f, 1.0f)) * UVScale, 0);
+	Accumulator += SourceTexture.SampleLevel(SourceSampler, (SourceCoord + float2( 1.0f,-1.0f)) * UVScale, 0);
 	Accumulator  = Accumulator * 0.25f;
+#endif
 
 	OutTextureMipColor[uint2(ThreadId.xy)] = Accumulator;
 }
@@ -75,14 +80,15 @@ void UpsampleColorCS(uint3 ThreadId : SV_DispatchThreadID)
 	}
 
 	float4 Accumulator = 0.0f;
-	float2 InvSrcMipResolution = 1.0f / float2(SrcMipResolution);
-	float2 SourceCoord = float2(ThreadId.xy) / 2;
+	float2 UVScale = 1.0f / (DstMipResolution  * 0.5f);
+	float2 SourceCoord = (float2(ThreadId.xy) + 0.5f) * 0.5f;
 
 	// 2x2 box filter
-	Accumulator += SourceTexture.SampleLevel(SourceSampler, (SourceCoord + float2(-0.5,-0.5)) * InvSrcMipResolution, 0);
-	Accumulator += SourceTexture.SampleLevel(SourceSampler, (SourceCoord + float2( 0.5,-0.5)) * InvSrcMipResolution, 0);
-	Accumulator += SourceTexture.SampleLevel(SourceSampler, (SourceCoord + float2( 0.5, 0.5)) * InvSrcMipResolution, 0);
-	Accumulator += SourceTexture.SampleLevel(SourceSampler, (SourceCoord + float2(-0.5, 0.5)) * InvSrcMipResolution, 0);
+	const float UpsampleScale = 0.5f;
+	Accumulator += SourceTexture.SampleLevel(SourceSampler, (SourceCoord + (float2(-1.0f,-1.0f) * UpsampleScale)) * UVScale, 0);
+	Accumulator += SourceTexture.SampleLevel(SourceSampler, (SourceCoord + (float2(-1.0f, 1.0f) * UpsampleScale)) * UVScale, 0);
+	Accumulator += SourceTexture.SampleLevel(SourceSampler, (SourceCoord + (float2( 1.0f, 1.0f) * UpsampleScale)) * UVScale, 0);
+	Accumulator += SourceTexture.SampleLevel(SourceSampler, (SourceCoord + (float2( 1.0f,-1.0f) * UpsampleScale)) * UVScale, 0);
 	Accumulator *= 0.25;
 
 	// This is a UAV typed load on OutTextureMipColor. See DoesPlatformSupportFSSS for details.
@@ -98,7 +104,7 @@ void UpsampleColorCS(uint3 ThreadId : SV_DispatchThreadID)
 #ifdef HALFRESOLUTIONUPSAMPLERPS
 
 int4 HalfResDepthMinMaxCoord;			// min and max of depth inside the hald res buffer were many views lives.
-float4 HalfResSourceResolutionAndInv;	// Resolution of buffer to upsample, isolated containing a single view
+float2 TexelCoordToUVs;					// Convert a texel to UVs according to view resolution and texture resolution.
 
 float FSSSSpreadScale;
 float FullResToFogResScale; // either 1 for full or 0.5 for half resolution
@@ -155,7 +161,7 @@ void HalfResolutionUpsamplerPS(
 		// This is a custom fade based on W for the clip above to note cause discontinuity.
 		const float FadeBasedOnW = saturate(W);
 
-		SourceTextureColor = HalfResolutionSourceTexture.SampleLevel(HalfResolutionSourceSampler, FullResLocalCoord * HalfResSourceResolutionAndInv.zw, Mip).rgba * FadeBasedOnW;
+		SourceTextureColor = HalfResolutionSourceTexture.SampleLevel(HalfResolutionSourceSampler, FullResLocalCoord * TexelCoordToUVs, Mip).rgba * FadeBasedOnW;
 		OutColor = CombineSourceTextures(SourceTextureColor, SourceTextureColor2);
 	}
 
@@ -235,8 +241,8 @@ void HalfResolutionUpsamplerPS(
 
 
 
-	SourceTextureColor = HalfResolutionSourceTexture.SampleLevel(HalfResolutionSourceSampler, HalfResLocalCoord * HalfResSourceResolutionAndInv.zw, 0).rgba;
-	SourceTextureColor2 = HalfResolutionSource2Texture.SampleLevel(HalfResolutionSource2Sampler, HalfResLocalCoord * HalfResSourceResolutionAndInv.zw, 0).r;
+	SourceTextureColor = HalfResolutionSourceTexture.SampleLevel(HalfResolutionSourceSampler, HalfResLocalCoord * TexelCoordToUVs, 0).rgba;
+	SourceTextureColor2 = HalfResolutionSource2Texture.SampleLevel(HalfResolutionSource2Sampler, HalfResLocalCoord * TexelCoordToUVs, 0).r;
 	OutColor = CombineSourceTextures(SourceTextureColor, SourceTextureColor2);
 
 
diff --git a/Engine/Shaders/Private/HeightFogPixelShader.usf b/Engine/Shaders/Private/HeightFogPixelShader.usf
--- a/Engine/Shaders/Private/HeightFogPixelShader.usf
+++ b/Engine/Shaders/Private/HeightFogPixelShader.usf
@@ -222,11 +222,10 @@ void ExponentialPixelMain(
 	const float SceneDepthCentimeters   = SceneDepth;
 	FFogSSSParameters FSSSParams = EvaluateFSSSParameter(FogTransmittanceToScene, SceneDepthCentimeters);
 
-	const float SceneColorToScatterAmount = FSSSParams.Coverage * FSSSParams.Albedo * FogStruct.FSSSSceneColorScatteringAmountScale; // *(1.0f / (4.0f * PI));		// Approximation, should use eq. 4 from paper?
+	const float SceneColorToScatterAmount = saturate(FSSSParams.Coverage * FSSSParams.Albedo * FogStruct.FSSSSceneColorScatteringAmountScale); // *(1.0f / (4.0f * PI));		// Approximation, should use eq. 4 from paper?
 
-	// This is point sampling scene color when using full resolution. At hald resolution this is average of 2x2 texels.
+	// This is point sampling scene color when using full resolution. At half resolution this is average of 2x2 texels.
 	const float3 SceneColorToScatter = FogTransmittanceToScene * SceneColorToScatterAmount * Texture2DSampleLevel(SceneTexturesStruct.SceneColorTexture, SceneColorTextureSampler, TexCoordSceneTex, 0).rgb;
-	// Otherwise point sampling SceneTexturesStruct.SceneColorTexture[FogSeparateCompositionResolutionDivider * ThreadId.xy + ViewHalfResolutionRectMin].rgb;
 
 	const float FinalFogTransmittanceToScene = FogTransmittanceToScene * saturate(1.0f - SceneColorToScatterAmount);
 
diff --git a/Engine/Source/Runtime/Engine/Classes/Components/ExponentialHeightFogComponent.h b/Engine/Source/Runtime/Engine/Classes/Components/ExponentialHeightFogComponent.h
--- a/Engine/Source/Runtime/Engine/Classes/Components/ExponentialHeightFogComponent.h
+++ b/Engine/Source/Runtime/Engine/Classes/Components/ExponentialHeightFogComponent.h
@@ -202,7 +202,7 @@ class UExponentialHeightFogComponent : public USceneComponent
 	 * When 0, only the fog luminance will be scattered.
 	 * When 1, fog and scene luminance will be scattered.
 	 */
-	UPROPERTY(EditAnywhere, BlueprintReadOnly, Category = FogScreenSpaceScattering, meta = (DisplayName = "Scene Color Scattering Amount Scale", UIMin = "0", UIMax = "1", ClampMin = "0", ClampMax = "1"))
+	UPROPERTY(EditAnywhere, BlueprintReadOnly, Category = FogScreenSpaceScattering, meta = (DisplayName = "Scene Color Scattering Amount Scale", UIMin = "0", UIMax = "1", ClampMin = "0"))
 	float FSSSSceneColorScatteringAmountScale;
 
 	/**
@@ -210,7 +210,7 @@ class UExponentialHeightFogComponent : public USceneComponent
 	 * Increasing will make scene and thicker fog look fuzzier.
 	 * But too high values can cause artifacts such as halos around depth discontinuities.
 	 */
-	UPROPERTY(EditAnywhere, BlueprintReadOnly, Category = FogScreenSpaceScattering, meta = (DisplayName = "Spread Scale", UIMin = "0", UIMax = "1", ClampMin = "0", ClampMax = "1"))
+	UPROPERTY(EditAnywhere, BlueprintReadOnly, Category = FogScreenSpaceScattering, meta = (DisplayName = "Spread Scale", UIMin = "0", UIMax = "1", ClampMin = "0"))
 	float FSSSSpreadScale;
 
 	/**
diff --git a/Engine/Source/Runtime/Renderer/Private/FogSeparateComposition.cpp b/Engine/Source/Runtime/Renderer/Private/FogSeparateComposition.cpp
--- a/Engine/Source/Runtime/Renderer/Private/FogSeparateComposition.cpp
+++ b/Engine/Source/Runtime/Renderer/Private/FogSeparateComposition.cpp
@@ -255,7 +255,7 @@ class FUpsampleHalfResTexturePS : public FGlobalShader
 
 	BEGIN_SHADER_PARAMETER_STRUCT(FParameters, )
 		SHADER_PARAMETER(FIntVector4, HalfResDepthMinMaxCoord)
-		SHADER_PARAMETER(FVector4f, HalfResSourceResolutionAndInv)
+		SHADER_PARAMETER(FVector2f, TexelCoordToUVs)
 		SHADER_PARAMETER(float, FullResToFogResScale)
 		SHADER_PARAMETER(float, FSSSSpreadScale)
 
@@ -310,7 +310,7 @@ void FogSeparateCompositionAllocateResources(
 	FogSeparateCompositionViewResources.ResolutionDivider = ShouldRenderFogUsingSeparateCompositionHalfRes() ? 2 : 1;
 
 	FogSeparateCompositionViewResources.ViewRect = GetDownscaledRect(View.ViewRect, FIntPoint(FogSeparateCompositionViewResources.ResolutionDivider, FogSeparateCompositionViewResources.ResolutionDivider));
-	const FIntPoint FogViewExtent = FogSeparateCompositionViewResources	.ViewRect.Size();
+	const FIntPoint FogViewExtent = FogSeparateCompositionViewResources.ViewRect.Size();
 
 	const bool bSceneRequestsFSSS = SceneRequestsFSSS(Scene);
 	const uint32 MipCount = bSceneRequestsFSSS ? GetViewBlurMipChainMipCount(FogViewExtent) : 1;
@@ -358,7 +358,11 @@ void UpsampleFogSeparateCompositionTextureForView(
 	const bool bSceneRequestsFSSS = SceneRequestsFSSS(Scene);
 
 	PassParameters->HalfResDepthMinMaxCoord = FIntVector4(HalfResViewRect.Min.X, HalfResViewRect.Min.Y, HalfResViewRect.Max.X, HalfResViewRect.Max.Y);
-	PassParameters->HalfResSourceResolutionAndInv = FVector4f(SourceTextureResolution.X, SourceTextureResolution.Y, 1.0f / SourceTextureResolution.X, 1.0f / SourceTextureResolution.Y);
+
+	const FVector2f TextureFullResolution = FVector2f(FogSeparateCompositionViewResources.FogSeparateCompositionTexture0->Desc.Extent);
+	const FVector2f TextureViewResolution = FVector2f(View.ViewRect.Size());
+	PassParameters->TexelCoordToUVs = (TextureViewResolution / TextureFullResolution) / TextureFullResolution;
+	
 	PassParameters->FullResToFogResScale = 1.0f / float(ResolutionDivider);
 	PassParameters->FSSSSpreadScale = bSceneRequestsFSSS ? FogInfo.FSSSSpreadScale : 0.0f; // Spread of 0 if FSSS is not enabled and mip chain not generated.
 
@@ -388,6 +392,6 @@ void UpsampleFogSeparateCompositionTextureForView(
 		RDG_EVENT_NAME("Fog::UpsampleSeparateCompositionTextures"),
 		PixelShader,
 		PassParameters,
-		FullResViewport.Rect,
+		View.ViewRect,
 		BlendState);
 }
\ No newline at end of file
