<enter description here>
<saved by Perforce>

#tests PC editor dx12
#rnx

#Synced-CL 52622009

diff --git a/Engine/Shaders/Private/Substrate/Substrate.ush b/Engine/Shaders/Private/Substrate/Substrate.ush
--- a/Engine/Shaders/Private/Substrate/Substrate.ush
+++ b/Engine/Shaders/Private/Substrate/Substrate.ush
@@ -317,6 +317,8 @@ struct FSubstrateBSDF
 	void SubstrateBakeSpecularProfileToReflectivity(in float NoV);
 
 	void SubstrateBakeGlintToReflectivity();
+
+	void SubstrateBakeHazinessToSingleRoughness();
 };
 
 ////////////////////////////////////////////////////////////////////////////////////////////////////////////////////////////////////////////////////////////////////////////////////////////////
@@ -3714,6 +3716,26 @@ void FSubstrateBSDF::SubstrateBakeGlintToReflectivity()
 #endif
 }
 
+void FSubstrateBSDF::SubstrateBakeHazinessToSingleRoughness()
+{
+#if !SUBSTRATE_FASTPATH
+	if (BSDF_GETHASHAZINESS(this))
+	{
+		FHaziness LocalHaziness = UnpackHaziness(SLAB_HAZINESS(this));
+
+		// Only blend roughness when we are not in the simple clear coat mode, otherwise we want to keep the original roughness for the legacy clear coat shading model.
+		if (!LocalHaziness.bSimpleClearCoat)
+		{
+			const float SlabRoughess = SLAB_ROUGHNESS(this);
+			SLAB_ROUGHNESS(this) = lerp(SlabRoughess, LocalHaziness.Roughness, LocalHaziness.Weight);
+			BSDF_SETHASHAZINESS(this, 0);
+
+			SLAB_HAZINESS(this) = PackHaziness(LocalHaziness);
+		}
+	}
+#endif
+}
+
 #if !SUBSTRATE_COMPILER_SUPPORTS_STRUCT_FORWARD_DECLARATION
 void UpdateSingleBSDFOperatorCoverageTransmittance(
 	inout FSubstrateTree SubstrateTree,
diff --git a/Engine/Shaders/Private/Substrate/SubstrateExport.ush b/Engine/Shaders/Private/Substrate/SubstrateExport.ush
--- a/Engine/Shaders/Private/Substrate/SubstrateExport.ush
+++ b/Engine/Shaders/Private/Substrate/SubstrateExport.ush
@@ -1103,6 +1103,9 @@ FExportResult SubstrateMaterialExportOut(
 			const float GreyScaleWeight = dot(CurrentBSDF.LuminanceWeightV, (1.0 / 3.0).xxx);
 			const float NoV = saturate(dot(SurfaceWorldNormal, V));
 
+			// Bake haziness second specular lobe with its own roughness into primary specular lobe (if not simple clear coat legacy shading model).
+			CurrentBSDF.SubstrateBakeHazinessToSingleRoughness();
+
 			// Bake the F90 into F0 of BSDFs to blend in order to retain some of the original look when using parameter blending or fully simplified materials.
 			CurrentBSDF.SubstrateBakeF90ToReflectivity(NoV);
 			BSDF_SETHASF90(CurrentBSDF, 0);
