Add hair strands occlusion to SSS transmittance.

#Synced-CL 52630564

diff --git a/Engine/Shaders/Private/RayTracing/RayTracingOcclusionRGS.usf b/Engine/Shaders/Private/RayTracing/RayTracingOcclusionRGS.usf
--- a/Engine/Shaders/Private/RayTracing/RayTracingOcclusionRGS.usf
+++ b/Engine/Shaders/Private/RayTracing/RayTracingOcclusionRGS.usf
@@ -275,6 +275,19 @@ float ComputeDiffuseTransmission(
 			RayFlags,
 			InstanceInclusionMask,
 			SSSRay);
+
+		// Take into account hair strands occlusion when computing transmission distance
+		#if USE_HAIR_LIGHTING
+		if (bUseHairVoxel && VisibilityPayload.IsMiss())
+		{
+			const float HairHitT = TraverseHair(PixelCoord, RandSequence, SSSRay.Origin, SSSRay.Direction, /*SSSRay.TMin*/-1, VirtualVoxel.Raytracing_ShadowOcclusionThreshold);
+			if (HairHitT > SSSRay.TMin)
+			{
+				VisibilityPayload.HitT = HairHitT;
+			}
+		}
+		#endif
+
 		if (VisibilityPayload.IsMiss())
 		{
 			TransmissionDistance = ScatterDistance + SSSPayload.HitT;
