[First Person] Fix for VSM occluder distance not working correctly with first person world space representation shadow casters and double sided foliage. The issue was that foliage relies on occluder distance to calculate shadow, but by just taking the minimum occluder distance between the FP shadow and the regular (foliage self-) shadow, the FP shadow factor was mismatched with the foliage self shadow occluder distance, leading to leaky FP shadows on foliage.

#Synced-CL 52675100

diff --git a/Engine/Shaders/Private/VirtualShadowMaps/VirtualShadowMapProjection.usf b/Engine/Shaders/Private/VirtualShadowMaps/VirtualShadowMapProjection.usf
--- a/Engine/Shaders/Private/VirtualShadowMaps/VirtualShadowMapProjection.usf
+++ b/Engine/Shaders/Private/VirtualShadowMaps/VirtualShadowMapProjection.usf
@@ -468,9 +468,9 @@ FVirtualShadowMapSampleResult ProjectLight(
 					}
 					else
 					{
-						// Merge the result, picking the most-shadowed value & nearest occluder distance
+						// Merge the result, picking the most-shadowed value and the associated occluder distance
 						Result.ShadowFactor = min(Result.ShadowFactor, FPResult.ShadowFactor);
-						Result.OccluderDistance = min(Result.OccluderDistance, FPResult.OccluderDistance);
+						Result.OccluderDistance = Result.ShadowFactor < FPResult.ShadowFactor ? Result.OccluderDistance : FPResult.OccluderDistance;
 					}
 				}
 			}
