FSSS - fix bug with alpha propagation.

#tests PC editor with without propagate alpha.
#rnx
#Synced-CL 52682737

diff --git a/Engine/Shaders/Private/HeightFogPixelShader.usf b/Engine/Shaders/Private/HeightFogPixelShader.usf
--- a/Engine/Shaders/Private/HeightFogPixelShader.usf
+++ b/Engine/Shaders/Private/HeightFogPixelShader.usf
@@ -205,7 +205,7 @@ void ExponentialPixelMain(
 	}
 	
 	FLATTEN
-	if (bPropagateAlpha && !defined(RENDER_FOG_COMP_TEXTURE_CS))
+	if (bPropagateAlpha && !RENDER_FOG_COMP_TEXTURE_CS)
 	{
 		// We need to output coverage so that we can potentially apply alpha hold out blending. See RenderViewFog.
 		OutColor = RETURN_COLOR(float4(FogInscatteringAndOpacity.rgb, 1.0 - FogInscatteringAndOpacity.w));
