Remove mobile-specific behavior for ES3_1 feature level on non-mobile platforms

#Synced-CL 52668772

diff --git a/Engine/Shaders/Private/GammaCorrectionCommon.ush b/Engine/Shaders/Private/GammaCorrectionCommon.ush
--- a/Engine/Shaders/Private/GammaCorrectionCommon.ush
+++ b/Engine/Shaders/Private/GammaCorrectionCommon.ush
@@ -36,13 +36,7 @@ half3 LinearToSrgbBranching(half3 lin)
 
 half3 LinearToSrgb(half3 lin) 
 {
-#if FEATURE_LEVEL > FEATURE_LEVEL_ES3_1
-	// Branching is faster than branchless on AMD on PC.
 	return LinearToSrgbBranching(lin);
-#else
-	// Adreno devices(Nexus5) with Android 4.4.2 do not handle branching version well, so always use branchless on Mobile
-	return LinearToSrgbBranchless(lin);
-#endif
 }
 
 half3 sRGBToLinear( half3 Color ) 
diff --git a/Engine/Shaders/Private/MaterialTemplate.ush b/Engine/Shaders/Private/MaterialTemplate.ush
--- a/Engine/Shaders/Private/MaterialTemplate.ush
+++ b/Engine/Shaders/Private/MaterialTemplate.ush
@@ -576,10 +576,6 @@ struct FMaterialPixelParameters
 	/** Per-particle properties. Only valid for particle vertex factories. */
 	FMaterialParticleParameters Particle;
 
-#if FEATURE_LEVEL == FEATURE_LEVEL_ES3_1
-	float4 LayerWeights;
-#endif
-
 #if TEX_COORD_SCALE_ANALYSIS
 	/** Parameters used by the MaterialTexCoordScales shader. */
 	FTexCoordScalesParams TexCoordScalesParams;
@@ -732,10 +728,6 @@ struct FMaterialVertexParameters
 	float2 TexCoords[NUM_MATERIAL_TEXCOORDS_VERTEX];
 #endif
 
-#if NUM_MATERIAL_TEXCOORDS_VERTEX && FEATURE_LEVEL == FEATURE_LEVEL_ES3_1
-	float2 TexCoordOffset; // Offset for UV localization for large UV values
-#endif
-
 	/** Per-particle properties. Only valid for particle vertex factories. */
 	FMaterialParticleParameters Particle;
 
@@ -4492,7 +4484,7 @@ void ClipLODTransition(float2 SvPosition, float DitherFactor)
 	if (abs(DitherFactor) > .001)
 	{
 		float ArgCos = dot(floor(SvPosition.xy), float2(347.83451793, 3343.28371963));
-#if FEATURE_LEVEL <= FEATURE_LEVEL_ES3_1
+#if PLATFORM_SUPPORTS_RELAXED_PRECISION
 		// Temporary workaround for precision issues on mobile when the argument is bigger than 10k
 		ArgCos = fmod(ArgCos, 10000);
 #endif
